#ttp 331464 - EDITOR: CURVE: Missing icon in the Curve asset tab
#proj UE4
#branch UE4
#summary Added missing icon style

diff --git a/Engine/Source/Editor/EditorStyle/Private/SlateEditorStyle.cpp b/Engine/Source/Editor/EditorStyle/Private/SlateEditorStyle.cpp
--- a/Engine/Source/Editor/EditorStyle/Private/SlateEditorStyle.cpp
+++ b/Engine/Source/Editor/EditorStyle/Private/SlateEditorStyle.cpp
@@ -4686,6 +4686,8 @@ void FSlateEditorStyle::FStyle::SetupToolkitStyles()
 
 	// Curve Editor
 	{
+		Set( "CurveAssetEditor.Tabs.Properties", new IMAGE_BRUSH( "Icons/AssetIcons/Curve_16x", Icon16x16 ) );
+
 		Set( "CurveEditor.FitHorizontally", new IMAGE_BRUSH( "Icons/icon_CurveEditor_Horizontal_40x", Icon40x40 ) );
 		Set( "CurveEditor.FitVertically", new IMAGE_BRUSH( "Icons/icon_CurveEditor_Vertical_40x", Icon40x40 ) );
 		Set( "CurveEditor.FitToAll", new IMAGE_BRUSH( "Icons/icon_CurveEditor_All_40x", Icon40x40 ) );
