#ue4 BlueprintFunctionLibrary can be created. It uses its own icon.

diff --git a/Engine/Config/BaseEditor.ini b/Engine/Config/BaseEditor.ini
--- a/Engine/Config/BaseEditor.ini
+++ b/Engine/Config/BaseEditor.ini
@@ -250,7 +250,7 @@ SupportWebsiteURL="http://unrealengine.com/support"
 bUseUserDefinedStructure=false
 
 [CustomBlueprintFunctionLibrary]
-bCanCreateNew=false
+bCanCreateNew=true
 
 [EditoronlyBP]
 bAllowClassAndBlueprintPinMatching=false
diff --git a/Engine/Source/Editor/EditorStyle/Private/SlateEditorStyle.cpp b/Engine/Source/Editor/EditorStyle/Private/SlateEditorStyle.cpp
--- a/Engine/Source/Editor/EditorStyle/Private/SlateEditorStyle.cpp
+++ b/Engine/Source/Editor/EditorStyle/Private/SlateEditorStyle.cpp
@@ -4136,6 +4136,7 @@ void FSlateEditorStyle::FStyle::SetupClassIconsAndThumbnails()
 		Set( "ClassThumbnail.BlueprintGeneratedClass", new IMAGE_BRUSH( "Icons/icon_class_Blueprint_64x", Icon32x32 ) );
 		Set( "ClassThumbnail.BlueprintInterface", new IMAGE_BRUSH( "Icons/icon_class_Blueprint_Interface_64x", Icon32x32 ) );
 		Set( "ClassThumbnail.BlueprintMacroLibrary", new IMAGE_BRUSH( "Icons/icon_class_Blueprint_MacroLibrary_64x", Icon32x32 ) );
+		Set( "ClassThumbnail.BlueprintFunctionLibrary", new IMAGE_BRUSH("Icons/icon_class_Blueprint_FunctionLibrary_64x", Icon32x32));
 		Set( "ClassThumbnail.ButtonStyleAsset", new IMAGE_BRUSH( "Icons/icon_class_ButtonStyle_64x", Icon32x32 ) );
 		Set( "ClassThumbnail.CameraAnim", new IMAGE_BRUSH( "Icons/icon_class_CameraAnim_64px", Icon32x32 ) );
 		Set( "ClassThumbnail.CurveBase", new IMAGE_BRUSH( "Icons/icon_class_Curve_64x", Icon32x32 ) );
diff --git a/Engine/Source/Editor/UnrealEd/Private/Factories/EditorFactories.cpp b/Engine/Source/Editor/UnrealEd/Private/Factories/EditorFactories.cpp
--- a/Engine/Source/Editor/UnrealEd/Private/Factories/EditorFactories.cpp
+++ b/Engine/Source/Editor/UnrealEd/Private/Factories/EditorFactories.cpp
@@ -6286,7 +6286,7 @@ FText UBlueprintFunctionLibraryFactory::GetDisplayName() const
 
 FName UBlueprintFunctionLibraryFactory::GetNewAssetThumbnailOverride() const
 {
-	return TEXT("ClassThumbnail.BlueprintMacroLibrary");
+	return TEXT("ClassThumbnail.BlueprintFunctionLibrary");
 }
 
 uint32 UBlueprintFunctionLibraryFactory::GetMenuCategories() const
