Added Styles for TranslationEditor

diff --git a/Engine/Source/Editor/EditorStyle/Private/SlateEditorStyle.cpp b/Engine/Source/Editor/EditorStyle/Private/SlateEditorStyle.cpp
--- a/Engine/Source/Editor/EditorStyle/Private/SlateEditorStyle.cpp
+++ b/Engine/Source/Editor/EditorStyle/Private/SlateEditorStyle.cpp
@@ -126,6 +126,7 @@ void FSlateEditorStyle::FStyle::Initialize()
 	SetupContentBrowserStyle();
 	SetupLandscapeEditorStyle();
 	SetupToolkitStyles();
+	SetupTranslationEditorStyles();
 	SetupMatineeStyle();
 	SetupSourceControlStyles();
 	SetupAutomationStyles();
@@ -5147,6 +5148,11 @@ void FSlateEditorStyle::FStyle::SetupAutomationStyles()
 	}
 }
 
+void FSlateEditorStyle::FStyle::SetupTranslationEditorStyles()
+{
+	Set("TranslationEditor.Import", new IMAGE_BRUSH("Icons/icon_Import_40x", Icon40x40));
+}
+
 #undef IMAGE_BRUSH
 #undef BOX_BRUSH
 #undef BORDER_BRUSH
diff --git a/Engine/Source/Editor/EditorStyle/Private/SlateEditorStyle.h b/Engine/Source/Editor/EditorStyle/Private/SlateEditorStyle.h
--- a/Engine/Source/Editor/EditorStyle/Private/SlateEditorStyle.h
+++ b/Engine/Source/Editor/EditorStyle/Private/SlateEditorStyle.h
@@ -66,6 +66,7 @@ public:
 		void SetupWindowStyles();
 		void SetupDockingStyles();
 		void SetupTutorialStyles();
+		void SetupTranslationEditorStyles();
 		void SetupPropertyEditorStyles();
 		void SetupProfilerStyle();
 		void SetupGraphEditorStyles();
