Fixed missing Font Editor toolbar icons.

Fixed up some icon specifiers in the style that werent mapping quite correctly.

TTP# 332847 - UE4: Font editor has missing toolbar icons.

diff --git a/Engine/Source/Editor/EditorStyle/Private/SlateEditorStyle.cpp b/Engine/Source/Editor/EditorStyle/Private/SlateEditorStyle.cpp
--- a/Engine/Source/Editor/EditorStyle/Private/SlateEditorStyle.cpp
+++ b/Engine/Source/Editor/EditorStyle/Private/SlateEditorStyle.cpp
@@ -4570,15 +4570,15 @@ void FSlateEditorStyle::FStyle::SetupToolkitStyles()
 		Set( "FontEditor.Update.Small", new IMAGE_BRUSH( "Icons/icon_FontEd_Update_40x", Icon20x20 ) );
 		Set( "FontEditor.UpdateAll", new IMAGE_BRUSH( "Icons/icon_FontEd_UpdateAll_40x", Icon40x40 ) );
 		Set( "FontEditor.UpdateAll.Small", new IMAGE_BRUSH( "Icons/icon_FontEd_UpdateAll_40x", Icon20x20 ) );
-		Set( "FontEditor.Export", new IMAGE_BRUSH( "Icons/icon_FontEd_Export_40x", Icon40x40 ) );
-		Set( "FontEditor.Export.Small", new IMAGE_BRUSH( "Icons/icon_FontEd_Export_40x", Icon20x20 ) );
-		Set( "FontEditor.ExportAll", new IMAGE_BRUSH( "Icons/icon_FontEd_ExportAll_40x", Icon40x40 ) );
-		Set( "FontEditor.ExportAll.Small", new IMAGE_BRUSH( "Icons/icon_FontEd_ExportAll_40x", Icon20x20 ) );
-
-		Set( "FontEditor.BackgroundColor", new IMAGE_BRUSH( "Icons/icon_FontEd_Background_40x", Icon40x40 ) );
-		Set( "FontEditor.BackgroundColor.Small", new IMAGE_BRUSH( "Icons/icon_FontEd_Background_40x", Icon20x20 ) );
-		Set( "FontEditor.ForegroundColor", new IMAGE_BRUSH( "Icons/icon_FontEd_Foreground_40x", Icon40x40 ) );
-		Set( "FontEditor.ForegroundColor.Small", new IMAGE_BRUSH( "Icons/icon_FontEd_Foreground_40x", Icon20x20 ) );
+		Set( "FontEditor.ExportPage", new IMAGE_BRUSH( "Icons/icon_FontEd_Export_40x", Icon40x40 ) );
+		Set( "FontEditor.ExportPage.Small", new IMAGE_BRUSH( "Icons/icon_FontEd_Export_40x", Icon20x20 ) );
+		Set( "FontEditor.ExportAllPages", new IMAGE_BRUSH( "Icons/icon_FontEd_ExportAll_40x", Icon40x40 ) );
+		Set( "FontEditor.ExportAllPages.Small", new IMAGE_BRUSH( "Icons/icon_FontEd_ExportAll_40x", Icon20x20 ) );
+
+		Set( "FontEditor.FontBackgroundColor", new IMAGE_BRUSH( "Icons/icon_FontEd_Background_40x", Icon40x40 ) );
+		Set( "FontEditor.FontBackgroundColor.Small", new IMAGE_BRUSH( "Icons/icon_FontEd_Background_40x", Icon20x20 ) );
+		Set( "FontEditor.FontForegroundColor", new IMAGE_BRUSH( "Icons/icon_FontEd_Foreground_40x", Icon40x40 ) );
+		Set( "FontEditor.FontForegroundColor.Small", new IMAGE_BRUSH( "Icons/icon_FontEd_Foreground_40x", Icon20x20 ) );
 	}
 
 	// SoundCueGraph Editor
