Fixed missing icon for "Stats" toolbar button in the material editor
TTP 333374 - Material editor ui bug: an icon tied to visualizing material stats seems to be missing with the small ui icons option enabled

diff --git a/Engine/Source/Editor/EditorStyle/Private/SlateEditorStyle.cpp b/Engine/Source/Editor/EditorStyle/Private/SlateEditorStyle.cpp
--- a/Engine/Source/Editor/EditorStyle/Private/SlateEditorStyle.cpp
+++ b/Engine/Source/Editor/EditorStyle/Private/SlateEditorStyle.cpp
@@ -4531,8 +4531,8 @@ void FSlateEditorStyle::FStyle::SetupToolkitStyles()
 		Set( "MaterialEditor.TogglePreviewGrid", new IMAGE_BRUSH( "Icons/icon_MatEd_Grid_40x", Icon40x40 ) );
 		Set( "MaterialEditor.TogglePreviewGrid.Small", new IMAGE_BRUSH( "Icons/icon_MatEd_Grid_40x", Icon20x20 ) );
 
-		Set( "MaterialEditor.ToggleStats", new IMAGE_BRUSH( "Icons/icon_MatEd_Stats_40x", Icon40x40 ) );
-		Set( "MaterialEditor.ToggleStats.Small", new IMAGE_BRUSH( "Icons/icon_MatEd_Stats_40x", Icon20x20 ) );
+		Set( "MaterialEditor.ToggleMaterialStats", new IMAGE_BRUSH( "Icons/icon_MatEd_Stats_40x", Icon40x40 ) );
+		Set( "MaterialEditor.ToggleMaterialStats.Small", new IMAGE_BRUSH( "Icons/icon_MatEd_Stats_40x", Icon20x20 ) );
 		Set( "MaterialEditor.ToggleMobileStats", new IMAGE_BRUSH( "Icons/icon_MobileStats_40x", Icon40x40 ) );
 		Set( "MaterialEditor.ToggleMobileStats.Small", new IMAGE_BRUSH( "Icons/icon_MobileStats_40x", Icon20x20 ) );
 		Set( "MaterialEditor.CleanUnusedExpressions", new IMAGE_BRUSH( "Icons/icon_MatEd_CleanUp_40x", Icon40x40 ) );
