Material Editor: Release Stats and Build In Stats toolbar buttons now have icons.

#ttp 333477 - EDITOR ART REQUEST: "Release Stats" and "Built-in Stats" toolbar buttons in the material editor need UI / Art work.

#branch UE4

#change Added missing slate styles for buttons

diff --git a/Engine/Source/Editor/EditorStyle/Private/SlateEditorStyle.cpp b/Engine/Source/Editor/EditorStyle/Private/SlateEditorStyle.cpp
--- a/Engine/Source/Editor/EditorStyle/Private/SlateEditorStyle.cpp
+++ b/Engine/Source/Editor/EditorStyle/Private/SlateEditorStyle.cpp
@@ -4534,6 +4534,10 @@ void FSlateEditorStyle::FStyle::SetupToolkitStyles()
 
 		Set( "MaterialEditor.ToggleMaterialStats", new IMAGE_BRUSH( "Icons/icon_MatEd_Stats_40x", Icon40x40 ) );
 		Set( "MaterialEditor.ToggleMaterialStats.Small", new IMAGE_BRUSH( "Icons/icon_MatEd_Stats_40x", Icon20x20 ) );
+		Set( "MaterialEditor.ToggleReleaseStats", new IMAGE_BRUSH( "Icons/icon_MatEd_ReleaseStats_40x", Icon40x40 ) );
+		Set( "MaterialEditor.ToggleReleaseStats.Small", new IMAGE_BRUSH( "Icons/icon_MatEd_ReleaseStats_40x", Icon20x20 ) );
+		Set( "MaterialEditor.ToggleBuiltinStats", new IMAGE_BRUSH( "Icons/icon_MatEd_BuiltInStats_40x", Icon40x40 ) );
+		Set( "MaterialEditor.ToggleBuiltinStats.Small", new IMAGE_BRUSH( "Icons/icon_MatEd_BuiltInStats_40x", Icon20x20 ) );
 		Set( "MaterialEditor.ToggleMobileStats", new IMAGE_BRUSH( "Icons/icon_MobileStats_40x", Icon40x40 ) );
 		Set( "MaterialEditor.ToggleMobileStats.Small", new IMAGE_BRUSH( "Icons/icon_MobileStats_40x", Icon20x20 ) );
 		Set( "MaterialEditor.CleanUnusedExpressions", new IMAGE_BRUSH( "Icons/icon_MatEd_CleanUp_40x", Icon40x40 ) );
