Moved asset icons to their own folder and unified ClassIcon. and ClassThumbnail. styles.

FClassIconFinder can now find thumbnails for classes. Added many more asset type icons should they be required. Updated the AssetBackground and ThumbnailShadow images to make them a bit cleaner and smoother.

diff --git a/Engine/Source/Editor/BspMode/Private/BspModeStyle.cpp b/Engine/Source/Editor/BspMode/Private/BspModeStyle.cpp
--- a/Engine/Source/Editor/BspMode/Private/BspModeStyle.cpp
+++ b/Engine/Source/Editor/BspMode/Private/BspModeStyle.cpp
@@ -100,23 +100,6 @@ void FBspModeStyle::Initialize()
 	StyleSet->Set( "BspMode.AddVolume", new IMAGE_BRUSH( "Icons/icon_volume_40x", Icon40x40 ) );
 	StyleSet->Set( "BspMode.AddVolume.Small", new IMAGE_BRUSH( "Icons/icon_volume_40x", Icon20x20 ) );
 
-	StyleSet->Set( "BspMode.BlockingVolume", new IMAGE_BRUSH( "Icons/icon_volume_Blocking_16x", Icon16x16 ) );
-	StyleSet->Set( "BspMode.CullDistanceVolume", new IMAGE_BRUSH( "Icons/icon_volume_CullDistance_16x", Icon16x16 ) );
-	StyleSet->Set( "BspMode.LevelStreamingVolume", new IMAGE_BRUSH( "Icons/icon_volume_LevelStreaming_16x", Icon16x16 ) );
-	StyleSet->Set( "BspMode.LightmassCharacterIndirectDetailVolume", new IMAGE_BRUSH( "Icons/icon_volume_LightmassCharacterIndirectDetail_16x", Icon16x16 ) );
-	StyleSet->Set( "BspMode.LightmassImportanceVolume", new IMAGE_BRUSH( "Icons/icon_volume_LightmassImportance_16x", Icon16x16 ) );
-	StyleSet->Set( "BspMode.MassiveLODOverrideVolume", new IMAGE_BRUSH( "Icons/icon_volume_MassiveLODOverride_16x", Icon16x16 ) );
-	StyleSet->Set( "BspMode.NavMeshBoundsVolume", new IMAGE_BRUSH( "Icons/icon_volume_NavMeshBounds_16x", Icon16x16 ) );
-	StyleSet->Set( "BspMode.NavModifierVolume", new IMAGE_BRUSH( "Icons/icon_volume_NavMeshModifier_16x", Icon16x16 ) );
-	StyleSet->Set( "BspMode.PainCausingVolume", new IMAGE_BRUSH( "Icons/icon_volume_PainCausing_16x", Icon16x16 ) );
-	StyleSet->Set( "BspMode.PhysicsVolume", new IMAGE_BRUSH( "Icons/icon_volume_Physics_16x", Icon16x16 ) );
-	StyleSet->Set( "BspMode.PostProcessVolume", new IMAGE_BRUSH( "Icons/icon_volume_PostProcess_16x", Icon16x16 ) );
-	StyleSet->Set( "BspMode.PrecomputedVisibilityOverrideVolume", new IMAGE_BRUSH( "Icons/icon_volume_precomputedVisibilityOverride_16x", Icon16x16 ) );
-	StyleSet->Set( "BspMode.PrecomputedVisibilityVolume", new IMAGE_BRUSH( "Icons/icon_volume_PrecomputedVisibility_16x", Icon16x16 ) );
-	StyleSet->Set( "BspMode.ReverbVolume", new IMAGE_BRUSH( "Icons/icon_volume_Reverb_16x", Icon16x16 ) );
-	StyleSet->Set( "BspMode.ReverbVolumeToggleable", new IMAGE_BRUSH( "Icons/icon_volume_ReverbToggleable_16x", Icon16x16 ) );
-	StyleSet->Set( "BspMode.TriggerVolume", new IMAGE_BRUSH( "Icons/icon_volume_Trigger_16x", Icon16x16 ) );
-
 	StyleSet->Set( "BspMode.TableRow", FTableRowStyle()
 		.SetEvenRowBackgroundBrush( FSlateNoResource() )
 		.SetEvenRowBackgroundHoveredBrush( IMAGE_BRUSH( "Common/Selection", Icon8x8, FLinearColor(1.0f, 1.0f, 1.0f, 0.1f) ) )
diff --git a/Engine/Source/Editor/EditorStyle/Private/SlateEditorStyle.cpp b/Engine/Source/Editor/EditorStyle/Private/SlateEditorStyle.cpp
--- a/Engine/Source/Editor/EditorStyle/Private/SlateEditorStyle.cpp
+++ b/Engine/Source/Editor/EditorStyle/Private/SlateEditorStyle.cpp
@@ -1151,8 +1151,8 @@ void FSlateEditorStyle::FStyle::SetupGeneralStyles()
 	}
 		
 	// Asset Thumbnail
-		{
-		Set( "AssetThumbnail.AssetBackground", new IMAGE_BRUSH( "Common/AssetBackground", FVector2D(64.f, 64.f), FLinearColor(0.75f, 0.75f, 0.75f, 1.0f) ) );
+	{
+		Set( "AssetThumbnail.AssetBackground", new IMAGE_BRUSH( "Common/AssetBackground", FVector2D(64.f, 64.f), FLinearColor(1.f, 1.f, 1.f, 1.0f) ) );
 		Set( "AssetThumbnail.ClassBackground", new IMAGE_BRUSH( "Common/ClassBackground_64x", FVector2D(64.f, 64.f), FLinearColor(0.75f, 0.75f, 0.75f, 1.0f) ) );
 		Set( "AssetThumbnail.DataOnlyBPAssetBackground", new IMAGE_BRUSH( "Common/DataOnlyBPAssetBackground_64x", FVector2D(64.f, 64.f), FLinearColor(1, 1, 1, 1) ) );
 		Set( "AssetThumbnail.Font", TTF_CORE_FONT( "Fonts/Roboto-Regular", 12 ) );
@@ -3944,7 +3944,6 @@ void FSlateEditorStyle::FStyle::SetupClassIconsAndThumbnails()
 #if WITH_EDITOR
 	// Actor Classes Outliner
 	{
-		Set( "ClassIcon.Default", new IMAGE_BRUSH( "Icons/ActorIcons/Default_16x", Icon16x16 ) );
 		Set( "ClassIcon.Emitter", new IMAGE_BRUSH( "Icons/ActorIcons/Emitter_16x", Icon16x16 ) );
 		Set( "ClassIcon.Light", new IMAGE_BRUSH( "Icons/ActorIcons/LightActor_16x", Icon16x16 ) );
 		Set( "ClassIcon.PointLight", new IMAGE_BRUSH( "Icons/ActorIcons/PointLight_16x", Icon16x16 ) );
@@ -3974,7 +3973,6 @@ void FSlateEditorStyle::FStyle::SetupClassIconsAndThumbnails()
 		Set( "ClassIcon.Deleted", new IMAGE_BRUSH( "Icons/ActorIcons/DeletedActor_16px", Icon16x16 ) );
 		Set( "ClassIcon.DestructibleActor", new IMAGE_BRUSH( "Icons/ActorIcons/Destructable_16px", Icon16x16 ) );
 		Set( "ClassIcon.SceneCapture", new IMAGE_BRUSH( "Icons/ActorIcons/SceneCapture_16px", Icon16x16 ) );
-		Set( "ClassIcon.SceneCapture2D", new IMAGE_BRUSH( "Icons/ActorIcons/SceneCapture2D_16px", Icon16x16 ) );
 		Set( "ClassIcon.TextRenderActor", new IMAGE_BRUSH( "Icons/ActorIcons/TextRenderer_16px", Icon16x16 ) );
 		Set( "ClassIcon.BoxReflectionCapture", new IMAGE_BRUSH( "Icons/ActorIcons/BoxReflectionCapture_16px", Icon16x16 ) );
 		Set( "ClassIcon.SphereReflectionCapture", new IMAGE_BRUSH( "Icons/ActorIcons/SphereReflectionCapture_16px", Icon16x16 ) );
@@ -4031,142 +4029,113 @@ void FSlateEditorStyle::FStyle::SetupClassIconsAndThumbnails()
 		Set( "ClassIcon.ActorFactorySphereVolume", new IMAGE_BRUSH( "Icons/icon_volume_sphere_16x", Icon16x16 ) );
 
 		// Asset Type Classes
-		Set( "ClassIcon.AimOffsetBlendSpace", new IMAGE_BRUSH( "Icons/AssetIcons/AimOffsetBlendSpace_16x", Icon16x16 ) );
-		Set( "ClassIcon.AimOffsetBlendSpace1D", new IMAGE_BRUSH( "Icons/AssetIcons/AimOffsetBlendSpace1D_16x", Icon16x16 ) );
-		Set( "ClassIcon.AnimBlueprint", new IMAGE_BRUSH( "Icons/AssetIcons/AnimBlueprint_16x", Icon16x16 ) );
-		Set( "ClassIcon.AnimComposite", new IMAGE_BRUSH( "Icons/AssetIcons/AnimComposite_16x", Icon16x16 ) );
-		Set( "ClassIcon.AnimMontage", new IMAGE_BRUSH( "Icons/AssetIcons/AnimMontage_16x", Icon16x16 ) );
-		Set( "ClassIcon.BlendSpace", new IMAGE_BRUSH( "Icons/AssetIcons/BlendSpace_16x", Icon16x16 ) );
-		Set( "ClassIcon.BlendSpace1D", new IMAGE_BRUSH( "Icons/AssetIcons/BlendSpace1D_16x", Icon16x16 ) );
-		Set( "ClassIcon.Blueprint", new IMAGE_BRUSH( "Icons/AssetIcons/Blueprint_16x", Icon16x16 ) );
-		Set( "ClassIcon.BlueprintGeneratedClass", new IMAGE_BRUSH( "Icons/AssetIcons/Blueprint_16x", Icon16x16 ) );
-		Set( "ClassIcon.BlueprintInterface", new IMAGE_BRUSH( "Icons/AssetIcons/Blueprint_Interface_16x", Icon16x16 ) );
-		Set( "ClassIcon.BlueprintMacroLibrary", new IMAGE_BRUSH( "Icons/AssetIcons/Blueprint_MacroLibrary_16x", Icon16x16 ) );
-		Set( "ClassIcon.ButtonStyleAsset", new IMAGE_BRUSH( "Icons/AssetIcons/ButtonStyle_16x", Icon16x16 ) );
-		Set( "ClassIcon.CameraAnim", new IMAGE_BRUSH( "Icons/ActorIcons/Camera_16x", Icon16x16 ) );
-		Set( "ClassIcon.CurveBase", new IMAGE_BRUSH( "Icons/AssetIcons/Curve_16x", Icon16x16 ) );
-		Set( "ClassIcon.DataAsset", new IMAGE_BRUSH( "Icons/AssetIcons/DataAsset_16x", Icon16x16 ) );
-		Set( "ClassIcon.Font", new IMAGE_BRUSH( "Icons/AssetIcons/Font_16x", Icon16x16 ) );
-		Set( "ClassIcon.InterpData", new IMAGE_BRUSH( "Icons/AssetIcons/MatineeData_16x", Icon16x16 ) );
-		Set( "ClassIcon.Material", new IMAGE_BRUSH( "Icons/AssetIcons/Material_16x", Icon16x16 ) );
-		Set( "ClassIcon.MaterialFunction", new IMAGE_BRUSH( "Icons/AssetIcons/MaterialFunction_16x", Icon16x16 ) );
-		Set( "ClassIcon.MaterialInstanceConstant", new IMAGE_BRUSH( "Icons/AssetIcons/MaterialInstanceConstant_16x", Icon16x16 ) );
-		Set( "ClassIcon.MaterialParameterCollection", new IMAGE_BRUSH( "Icons/AssetIcons/MaterialParameterCollection_16x", Icon16x16 ) );
-		Set( "ClassIcon.MultiFont", new IMAGE_BRUSH( "Icons/AssetIcons/MultiFont_16x", Icon16x16 ) );
-		Set( "ClassIcon.ObjectLibrary", new IMAGE_BRUSH( "Icons/AssetIcons/ObjectLibrary_16x", Icon16x16 ) );
-		Set( "ClassIcon.PaperFlipbook", new IMAGE_BRUSH( "Icons/AssetIcons/PaperFlipbook_16x", Icon16x16 ) );
-		Set( "ClassIcon.ParticleSystem", new IMAGE_BRUSH( "Icons/AssetIcons/ParticleSystem_16x", Icon16x16 ) );
-		Set( "ClassIcon.PhysicalMaterial", new IMAGE_BRUSH( "Icons/AssetIcons/PhysicalMaterial_16x", Icon16x16 ) );
-		Set( "ClassIcon.ReverbEffect", new IMAGE_BRUSH( "Icons/AssetIcons/ReverbEffect_16x", Icon16x16 ) );
-		Set( "ClassIcon.SlateBrushAsset", new IMAGE_BRUSH( "Icons/AssetIcons/SlateBrush_16x", Icon16x16 ) );
-		Set( "ClassIcon.SoundAttenuation", new IMAGE_BRUSH( "Icons/AssetIcons/SoundAttenuation_16x", Icon16x16 ) );
-		Set( "ClassIcon.SoundClass", new IMAGE_BRUSH( "Icons/AssetIcons/SoundClass_16x", Icon16x16 ) );
-		Set( "ClassIcon.SoundCue", new IMAGE_BRUSH( "Icons/AssetIcons/SoundCue_16x", Icon16x16 ) );
-		Set( "ClassIcon.SoundMix", new IMAGE_BRUSH( "Icons/AssetIcons/SoundMix_16x", Icon16x16 ) );
-		Set( "ClassIcon.StaticMesh", new IMAGE_BRUSH( "Icons/AssetIcons/StaticMesh_16x", Icon16x16 ) );
-		Set( "ClassIcon.Texture", new IMAGE_BRUSH( "Icons/AssetIcons/Texture_16x", Icon16x16 ) );
-		Set( "ClassIcon.Texture2D", new IMAGE_BRUSH( "Icons/AssetIcons/Texture2D_16x", Icon16x16 ) );
-		Set( "ClassIcon.TextureRenderTarget2D", new IMAGE_BRUSH( "Icons/AssetIcons/TextureRenderTarget2D_16x", Icon16x16 ) );
-		Set( "ClassIcon.TextureRenderTargetCube", new IMAGE_BRUSH( "Icons/AssetIcons/TextureRenderTargetCube_16x", Icon16x16 ) );
-		Set( "ClassIcon.UserDefinedEnum", new IMAGE_BRUSH( "Icons/AssetIcons/Enumeration_16x", Icon16x16 ) );
-		Set( "ClassIcon.UserDefinedStruct", new IMAGE_BRUSH( "Icons/AssetIcons/Structure_16x", Icon16x16 ) );
-		Set( "ClassIcon.SlateWidgetStyleAsset", new IMAGE_BRUSH("Icons/AssetIcons/SlateStyle_16x", Icon16x16 ) );
-
-		// Builder brushes
-		Set( "ClassIcon.CubeBuilder", new IMAGE_BRUSH("Icons/icon_box_40x", Icon16x16 ) );
-		Set( "ClassIcon.ConeBuilder", new IMAGE_BRUSH("Icons/icon_cone_40x", Icon16x16 ) );
-		Set( "ClassIcon.CylinderBuilder", new IMAGE_BRUSH("Icons/icon_cylinder_40x", Icon16x16 ) );
-		Set( "ClassIcon.CurvedStairBuilder", new IMAGE_BRUSH("Icons/icon_curvedstair_40x", Icon16x16 ) );
-		Set( "ClassIcon.LinearStairBuilder", new IMAGE_BRUSH("Icons/icon_stair_40x", Icon16x16 ) );
-		Set( "ClassIcon.SpiralStairBuilder", new IMAGE_BRUSH("Icons/icon_spiralstair_40x", Icon16x16 ) );
-		Set( "ClassIcon.VolumetricBuilder", new IMAGE_BRUSH("Icons/icon_volumetric_40x", Icon16x16 ) );
-		Set( "ClassIcon.TetrahedronBuilder", new IMAGE_BRUSH("Icons/icon_sphere_40x", Icon16x16 ) );
-		Set( "ClassIcon.SheetBuilder", new IMAGE_BRUSH("Icons/icon_plane_40x", Icon16x16 ) );
-	}
-
-	{
-		Set( "ClassThumbnail.AmbientSound", new IMAGE_BRUSH( "Icons/icon_class_AmbientSound_64px", Icon64x64 ) );
-		Set( "ClassThumbnail.BoxReflectionCapture", new IMAGE_BRUSH( "Icons/icon_class_BoxReflectionCapture_64px", Icon64x64 ) );
-		Set( "ClassThumbnail.CameraActor", new IMAGE_BRUSH( "Icons/icon_class_CameraActor_64px", Icon64x64 ) );
-		Set( "ClassThumbnail.DecalActor", new IMAGE_BRUSH( "Icons/icon_class_DecalActor_64px", Icon64x64 ) );
-		Set( "ClassThumbnail.MatineeActor", new IMAGE_BRUSH( "Icons/icon_class_MatineeActor_64px", Icon64x64 ) );
-		Set( "ClassThumbnail.PlaneReflectionCapture", new IMAGE_BRUSH( "Icons/icon_class_PlaneReflectionCapture_64px", Icon64x64 ) );
-		Set( "ClassThumbnail.PlayerStart", new IMAGE_BRUSH( "Icons/icon_class_PlayerStart_64px", Icon64x64 ) );
-		Set( "ClassThumbnail.SphereReflectionCapture", new IMAGE_BRUSH( "Icons/icon_class_SphereReflectionCapture_64px", Icon64x64 ) );
-		Set( "ClassThumbnail.TextRenderActor", new IMAGE_BRUSH( "Icons/icon_class_TextRenderActor_64px", Icon64x64 ) );
-		Set( "ClassThumbnail.TriggerBox", new IMAGE_BRUSH( "Icons/icon_class_TriggerBox_64px", Icon64x64 ) );
-		Set( "ClassThumbnail.TriggerCapsule", new IMAGE_BRUSH( "Icons/icon_class_TriggerCapsule_64px", Icon64x64 ) );
-		Set( "ClassThumbnail.TriggerSphere", new IMAGE_BRUSH( "Icons/icon_class_TriggerSphere_64px", Icon64x64 ) );
-
-		Set( "ClassThumbnail.NavLinkProxy", new IMAGE_BRUSH( "Icons/icon_class_NavLinkProxy_64x", Icon64x64 ) );
-		Set( "ClassThumbnail.Note", new IMAGE_BRUSH( "Icons/icon_class_Notes_64x", Icon64x64 ) );
-		Set( "ClassThumbnail.RadialForceActor", new IMAGE_BRUSH( "Icons/icon_class_RadialForceActor_64x", Icon64x64 ) );
-		Set( "ClassThumbnail.TargetPoint", new IMAGE_BRUSH( "Icons/icon_class_TargetPoint_64x", Icon64x64 ) );
-		Set( "ClassThumbnail.AtmosphericFog", new IMAGE_BRUSH( "Icons/icon_class_AtmosphericFog_64x", Icon64x64 ) );
-		Set( "ClassThumbnail.ExponentialHeightFog", new IMAGE_BRUSH( "Icons/icon_class_ExpHeightFog_64x", Icon64x64 ) );
-
-		Set( "ClassThumbnail.TriggerBase", new IMAGE_BRUSH( "Icons/icon_class_TriggerBox_64px", Icon64x64 ) );
-		Set( "ClassThumbnail.DirectionalLight", new IMAGE_BRUSH( "Icons/icon_class_DirectionalLight_64px", Icon64x64 ) );
-		Set( "ClassThumbnail.PointLight", new IMAGE_BRUSH( "Icons/icon_class_PointLight_64px", Icon64x64 ) );
-		Set( "ClassThumbnail.SpotLight", new IMAGE_BRUSH( "Icons/icon_class_SpotLight_64px", Icon64x64 ) );
-		Set( "ClassThumbnail.SkyLight", new IMAGE_BRUSH( "Icons/icon_class_Skylight_64px", Icon64x64 ) ); 
-		Set( "ClassThumbnail.ReflectionCapture", new IMAGE_BRUSH( "Icons/icon_class_SphereReflectionCapture_64px", Icon64x64 ) );
-
-		Set( "ClassThumbnail.BlockingVolume", new IMAGE_BRUSH( "Icons/icon_volume_Blocking_64x", Icon64x64 ) );
-		Set( "ClassThumbnail.CullDistanceVolume", new IMAGE_BRUSH( "Icons/icon_volume_CullDistance_64x", Icon64x64 ) );
-		Set( "ClassThumbnail.LevelStreamingVolume", new IMAGE_BRUSH( "Icons/icon_volume_LevelStreaming_64x", Icon64x64 ) );
-		Set( "ClassThumbnail.LightmassCharacterIndirectDetailVolume", new IMAGE_BRUSH( "Icons/icon_volume_LightmassCharacterIndirectDetail_64x", Icon64x64 ) );
-		Set( "ClassThumbnail.LightmassImportanceVolume", new IMAGE_BRUSH( "Icons/icon_volume_LightmassImportance_64x", Icon64x64 ) );
-		Set( "ClassThumbnail.MassiveLODOverrideVolume", new IMAGE_BRUSH( "Icons/icon_volume_MassiveLODOverride_64x", Icon64x64 ) );
-		Set( "ClassThumbnail.NavMeshBoundsVolume", new IMAGE_BRUSH( "Icons/icon_volume_NavMeshBounds_64x", Icon64x64 ) );
-		Set( "ClassThumbnail.NavModifierVolume", new IMAGE_BRUSH( "Icons/icon_volume_NavMeshModifier_64x", Icon64x64 ) );
-		Set( "ClassThumbnail.PainCausingVolume", new IMAGE_BRUSH( "Icons/icon_volume_PainCausing_64x", Icon64x64 ) );
-		Set( "ClassThumbnail.PhysicsVolume", new IMAGE_BRUSH( "Icons/icon_volume_Physics_64x", Icon64x64 ) );
-		Set( "ClassThumbnail.PostProcessVolume", new IMAGE_BRUSH( "Icons/icon_volume_PostProcess_64x", Icon64x64 ) );
-		Set( "ClassThumbnail.PrecomputedVisibilityOverrideVolume", new IMAGE_BRUSH( "Icons/icon_volume_precomputedVisibilityOverride_64x", Icon64x64 ) );
-		Set( "ClassThumbnail.PrecomputedVisibilityVolume", new IMAGE_BRUSH( "Icons/icon_volume_PrecomputedVisibility_64x", Icon64x64 ) );
-		Set( "ClassThumbnail.ReverbVolume", new IMAGE_BRUSH( "Icons/icon_volume_Reverb_64x", Icon64x64 ) );
-		Set( "ClassThumbnail.ReverbVolumeToggleable", new IMAGE_BRUSH( "Icons/icon_volume_ReverbToggleable_64x", Icon64x64 ) );
-		Set( "ClassThumbnail.TriggerVolume", new IMAGE_BRUSH( "Icons/icon_volume_Trigger_64x", Icon64x64 ) );
-
-		// Asset Types
-		Set( "ClassThumbnail.AimOffsetBlendSpace", new IMAGE_BRUSH( "Icons/icon_class_AimOffset_64x", Icon32x32 ) );
-		Set( "ClassThumbnail.AimOffsetBlendSpace1D", new IMAGE_BRUSH( "Icons/icon_class_AimOffset1D_64x", Icon32x32 ) );
-		Set( "ClassThumbnail.AnimBlueprint", new IMAGE_BRUSH( "Icons/icon_class_AnimBlueprint_64x", Icon32x32 ) );
-		Set( "ClassThumbnail.AnimComposite", new IMAGE_BRUSH( "Icons/icon_class_AnimComposite_64x", Icon32x32 ) );
-		Set( "ClassThumbnail.AnimMontage", new IMAGE_BRUSH( "Icons/icon_class_AnimMontage_64x", Icon32x32 ) );
-		Set( "ClassThumbnail.BlendSpace", new IMAGE_BRUSH( "Icons/icon_class_BlendSpace_64x", Icon32x32 ) );
-		Set( "ClassThumbnail.BlendSpace1D", new IMAGE_BRUSH( "Icons/icon_class_BlendSpace1D_64x", Icon32x32 ) );
-		Set( "ClassThumbnail.Blueprint", new IMAGE_BRUSH( "Icons/icon_class_Blueprint_64x", Icon32x32 ) );
-		Set( "ClassThumbnail.BlueprintGeneratedClass", new IMAGE_BRUSH( "Icons/icon_class_Blueprint_64x", Icon32x32 ) );
-		Set( "ClassThumbnail.BlueprintInterface", new IMAGE_BRUSH( "Icons/icon_class_Blueprint_Interface_64x", Icon32x32 ) );
-		Set( "ClassThumbnail.BlueprintMacroLibrary", new IMAGE_BRUSH( "Icons/icon_class_Blueprint_MacroLibrary_64x", Icon32x32 ) );
-		Set( "ClassThumbnail.BlueprintFunctionLibrary", new IMAGE_BRUSH("Icons/icon_class_Blueprint_FunctionLibrary_64x", Icon32x32));
-		Set( "ClassThumbnail.ButtonStyleAsset", new IMAGE_BRUSH( "Icons/icon_class_ButtonStyle_64x", Icon32x32 ) );
-		Set( "ClassThumbnail.CameraAnim", new IMAGE_BRUSH( "Icons/icon_class_CameraAnim_64px", Icon32x32 ) );
-		Set( "ClassThumbnail.CurveBase", new IMAGE_BRUSH( "Icons/icon_class_Curve_64x", Icon32x32 ) );
-		Set( "ClassThumbnail.DataAsset", new IMAGE_BRUSH( "Icons/icon_class_DataAsset_64x", Icon32x32 ) );
-		Set( "ClassThumbnail.Font", new IMAGE_BRUSH( "Icons/icon_class_Font_64x", Icon32x32 ) );
-		Set( "ClassThumbnail.InterpData", new IMAGE_BRUSH( "Icons/icon_class_MatineeData_64x", Icon32x32 ) );
-		Set( "ClassThumbnail.Material", new IMAGE_BRUSH( "Icons/icon_class_Material_64x", Icon32x32 ) );
-		Set( "ClassThumbnail.MaterialFunction", new IMAGE_BRUSH( "Icons/icon_class_MaterialFunction_64x", Icon32x32 ) );
-		Set( "ClassThumbnail.MaterialInstanceConstant", new IMAGE_BRUSH( "Icons/icon_class_MaterialInstanceConstant_64x", Icon32x32 ) );
-		Set( "ClassThumbnail.MaterialParameterCollection", new IMAGE_BRUSH( "Icons/icon_class_MaterialParameterCollection_64x", Icon32x32 ) );
-		Set( "ClassThumbnail.MultiFont", new IMAGE_BRUSH( "Icons/icon_class_Multifont_64x", Icon32x32 ) );
-		Set( "ClassThumbnail.ObjectLibrary", new IMAGE_BRUSH( "Icons/icon_class_ObjectLibrary_64x", Icon32x32 ) );
-		Set( "ClassThumbnail.ParticleSystem", new IMAGE_BRUSH( "Icons/icon_class_ParticleSystem_64x", Icon32x32 ) );
-		Set( "ClassThumbnail.PhysicalMaterial", new IMAGE_BRUSH( "Icons/icon_class_PhysicalMaterial_64x", Icon32x32 ) );
-		Set( "ClassThumbnail.ReverbEffect", new IMAGE_BRUSH( "Icons/icon_class_ReverbEffect_64x", Icon32x32 ) );
-		Set( "ClassThumbnail.SlateBrushAsset", new IMAGE_BRUSH( "Icons/icon_class_SlateBrush_64x", Icon32x32 ) );
-		Set( "ClassThumbnail.SoundAttenuation", new IMAGE_BRUSH( "Icons/icon_class_SoundAttenuation_64x", Icon32x32 ) );
-		Set( "ClassThumbnail.SoundClass", new IMAGE_BRUSH( "Icons/icon_class_SoundClass_64x", Icon32x32 ) );
-		Set( "ClassThumbnail.SoundCue", new IMAGE_BRUSH( "Icons/icon_class_SoundCue_64x", Icon32x32 ) );
-		Set( "ClassThumbnail.SoundMix", new IMAGE_BRUSH( "Icons/icon_class_SoundMix_64x", Icon32x32 ) );
-		Set( "ClassThumbnail.TextureRenderTarget2D", new IMAGE_BRUSH( "Icons/icon_class_TextureRenderTarget2D_64x", Icon32x32 ) );
-		Set( "ClassThumbnail.TextureRenderTargetCube", new IMAGE_BRUSH( "Icons/icon_class_TextureRenderTargetCube_64x", Icon32x32 ) );
-		Set( "ClassThumbnail.UserDefinedEnum", new IMAGE_BRUSH( "Icons/icon_class_Enumeration_64x", Icon32x32 ) );
-		Set( "ClassThumbnail.UserDefinedStruct", new IMAGE_BRUSH( "Icons/icon_class_Structure_64x", Icon32x32 ) );
-		Set( "ClassThumbnail.SlateWidgetStyleAsset", new IMAGE_BRUSH("Icons/icon_class_SlateStyle_64x", Icon32x32 ) );
+		const TCHAR* AssetTypes[] = {
+
+			TEXT("AimOffsetBlendSpace"),
+			TEXT("AimOffsetBlendSpace1D"),
+			TEXT("AmbientSound"),
+			TEXT("AnimBlueprint"),
+			TEXT("AnimComposite"),
+			TEXT("AnimMontage"),
+			TEXT("AnimSequence"),
+			TEXT("AtmosphericFog"),
+			TEXT("BehaviorTree"),
+			TEXT("BlendSpace"),
+			TEXT("BlendSpace1D"),
+			TEXT("BlockingVolume"),
+			TEXT("Blueprint"),
+			TEXT("BlueprintFunctionLibrary"),
+			TEXT("BlueprintInterface"),
+			TEXT("BlueprintMacroLibrary"),
+			TEXT("BoxReflectionCapture"),
+			TEXT("ButtonStyleAsset"),
+			TEXT("CableActor"),
+			TEXT("CameraActor"),
+			TEXT("CameraAnim"),
+			TEXT("CameraBlockingVolume"),
+			TEXT("Class"),
+			TEXT("CullDistanceVolume"),
+			TEXT("CurveBase"),
+			TEXT("DataAsset"),
+			TEXT("DecalActor"),
+			TEXT("Default"),
+			TEXT("DirectionalLight"),
+			TEXT("DirectionalLightMovable"),
+			TEXT("DirectionalLightStatic"),
+			TEXT("DirectionalLightStationary"),
+			TEXT("ExponentialHeightFog"),
+			TEXT("Font"),
+			TEXT("ForceFeedbackEffect"),
+			TEXT("InterpData"),
+			TEXT("KillZVolume"),
+			TEXT("LevelStreamingVolume"),
+			TEXT("LightmassCharacterIndirectDetailVolume"),
+			TEXT("LightmassImportanceVolume"),
+			TEXT("MassiveLODOverrideVolume"),
+			TEXT("Material"),
+			TEXT("MaterialFunction"),
+			TEXT("MaterialInstanceConstant"),
+			TEXT("MaterialParameterCollection"),
+			TEXT("MultiFont"),
+			TEXT("NavLinkProxy"),
+			TEXT("NavMeshBoundsVolume"),
+			TEXT("NavModifierVolume"),
+			TEXT("Note"),
+			TEXT("ObjectLibrary"),
+			TEXT("PainCausingVolume"),
+			TEXT("ParticleSystem"),
+			TEXT("PhysicalMaterial"),
+			TEXT("PhysicsVolume"),
+			TEXT("PlaneReflectionCapture"),
+			TEXT("PlayerStart"),
+			TEXT("PointLight"),
+			TEXT("PostProcessVolume"),
+			TEXT("PrecomputedVisibilityOverrideVolume"),
+			TEXT("PrecomputedVisibilityVolume"),
+			TEXT("RadialForceActor"),
+			TEXT("ReflectionCapture"),
+			TEXT("ReverbEffect"),
+			TEXT("ReverbVolume"),
+			TEXT("ReverbVolumeToggleable"),
+			TEXT("SceneCapture2D"),
+			TEXT("SkyLight"),
+			TEXT("SkeletalMesh"),
+			TEXT("SlateBrushAsset"),
+			TEXT("SlateWidgetStyleAsset"),
+			TEXT("SoundAttenuation"),
+			TEXT("SoundClass"),
+			TEXT("SoundCue"),
+			TEXT("SoundMix"),
+			TEXT("SphereReflectionCapture"),
+			TEXT("SpotLight"),
+			TEXT("SpotLightMovable"),
+			TEXT("SpotLightStatic"),
+			TEXT("SpotLightStationary"),
+			TEXT("StaticMesh"),
+			TEXT("TargetPoint"),
+			TEXT("TextRenderActor"),
+			TEXT("Texture2D"),
+			TEXT("TextureRenderTarget2D"),
+			TEXT("TextureRenderTargetCube"),
+			TEXT("TriggerBase"),
+			TEXT("TriggerBox"),
+			TEXT("TriggerCapsule"),
+			TEXT("TriggerSphere"),
+			TEXT("TriggerVolume"),
+			TEXT("TouchInterface"),
+			TEXT("UserDefinedEnum"),
+			TEXT("UserDefinedStruct"),
+		};
+
+		const TCHAR* StyleString = TEXT("ClassIcon.%s%s");
+		const TCHAR* Path = TEXT("Icons/AssetIcons/%s_%dx");
+		
+		for (int32 TypeIndex = 0; TypeIndex < ARRAY_COUNT(AssetTypes); ++TypeIndex)
+		{
+			const TCHAR* Type = AssetTypes[TypeIndex];
+			Set( *FString::Printf(TEXT("ClassIcon.%s"), Type),		new IMAGE_BRUSH(FString::Printf(Path, Type, 16), Icon16x16 ) );
+			Set( *FString::Printf(TEXT("ClassThumbnail.%s"), Type),	new IMAGE_BRUSH(FString::Printf(Path, Type, 64), Icon64x64 ) );
+		}
 	}
 #endif
 }
@@ -4711,7 +4680,7 @@ void FSlateEditorStyle::FStyle::SetupToolkitStyles()
 
 	// Curve Editor
 	{
-		Set( "CurveAssetEditor.Tabs.Properties", new IMAGE_BRUSH( "Icons/AssetIcons/Curve_16x", Icon16x16 ) );
+		Set( "CurveAssetEditor.Tabs.Properties", new IMAGE_BRUSH( "Icons/AssetIcons/CurveBase_16x", Icon16x16 ) );
 
 		Set( "CurveEditor.FitHorizontally", new IMAGE_BRUSH( "Icons/icon_CurveEditor_Horizontal_40x", Icon40x40 ) );
 		Set( "CurveEditor.FitVertically", new IMAGE_BRUSH( "Icons/icon_CurveEditor_Vertical_40x", Icon40x40 ) );
diff --git a/Engine/Source/Editor/UnrealEd/Private/ClassIconFinder.cpp b/Engine/Source/Editor/UnrealEd/Private/ClassIconFinder.cpp
--- a/Engine/Source/Editor/UnrealEd/Private/ClassIconFinder.cpp
+++ b/Engine/Source/Editor/UnrealEd/Private/ClassIconFinder.cpp
@@ -79,26 +79,40 @@ FName FClassIconFinder::FindIconNameForActor( const TWeakObjectPtr<AActor>& InAc
 	return BrushName;
 }
 
+const FSlateBrush* FClassIconFinder::FindIconForClass(const UClass* InClass, const FName& InDefaultName )
+{
+	return FEditorStyle::GetBrush( FindIconNameImpl( InClass, InDefaultName ) );
+}
+
+FName FClassIconFinder::FindIconNameForClass(const UClass* InClass, const FName& InDefaultName )
+{
+	return FindIconNameImpl( InClass, InDefaultName );
+}
+
+const FSlateBrush* FClassIconFinder::FindThumbnailForClass(const UClass* InClass, const FName& InDefaultName )
+{
+	return FEditorStyle::GetBrush( FindIconNameImpl( InClass, InDefaultName, TEXT("ClassThumbnail") ) );
+}
 
-const FSlateBrush* FClassIconFinder::FindIconForClass(UClass* InClass, const FName& InDefaultName )
+FName FClassIconFinder::FindThumbnailNameForClass(const UClass* InClass, const FName& InDefaultName )
 {
-	return FEditorStyle::GetBrush( FindIconNameForClass( InClass, InDefaultName ) );
+	return FindIconNameImpl( InClass, InDefaultName, TEXT("ClassThumbnail") );
 }
 
-FName FClassIconFinder::FindIconNameForClass(UClass* InClass, const FName& InDefaultName )
+FName FClassIconFinder::FindIconNameImpl(const UClass* InClass, const FName& InDefaultName, const TCHAR* StyleRoot )
 {
-	FName BrushName = InDefaultName;
+	FName BrushName;
 	const FSlateBrush* Brush = NULL;
 
 	if ( InClass != NULL )
 	{
 		// walk up class hierarchy until we find an icon
-		UClass* ActorClass = InClass;
-		while( Brush == NULL && ActorClass && ActorClass != AActor::StaticClass() )
+		const UClass* CurrentClass = InClass;
+		while( Brush == NULL && CurrentClass && (CurrentClass != AActor::StaticClass()) )
 		{
-			BrushName = *FString::Printf( TEXT( "ClassIcon.%s" ), *ActorClass->GetName() );
+			BrushName = *FString::Printf( TEXT( "%s.%s" ), StyleRoot, *CurrentClass->GetName() );
 			Brush = FEditorStyle::GetOptionalBrush( BrushName, nullptr, nullptr );
-			ActorClass = ActorClass->GetSuperClass();
+			CurrentClass = CurrentClass->GetSuperClass();
 		}
 	}
 
@@ -107,7 +121,7 @@ FName FClassIconFinder::FindIconNameForClass(UClass* InClass, const FName& InDef
 		// If we didn't supply an override name for the default icon use default class icon.
 		if( InDefaultName == "" )
 		{
-			BrushName = TEXT( "ClassIcon.Default" );
+			BrushName = *FString::Printf( TEXT( "%s.Default" ), StyleRoot);
 		}
 		else
 		{
diff --git a/Engine/Source/Editor/UnrealEd/Public/ClassIconFinder.h b/Engine/Source/Editor/UnrealEd/Public/ClassIconFinder.h
--- a/Engine/Source/Editor/UnrealEd/Public/ClassIconFinder.h
+++ b/Engine/Source/Editor/UnrealEd/Public/ClassIconFinder.h
@@ -19,10 +19,20 @@ public:
 	UNREALED_API static FName FindIconNameForActor( const TWeakObjectPtr<AActor>& InActor );
 
 	/** Find the small icon to use for the supplied class */
-	UNREALED_API static const FSlateBrush* FindIconForClass(UClass* InClass, const FName& InDefaultName = FName() );
+	UNREALED_API static const FSlateBrush* FindIconForClass(const UClass* InClass, const FName& InDefaultName = FName() );
 
 	/** Find the small icon name to use for the supplied class */
-	UNREALED_API static FName FindIconNameForClass(UClass* InClass, const FName& InDefaultName = FName() );
+	UNREALED_API static FName FindIconNameForClass(const UClass* InClass, const FName& InDefaultName = FName() );
+
+	/** Find the small icon to use for the supplied class */
+	UNREALED_API static const FSlateBrush* FindThumbnailForClass(const UClass* InClass, const FName& InDefaultName = FName() );
+
+	/** Find the large thumbnail name to use for the supplied class */
+	UNREALED_API static FName FindThumbnailNameForClass(const UClass* InClass, const FName& InDefaultName = FName() );
+
+private:
+	/** Find a thumbnail/icon name to use for the supplied class */
+	UNREALED_API static FName FindIconNameImpl(const UClass* InClass, const FName& InDefaultName = FName(), const TCHAR* StyleRoot = TEXT("ClassIcon") );
 };
 
 #endif // __ActorIconFinder_h__
\ No newline at end of file
