Re-added black strokes to smaller class icon styles, and fixed some erroneous perforce reconciliation

It appears that I had some issues with a combination of unshelving and perforce reconcile not giving me the most up-to-date version of some files which resulted in the black stroke being missed off a lot of class icon styles.

diff --git a/Engine/Source/Editor/EditorStyle/Private/SlateEditorStyle.cpp b/Engine/Source/Editor/EditorStyle/Private/SlateEditorStyle.cpp
--- a/Engine/Source/Editor/EditorStyle/Private/SlateEditorStyle.cpp
+++ b/Engine/Source/Editor/EditorStyle/Private/SlateEditorStyle.cpp
@@ -3947,73 +3947,56 @@ void FSlateEditorStyle::FStyle::SetupClassIconsAndThumbnails()
 	{
 		Set( "ClassIcon.Emitter", new IMAGE_BRUSH( "Icons/ActorIcons/Emitter_16x", Icon16x16 ) );
 		Set( "ClassIcon.Light", new IMAGE_BRUSH( "Icons/ActorIcons/LightActor_16x", Icon16x16 ) );
-		Set( "ClassIcon.PointLight", new IMAGE_BRUSH( "Icons/ActorIcons/PointLight_16x", Icon16x16 ) );
-		Set( "ClassIcon.SpotLight", new IMAGE_BRUSH( "Icons/ActorIcons/SpotLight_16x", Icon16x16 ) );
-		Set( "ClassIcon.DirectionalLight", new IMAGE_BRUSH( "Icons/ActorIcons/DirectionalLight_16x", Icon16x16 ) );
 		Set( "ClassIcon.MatineeActor", new IMAGE_BRUSH( "Icons/ActorIcons/MatineeActor_16x", Icon16x16 ) );
-		Set( "ClassIcon.PlayerStart", new IMAGE_BRUSH( "Icons/ActorIcons/PlayerStart_16x", Icon16x16 ) );
-		Set( "ClassIcon.AmbientSound", new IMAGE_BRUSH( "Icons/ActorIcons/SoundActor_16x", Icon16x16 ) );
-		Set( "ClassIcon.TriggerBase", new IMAGE_BRUSH( "Icons/ActorIcons/TriggerActor_16x", Icon16x16 ) );
 		Set( "ClassIcon.Brush", new IMAGE_BRUSH( "Icons/ActorIcons/Brush_16x", Icon16x16 ) );
 		Set( "ClassIcon.BrushAdditive", new IMAGE_BRUSH( "Icons/ActorIcons/Brush_Add_16x", Icon16x16 ) );
 		Set( "ClassIcon.BrushSubtractive", new IMAGE_BRUSH( "Icons/ActorIcons/Brush_Subtract_16x", Icon16x16 ) );
 		Set( "ClassIcon.Volume", new IMAGE_BRUSH( "Icons/ActorIcons/Volume_16x", Icon16x16 ) );
-		Set( "ClassIcon.BlockingVolume", new IMAGE_BRUSH( "Icons/ActorIcons/BlockingVolume_16x", Icon16x16 ) );
-		Set( "ClassIcon.BlueprintActorBase", new IMAGE_BRUSH( "Icons/ActorIcons/Blueprint_16x", Icon16x16 ) );
-		Set( "ClassIcon.StaticMeshActor", new IMAGE_BRUSH( "Icons/ActorIcons/StaticMesh_16x", Icon16x16 ) );
-		Set( "ClassIcon.SkeletalMeshActor", new IMAGE_BRUSH( "Icons/ActorIcons/SkeletalMesh_16x", Icon16x16 ) );
-		Set( "ClassIcon.CameraActor", new IMAGE_BRUSH( "Icons/ActorIcons/Camera_16x", Icon16x16 ) );
-		Set( "ClassIcon.DecalActor", new IMAGE_BRUSH( "Icons/ActorIcons/Decal_16x", Icon16x16 ) );
 		Set( "ClassIcon.SkeletalPhysicsActor", new IMAGE_BRUSH( "Icons/ActorIcons/SkeletalPhysicsActor_16x", Icon16x16 ) );
-		Set( "ClassIcon.ExponentialHeightFog", new IMAGE_BRUSH( "Icons/ActorIcons/ExpoHeightFog_16x", Icon16x16 ) );
-		Set( "ClassIcon.GroupActor", new IMAGE_BRUSH( "Icons/ActorIcons/Group_16x", Icon16x16 ) );
-		Set( "ClassIcon.Landscape", new IMAGE_BRUSH( "Icons/ActorIcons/Landscape_16x", Icon16x16 ) );
-		Set( "ClassIcon.RadialForceActor", new IMAGE_BRUSH( "Icons/ActorIcons/RadialForce_16x", Icon16x16 ) );
-		Set( "ClassIcon.TargetPoint", new IMAGE_BRUSH( "Icons/ActorIcons/Target_16x", Icon16x16 ) );
+		Set( "ClassIcon.GroupActor", new IMAGE_BRUSH( "Icons/ActorIcons/GroupActor_16x", Icon16x16 ) );
 		Set( "ClassIcon.VectorFieldVolume", new IMAGE_BRUSH( "Icons/ActorIcons/VectorFieldVolume_16x", Icon16x16 ) );
 		Set( "ClassIcon.Deleted", new IMAGE_BRUSH( "Icons/ActorIcons/DeletedActor_16px", Icon16x16 ) );
 		Set( "ClassIcon.DestructibleActor", new IMAGE_BRUSH( "Icons/ActorIcons/Destructable_16px", Icon16x16 ) );
-		Set( "ClassIcon.SceneCapture", new IMAGE_BRUSH( "Icons/ActorIcons/SceneCapture_16px", Icon16x16 ) );
-		Set( "ClassIcon.TextRenderActor", new IMAGE_BRUSH( "Icons/ActorIcons/TextRenderer_16px", Icon16x16 ) );
-		Set( "ClassIcon.BoxReflectionCapture", new IMAGE_BRUSH( "Icons/ActorIcons/BoxReflectionCapture_16px", Icon16x16 ) );
-		Set( "ClassIcon.SphereReflectionCapture", new IMAGE_BRUSH( "Icons/ActorIcons/SphereReflectionCapture_16px", Icon16x16 ) );
 		Set( "ClassIcon.DefaultPawn", new IMAGE_BRUSH( "Icons/ActorIcons/DefaultPawn_16px", Icon16x16 ) );
 		Set( "ClassIcon.Pawn", new IMAGE_BRUSH( "Icons/ActorIcons/Pawn_16px", Icon16x16 ) );
 		Set( "ClassIcon.Note", new IMAGE_BRUSH( "Icons/ActorIcons/Note_16px", Icon16x16 ) );
 		Set( "ClassIcon.Character", new IMAGE_BRUSH( "Icons/ActorIcons/Character_16px", Icon16x16 ) );
 
+		Set("ClassIcon.BlueprintActorBase", new IMAGE_BRUSH("Icons/AssetIcons/Blueprint_16x", Icon16x16));
+		Set("ClassIcon.StaticMeshActor", new IMAGE_BRUSH("Icons/AssetIcons/StaticMesh_16x", Icon16x16));
+		Set("ClassIcon.SkeletalMeshActor", new IMAGE_BRUSH("Icons/AssetIcons/SkeletalMesh_16x", Icon16x16));
 
 		// Component classes
 		Set( "ClassIcon.AudioComponent", new IMAGE_BRUSH( "Icons/ActorIcons/SoundActor_16x", Icon16x16 ) );
-		Set( "ClassIcon.CameraComponent", new IMAGE_BRUSH( "Icons/ActorIcons/Camera_16x", Icon16x16 ) );
-		Set( "ClassIcon.BlueprintCore", new IMAGE_BRUSH( "Icons/ActorIcons/Blueprint_16x", Icon16x16 ) );
+		Set( "ClassIcon.CameraComponent", new IMAGE_BRUSH( "Icons/AssetIcons/CameraActor_16x", Icon16x16 ) );
+		Set( "ClassIcon.BlueprintCore", new IMAGE_BRUSH( "Icons/AssetIcons/Blueprint_16x", Icon16x16 ) );
 		Set( "ClassIcon.BrushComponent", new IMAGE_BRUSH( "Icons/ActorIcons/Brush_16x", Icon16x16 ) );
-		Set( "ClassIcon.DecalComponent", new IMAGE_BRUSH( "Icons/ActorIcons/Decal_16x", Icon16x16 ) );
-		Set( "ClassIcon.DirectionalLightComponent", new IMAGE_BRUSH( "Icons/ActorIcons/DirectionalLight_16x", Icon16x16 ) );
-		Set( "ClassIcon.ExponentialHeightFogComponent", new IMAGE_BRUSH( "Icons/ActorIcons/ExpoHeightFog_16x", Icon16x16 ) );
-		Set( "ClassIcon.LandscapeComponent", new IMAGE_BRUSH( "Icons/ActorIcons/Landscape_16x", Icon16x16 ) );
+		Set( "ClassIcon.DecalComponent", new IMAGE_BRUSH( "Icons/ActorIcons/DecalActor_16x", Icon16x16 ) );
+		Set( "ClassIcon.DirectionalLightComponent", new IMAGE_BRUSH( "Icons/AssetIcons/DirectionalLight_16x", Icon16x16 ) );
+		Set( "ClassIcon.ExponentialHeightFogComponent", new IMAGE_BRUSH( "Icons/AssetIcons/ExponentialHeightFog_16x", Icon16x16 ) );
+		Set( "ClassIcon.LandscapeComponent", new IMAGE_BRUSH( "Icons/AssetIcons/Landscape_16x", Icon16x16 ) );
 		Set( "ClassIcon.LightComponent", new IMAGE_BRUSH( "Icons/ActorIcons/LightActor_16x", Icon16x16 ) );
-		Set( "ClassIcon.ParticleSystemComponent", new IMAGE_BRUSH( "Icons/ActorIcons/Emitter_16x", Icon16x16 ) );
-		Set( "ClassIcon.PointLightComponent", new IMAGE_BRUSH( "Icons/ActorIcons/PointLight_16x", Icon16x16 ) );
-		Set( "ClassIcon.RB_RadialForceComponent", new IMAGE_BRUSH( "Icons/ActorIcons/RadialForce_16x", Icon16x16 ) );
-		Set( "ClassIcon.SingleAnimSkeletalComponent", new IMAGE_BRUSH( "Icons/ActorIcons/SkeletalMesh_16x", Icon16x16 ) );		
-		Set( "ClassIcon.SkeletalMeshComponent", new IMAGE_BRUSH( "Icons/ActorIcons/SkeletalMesh_16x", Icon16x16 ) );		
-		Set( "ClassIcon.SpotLightComponent", new IMAGE_BRUSH( "Icons/ActorIcons/SpotLight_16x", Icon16x16 ) );
-		Set( "ClassIcon.StaticMeshComponent", new IMAGE_BRUSH( "Icons/ActorIcons/StaticMesh_16x", Icon16x16 ) );
+		Set( "ClassIcon.ParticleSystemComponent", new IMAGE_BRUSH( "Icons/AssetIcons/ParticleSystem_16x", Icon16x16 ) );
+		Set( "ClassIcon.PointLightComponent", new IMAGE_BRUSH( "Icons/AssetIcons/PointLight_16x", Icon16x16 ) );
+		Set( "ClassIcon.RB_RadialForceComponent", new IMAGE_BRUSH( "Icons/AssetIcons/RadialForceActor_16x", Icon16x16 ) );
+		Set( "ClassIcon.SingleAnimSkeletalComponent", new IMAGE_BRUSH( "Icons/AssetIcons/SkeletalMesh_16x", Icon16x16 ) );		
+		Set( "ClassIcon.SkeletalMeshComponent", new IMAGE_BRUSH( "Icons/AssetIcons/SkeletalMesh_16x", Icon16x16 ) );		
+		Set( "ClassIcon.SpotLightComponent", new IMAGE_BRUSH( "Icons/AssetIcons/SpotLight_16x", Icon16x16 ) );
+		Set( "ClassIcon.StaticMeshComponent", new IMAGE_BRUSH( "Icons/AssetIcons/StaticMesh_16x", Icon16x16 ) );
 		Set( "ClassIcon.VectorFieldComponent", new IMAGE_BRUSH( "Icons/ActorIcons/VectorFieldVolume_16x", Icon16x16 ) );
 		Set( "ClassIcon.ArrowComponent", new IMAGE_BRUSH( "Icons/ActorIcons/Arrow_16px", Icon16x16 ) );
-		Set( "ClassIcon.AtmosphericFogComponent", new IMAGE_BRUSH( "Icons/ActorIcons/AtmosphericFog_16px", Icon16x16 ) );
-		Set( "ClassIcon.BoxComponent", new IMAGE_BRUSH( "Icons/ActorIcons/box_16px", Icon16x16 ) );
+		Set( "ClassIcon.AtmosphericFogComponent", new IMAGE_BRUSH( "Icons/AssetIcons/AtmosphericFog_16x", Icon16x16 ) );
+		Set( "ClassIcon.BoxComponent", new IMAGE_BRUSH( "Icons/ActorIcons/Box_16px", Icon16x16 ) );
 		Set( "ClassIcon.CapsuleComponent", new IMAGE_BRUSH( "Icons/ActorIcons/Capsule_16px", Icon16x16 ) );
 		Set( "ClassIcon.DestructibleComponent", new IMAGE_BRUSH( "Icons/ActorIcons/Destructable_16px", Icon16x16 ) );
 		Set( "ClassIcon.InstancedStaticMeshComponent", new IMAGE_BRUSH( "Icons/ActorIcons/InstancedStaticMesh_16px", Icon16x16 ) );
 		Set( "ClassIcon.MaterialBillboardComponent", new IMAGE_BRUSH( "Icons/ActorIcons/MaterialSprite_16px", Icon16x16 ) );
-		Set( "ClassIcon.SceneCaptureComponent2D", new IMAGE_BRUSH( "Icons/ActorIcons/SceneCapture2D_16px", Icon16x16 ) );
+		Set( "ClassIcon.SceneCaptureComponent2D", new IMAGE_BRUSH( "Icons/AssetIcons/SceneCapture2D_16x", Icon16x16 ) );
 		Set( "ClassIcon.SceneCaptureComponent", new IMAGE_BRUSH( "Icons/ActorIcons/SceneCapture_16px", Icon16x16 ) );
 		Set( "ClassIcon.SceneComponent", new IMAGE_BRUSH( "Icons/ActorIcons/Scene_16px", Icon16x16 ) );
 		Set( "ClassIcon.SphereComponent", new IMAGE_BRUSH( "Icons/ActorIcons/Sphere_16px", Icon16x16 ) );
 		Set( "ClassIcon.BillboardComponent", new IMAGE_BRUSH( "Icons/ActorIcons/SpriteComponent_16px", Icon16x16 ) );
-		Set( "ClassIcon.TextRenderComponent", new IMAGE_BRUSH( "Icons/ActorIcons/TextRenderer_16px", Icon16x16 ) );
+		Set( "ClassIcon.TextRenderComponent", new IMAGE_BRUSH( "Icons/AssetIcons/TextRenderActor_16x", Icon16x16 ) );
 		Set( "ClassIcon.TimelineComponent", new IMAGE_BRUSH("Icons/ActorIcons/TimelineComponent_16px", Icon16x16));
 		Set( "ClassIcon.ChildActorComponent", new IMAGE_BRUSH( "Icons/ActorIcons/ChildActorComponent_16px", Icon16x16 ) );
 		Set( "ClassIcon.ComponentMobilityStaticPip", new IMAGE_BRUSH( "Icons/ActorIcons/ComponentMobilityStationary_7x16px", Icon7x16, FLinearColor(0.f, 0.f, 0.f, 0.f) ) );
@@ -4099,6 +4082,7 @@ void FSlateEditorStyle::FStyle::SetupClassIconsAndThumbnails()
 			TEXT("ReverbVolume"),
 			TEXT("ReverbVolumeToggleable"),
 			TEXT("SceneCapture2D"),
+			TEXT("SceneCaptureCube"),
 			TEXT("SkyLight"),
 			TEXT("SkeletalMesh"),
 			TEXT("SlateBrushAsset"),
@@ -4128,7 +4112,6 @@ void FSlateEditorStyle::FStyle::SetupClassIconsAndThumbnails()
 			TEXT("UserDefinedStruct"),
 		};
 
-		const TCHAR* StyleString = TEXT("ClassIcon.%s%s");
 		const TCHAR* Path = TEXT("Icons/AssetIcons/%s_%dx");
 		
 		for (int32 TypeIndex = 0; TypeIndex < ARRAY_COUNT(AssetTypes); ++TypeIndex)
