Profiler - Added basic slate box brush for drawing thread-view samples

diff --git a/Engine/Source/Editor/EditorStyle/Private/SlateEditorStyle.cpp b/Engine/Source/Editor/EditorStyle/Private/SlateEditorStyle.cpp
--- a/Engine/Source/Editor/EditorStyle/Private/SlateEditorStyle.cpp
+++ b/Engine/Source/Editor/EditorStyle/Private/SlateEditorStyle.cpp
@@ -2572,6 +2572,9 @@ void FSlateEditorStyle::FStyle::SetupProfilerStyle()
 			.SetShadowOffset( FVector2D(0.0f, 0.0f) )
 			);
 
+		// Thread-view
+		Set( "Profiler.ThreadView.SampleBorder", new BOX_BRUSH( "Icons/Profiler/Profiler_ThreadView_SampleBorder_16x", FMargin( 1.0f / 16.0f ) ) );
+
 		// Event graph selected event border
 		Set( "Profiler.EventGraph.Border.TB", new BOX_BRUSH( "Icons/Profiler/Profiler_Border_TB_16x", FMargin(4.0f/16.0f) ) );
 		Set( "Profiler.EventGraph.Border.L", new BOX_BRUSH( "Icons/Profiler/Profiler_Border_L_16x",   FMargin(4.0f/16.0f) ) );
