Profiler - Better brush for drawing thread-view samples

diff --git a/Engine/Source/Editor/EditorStyle/Private/SlateEditorStyle.cpp b/Engine/Source/Editor/EditorStyle/Private/SlateEditorStyle.cpp
--- a/Engine/Source/Editor/EditorStyle/Private/SlateEditorStyle.cpp
+++ b/Engine/Source/Editor/EditorStyle/Private/SlateEditorStyle.cpp
@@ -2573,7 +2573,7 @@ void FSlateEditorStyle::FStyle::SetupProfilerStyle()
 			);
 
 		// Thread-view
-		Set( "Profiler.ThreadView.SampleBorder", new BOX_BRUSH( "Icons/Profiler/Profiler_ThreadView_SampleBorder_16x", FMargin( 1.0f / 16.0f ) ) );
+		Set( "Profiler.ThreadView.SampleBorder", new BOX_BRUSH( "Icons/Profiler/Profiler_ThreadView_SampleBorder_16x", FMargin( 2.0f / 16.0f ) ) );
 
 		// Event graph selected event border
 		Set( "Profiler.EventGraph.Border.TB", new BOX_BRUSH( "Icons/Profiler/Profiler_Border_TB_16x", FMargin(4.0f/16.0f) ) );
