SessionBrowser: fixed text styling of the 'Terminate Session' button

diff --git a/Engine/Source/Developer/SessionFrontend/Private/Widgets/Browser/SSessionBrowser.cpp b/Engine/Source/Developer/SessionFrontend/Private/Widgets/Browser/SSessionBrowser.cpp
--- a/Engine/Source/Developer/SessionFrontend/Private/Widgets/Browser/SSessionBrowser.cpp
+++ b/Engine/Source/Developer/SessionFrontend/Private/Widgets/Browser/SSessionBrowser.cpp
@@ -100,7 +100,7 @@ void SSessionBrowser::Construct( const FArguments& InArgs, ISessionManagerRef In
 									.Padding(4.0f, 1.0f, 0.0f, 0.0f)
 									[
 										SNew(STextBlock)
-											.TextStyle(FEditorStyle::Get(), "ContentBrowser.TopBar.Font")
+											.TextStyle(FEditorStyle::Get(), "SessionBrowser.Terminate.Font")
 											.Text(LOCTEXT("TerminateSessionButtonLabel", "Terminate Session").ToString())
 									]
 							]
diff --git a/Engine/Source/Editor/EditorStyle/Private/SlateEditorStyle.cpp b/Engine/Source/Editor/EditorStyle/Private/SlateEditorStyle.cpp
--- a/Engine/Source/Editor/EditorStyle/Private/SlateEditorStyle.cpp
+++ b/Engine/Source/Editor/EditorStyle/Private/SlateEditorStyle.cpp
@@ -868,6 +868,13 @@ void FSlateEditorStyle::FStyle::SetupGeneralStyles()
 	// Session Browser tab
 	{
 		Set("SessionBrowser.Terminate", new IMAGE_BRUSH( "Icons/icon_DevicePowerOff_40x", Icon20x20 ) );
+
+		Set("SessionBrowser.Terminate.Font", FTextBlockStyle(NormalText)
+			.SetFont(TTF_CORE_FONT( "Fonts/Roboto-Bold", 12))
+			.SetColorAndOpacity(FLinearColor(1.0f, 1.0f, 1.0f))
+			.SetHighlightColor(FLinearColor(1.0f, 1.0f, 1.0f))
+			.SetShadowOffset(FVector2D(1.0f, 1.0f))
+			.SetShadowColorAndOpacity(FLinearColor(0.0f, 0.0f, 0.0f, 0.9f)));
 	}
 
 	// Session Console tab
