Reinstated GroupActor scene outliner icon.
#ttp 335452 - Regression: EDITOR: Scene Outliner: Group actor icon missing and appears as a white square
#branch UE4
#proj Editor.EditorStyle
#change Corrected decal actor icon path.
#add Re-added Engine/Content/Editor/Slate/Icons/ActorIcons/GroupActor_16x.png.

diff --git a/Engine/Source/Editor/EditorStyle/Private/SlateEditorStyle.cpp b/Engine/Source/Editor/EditorStyle/Private/SlateEditorStyle.cpp
--- a/Engine/Source/Editor/EditorStyle/Private/SlateEditorStyle.cpp
+++ b/Engine/Source/Editor/EditorStyle/Private/SlateEditorStyle.cpp
@@ -3985,7 +3985,7 @@ void FSlateEditorStyle::FStyle::SetupClassIconsAndThumbnails()
 		Set( "ClassIcon.CameraComponent", new IMAGE_BRUSH( "Icons/AssetIcons/CameraActor_16x", Icon16x16 ) );
 		Set( "ClassIcon.BlueprintCore", new IMAGE_BRUSH( "Icons/AssetIcons/Blueprint_16x", Icon16x16 ) );
 		Set( "ClassIcon.BrushComponent", new IMAGE_BRUSH( "Icons/ActorIcons/Brush_16x", Icon16x16 ) );
-		Set( "ClassIcon.DecalComponent", new IMAGE_BRUSH( "Icons/ActorIcons/DecalActor_16x", Icon16x16 ) );
+		Set( "ClassIcon.DecalComponent", new IMAGE_BRUSH( "Icons/AssetIcons/DecalActor_16x", Icon16x16 ) );
 		Set( "ClassIcon.DirectionalLightComponent", new IMAGE_BRUSH( "Icons/AssetIcons/DirectionalLight_16x", Icon16x16 ) );
 		Set( "ClassIcon.ExponentialHeightFogComponent", new IMAGE_BRUSH( "Icons/AssetIcons/ExponentialHeightFog_16x", Icon16x16 ) );
 		Set( "ClassIcon.LandscapeComponent", new IMAGE_BRUSH( "Icons/AssetIcons/Landscape_16x", Icon16x16 ) );
