Scene Outliner inactive selection state is now derived from the main selection state, rather than being hard-coded

This fixes TTP#335121 - EDITOR: Scene Outliner: Inactive selection color does not conform to appearence preferences

diff --git a/Engine/Source/Editor/EditorStyle/Private/SlateEditorStyle.cpp b/Engine/Source/Editor/EditorStyle/Private/SlateEditorStyle.cpp
--- a/Engine/Source/Editor/EditorStyle/Private/SlateEditorStyle.cpp
+++ b/Engine/Source/Editor/EditorStyle/Private/SlateEditorStyle.cpp
@@ -56,6 +56,7 @@ FSlateEditorStyle::FStyle::FStyle( const TWeakObjectPtr< UEditorStyleSettings >&
 	, InvertedForeground_LinearRef( MakeShareable( new FLinearColor( 0, 0, 0 ) ) )
 	, SelectorColor_LinearRef( MakeShareable( new FLinearColor( 0.701f, 0.225f, 0.003f ) ) )
 	, SelectionColor_LinearRef( MakeShareable( new FLinearColor( 0.728f, 0.364f, 0.003f ) ) )
+	, SelectionColor_Subdued_LinearRef( MakeShareable( new FLinearColor( 0.807f, 0.596f, 0.388f ) ) )
 	, SelectionColor_Inactive_LinearRef( MakeShareable( new FLinearColor( 0.25f, 0.25f, 0.25f ) ) )
 	, SelectionColor_Pressed_LinearRef( MakeShareable( new FLinearColor( 0.701f, 0.225f, 0.003f ) ) )
 
@@ -64,6 +65,7 @@ FSlateEditorStyle::FStyle::FStyle( const TWeakObjectPtr< UEditorStyleSettings >&
 	, InvertedForeground( InvertedForeground_LinearRef )
 	, SelectorColor( SelectorColor_LinearRef )
 	, SelectionColor( SelectionColor_LinearRef )
+	, SelectionColor_Subdued( SelectionColor_Subdued_LinearRef )
 	, SelectionColor_Inactive( SelectionColor_Inactive_LinearRef )
 	, SelectionColor_Pressed( SelectionColor_Pressed_LinearRef )
 
@@ -98,6 +100,12 @@ void FSlateEditorStyle::FStyle::SyncSettings()
 		SetColor( SelectionColor_Inactive_LinearRef, Settings->InactiveSelectionColor );
 		SetColor( SelectionColor_Pressed_LinearRef, Settings->PressedSelectionColor );
 
+		// The subdued selection color is derived from the selection color
+		auto SubduedSelectionColor = Settings->SelectionColor.LinearRGBToHSV();
+		SubduedSelectionColor.G *= 0.55f;		// take the saturation 
+		SubduedSelectionColor.B *= 0.8f;		// and brightness down
+		SetColor( SelectionColor_Subdued_LinearRef, SubduedSelectionColor.HSVToLinearRGB() );
+
 		// Also sync the colors used by FCoreStyle, as FEditorStyle isn't yet being used as an override everywhere
 		FCoreStyle::SetSelectorColor( Settings->KeyboardFocusColor );
 		FCoreStyle::SetSelectionColor( Settings->SelectionColor );
@@ -1073,10 +1081,11 @@ void FSlateEditorStyle::FStyle::SetupGeneralStyles()
 		Set( "SceneOutliner.NewFolderIcon", new IMAGE_BRUSH("Icons/icon_AddFolder_16x", Icon16x16 ) );
 		Set( "SceneOutliner.MoveToRoot", new IMAGE_BRUSH("Icons/icon_NoFolder_16x", Icon16x16 ) );
 		Set( "SceneOutliner.ChangedItemHighlight", new BOX_BRUSH( "Common/EditableTextSelectionBackground", FMargin(4.f/16.f) ) );
+
 		// Selection color should still be orange to align with the editor viewport.
 		// But must also give the hint that the tree is no longer focused.
 		Set( "SceneOutliner.TableViewRow", FTableRowStyle(NormalTableRowStyle)
-			.SetInactiveBrush( IMAGE_BRUSH( "Common/Selection", Icon8x8, FLinearColor(FColor(188,161,97)) ) )
+			.SetInactiveBrush( IMAGE_BRUSH( "Common/Selection", Icon8x8, SelectionColor_Subdued ) )
 		);
 
 		Set( "SceneOutliner.EditBlueprintHyperlinkStyle",
diff --git a/Engine/Source/Editor/EditorStyle/Private/SlateEditorStyle.h b/Engine/Source/Editor/EditorStyle/Private/SlateEditorStyle.h
--- a/Engine/Source/Editor/EditorStyle/Private/SlateEditorStyle.h
+++ b/Engine/Source/Editor/EditorStyle/Private/SlateEditorStyle.h
@@ -108,6 +108,7 @@ public:
 		const TSharedRef< FLinearColor > InvertedForeground_LinearRef;
 		const TSharedRef< FLinearColor > SelectorColor_LinearRef;
 		const TSharedRef< FLinearColor > SelectionColor_LinearRef;
+		const TSharedRef< FLinearColor > SelectionColor_Subdued_LinearRef;
 		const TSharedRef< FLinearColor > SelectionColor_Inactive_LinearRef;
 		const TSharedRef< FLinearColor > SelectionColor_Pressed_LinearRef;
 
@@ -116,6 +117,7 @@ public:
 		const FSlateColor InvertedForeground;
 		const FSlateColor SelectorColor;
 		const FSlateColor SelectionColor;
+		const FSlateColor SelectionColor_Subdued;
 		const FSlateColor SelectionColor_Inactive;
 		const FSlateColor SelectionColor_Pressed;
 
