Add class icon for SplineComponent

diff --git a/Engine/Source/Editor/EditorStyle/Private/SlateEditorStyle.cpp b/Engine/Source/Editor/EditorStyle/Private/SlateEditorStyle.cpp
--- a/Engine/Source/Editor/EditorStyle/Private/SlateEditorStyle.cpp
+++ b/Engine/Source/Editor/EditorStyle/Private/SlateEditorStyle.cpp
@@ -4035,7 +4035,8 @@ void FSlateEditorStyle::FStyle::SetupClassIconsAndThumbnails()
 		Set( "ClassIcon.SceneCaptureComponent2D", new IMAGE_BRUSH( "Icons/AssetIcons/SceneCapture2D_16x", Icon16x16 ) );
 		Set( "ClassIcon.SceneCaptureComponent", new IMAGE_BRUSH( "Icons/ActorIcons/SceneCapture_16px", Icon16x16 ) );
 		Set( "ClassIcon.SceneComponent", new IMAGE_BRUSH( "Icons/ActorIcons/Scene_16px", Icon16x16 ) );
-		Set( "ClassIcon.SphereComponent", new IMAGE_BRUSH( "Icons/ActorIcons/Sphere_16px", Icon16x16 ) );
+		Set( "ClassIcon.SphereComponent", new IMAGE_BRUSH("Icons/ActorIcons/Sphere_16px", Icon16x16));
+		Set( "ClassIcon.SplineComponent", new IMAGE_BRUSH("Icons/ActorIcons/Spline_16px", Icon16x16));
 		Set( "ClassIcon.BillboardComponent", new IMAGE_BRUSH( "Icons/ActorIcons/SpriteComponent_16px", Icon16x16 ) );
 		Set( "ClassIcon.TextRenderComponent", new IMAGE_BRUSH( "Icons/AssetIcons/TextRenderActor_16x", Icon16x16 ) );
 		Set( "ClassIcon.TimelineComponent", new IMAGE_BRUSH("Icons/ActorIcons/TimelineComponent_16px", Icon16x16));
