UE4 - Fix xb1 compiling

diff --git a/Engine/Source/Editor/EditorStyle/EditorStyle.Build.cs b/Engine/Source/Editor/EditorStyle/EditorStyle.Build.cs
--- a/Engine/Source/Editor/EditorStyle/EditorStyle.Build.cs
+++ b/Engine/Source/Editor/EditorStyle/EditorStyle.Build.cs
@@ -17,7 +17,7 @@ public class EditorStyle : ModuleRules
 		PrivateDependencyModuleNames.AddRange(
 			new string[] {
 				"SlateCore",
-				"DesktopPlatform",
+				"TargetPlatform",
 			}
 		);
 
diff --git a/Engine/Source/Editor/EditorStyle/Private/EditorStylePrivatePCH.h b/Engine/Source/Editor/EditorStyle/Private/EditorStylePrivatePCH.h
--- a/Engine/Source/Editor/EditorStyle/Private/EditorStylePrivatePCH.h
+++ b/Engine/Source/Editor/EditorStyle/Private/EditorStylePrivatePCH.h
@@ -17,7 +17,7 @@
 #include "Settings.h"
 #include "Slate.h"
 #include "SlateStyle.h"
-#include "PlatformInfo.h"
+#include "TargetPlatform.h"
 
 
 /* Private includes
