EditorStyle now only uses DesktopPlatform when building an Editor or Program target

diff --git a/Engine/Source/Developer/Android/Android_ATCTargetPlatform/Android_ATCTargetPlatform.Build.cs b/Engine/Source/Developer/Android/Android_ATCTargetPlatform/Android_ATCTargetPlatform.Build.cs
--- a/Engine/Source/Developer/Android/Android_ATCTargetPlatform/Android_ATCTargetPlatform.Build.cs
+++ b/Engine/Source/Developer/Android/Android_ATCTargetPlatform/Android_ATCTargetPlatform.Build.cs
@@ -13,6 +13,7 @@ public class Android_ATCTargetPlatform : ModuleRules
 			{
 				"Core",
 				"TargetPlatform",
+				"DesktopPlatform",
 				"AndroidDeviceDetection",
 			}
 		);
diff --git a/Engine/Source/Developer/Android/Android_DXTTargetPlatform/Android_DXTTargetPlatform.Build.cs b/Engine/Source/Developer/Android/Android_DXTTargetPlatform/Android_DXTTargetPlatform.Build.cs
--- a/Engine/Source/Developer/Android/Android_DXTTargetPlatform/Android_DXTTargetPlatform.Build.cs
+++ b/Engine/Source/Developer/Android/Android_DXTTargetPlatform/Android_DXTTargetPlatform.Build.cs
@@ -13,6 +13,7 @@ public class Android_DXTTargetPlatform : ModuleRules
 			{
 				"Core",
 				"TargetPlatform",
+				"DesktopPlatform",
 				"AndroidDeviceDetection",
 			}
 		);
diff --git a/Engine/Source/Developer/Android/Android_ETC1TargetPlatform/Android_ETC1TargetPlatform.Build.cs b/Engine/Source/Developer/Android/Android_ETC1TargetPlatform/Android_ETC1TargetPlatform.Build.cs
--- a/Engine/Source/Developer/Android/Android_ETC1TargetPlatform/Android_ETC1TargetPlatform.Build.cs
+++ b/Engine/Source/Developer/Android/Android_ETC1TargetPlatform/Android_ETC1TargetPlatform.Build.cs
@@ -20,6 +20,7 @@ public class Android_ETC1TargetPlatform : ModuleRules
 			{
 				"Core",
 				"TargetPlatform",
+				"DesktopPlatform",
 				"AndroidDeviceDetection",
 			}
 		);
diff --git a/Engine/Source/Developer/Android/Android_ETC2TargetPlatform/Android_ETC2TargetPlatform.Build.cs b/Engine/Source/Developer/Android/Android_ETC2TargetPlatform/Android_ETC2TargetPlatform.Build.cs
--- a/Engine/Source/Developer/Android/Android_ETC2TargetPlatform/Android_ETC2TargetPlatform.Build.cs
+++ b/Engine/Source/Developer/Android/Android_ETC2TargetPlatform/Android_ETC2TargetPlatform.Build.cs
@@ -13,6 +13,7 @@ public class Android_ETC2TargetPlatform : ModuleRules
 			{
 				"Core",
 				"TargetPlatform",
+				"DesktopPlatform",
 				"AndroidDeviceDetection",
 			}
 		);
diff --git a/Engine/Source/Developer/Android/Android_PVRTCTargetPlatform/Android_PVRTCTargetPlatform.Build.cs b/Engine/Source/Developer/Android/Android_PVRTCTargetPlatform/Android_PVRTCTargetPlatform.Build.cs
--- a/Engine/Source/Developer/Android/Android_PVRTCTargetPlatform/Android_PVRTCTargetPlatform.Build.cs
+++ b/Engine/Source/Developer/Android/Android_PVRTCTargetPlatform/Android_PVRTCTargetPlatform.Build.cs
@@ -13,6 +13,7 @@ public class Android_PVRTCTargetPlatform : ModuleRules
 			{
 				"Core",
 				"TargetPlatform",
+				"DesktopPlatform",
 				"AndroidDeviceDetection",
 			}
 		);
diff --git a/Engine/Source/Editor/EditorStyle/EditorStyle.Build.cs b/Engine/Source/Editor/EditorStyle/EditorStyle.Build.cs
--- a/Engine/Source/Editor/EditorStyle/EditorStyle.Build.cs
+++ b/Engine/Source/Editor/EditorStyle/EditorStyle.Build.cs
@@ -17,7 +17,6 @@ public class EditorStyle : ModuleRules
 		PrivateDependencyModuleNames.AddRange(
 			new string[] {
 				"SlateCore",
-				"TargetPlatform",
 			}
 		);
 
@@ -27,6 +26,20 @@ public class EditorStyle : ModuleRules
 			}
 		);
 
+		// DesktopPlatform is only available for Editor and Program targets (running on a desktop platform)
+		bool IsDesktopPlatformType = Target.Platform == UnrealBuildTool.UnrealTargetPlatform.Win32
+			|| Target.Platform == UnrealBuildTool.UnrealTargetPlatform.Win64
+			|| Target.Platform == UnrealBuildTool.UnrealTargetPlatform.Mac
+			|| Target.Platform == UnrealBuildTool.UnrealTargetPlatform.Linux;
+		if (Target.Type == TargetRules.TargetType.Editor || (Target.Type == TargetRules.TargetType.Program && IsDesktopPlatformType))
+		{
+			PrivateDependencyModuleNames.AddRange(
+				new string[] {
+					"DesktopPlatform",
+				}
+			);
+		}
+
 		OptimizeCode = CodeOptimization.Never;
 	}
 }
diff --git a/Engine/Source/Editor/EditorStyle/Private/EditorStylePrivatePCH.h b/Engine/Source/Editor/EditorStyle/Private/EditorStylePrivatePCH.h
--- a/Engine/Source/Editor/EditorStyle/Private/EditorStylePrivatePCH.h
+++ b/Engine/Source/Editor/EditorStyle/Private/EditorStylePrivatePCH.h
@@ -17,7 +17,6 @@
 #include "Settings.h"
 #include "Slate.h"
 #include "SlateStyle.h"
-#include "TargetPlatform.h"
 
 
 /* Private includes
diff --git a/Engine/Source/Editor/EditorStyle/Private/SlateEditorStyle.cpp b/Engine/Source/Editor/EditorStyle/Private/SlateEditorStyle.cpp
--- a/Engine/Source/Editor/EditorStyle/Private/SlateEditorStyle.cpp
+++ b/Engine/Source/Editor/EditorStyle/Private/SlateEditorStyle.cpp
@@ -6,6 +6,9 @@
 
 #include "EditorStylePrivatePCH.h"
 
+#if (WITH_EDITOR || (IS_PROGRAM && PLATFORM_DESKTOP))
+	#include "PlatformInfo.h"
+#endif
 
 #define IMAGE_BRUSH( RelativePath, ... ) FSlateImageBrush( RootToContentDir( RelativePath, TEXT(".png") ), __VA_ARGS__ )
 #define BOX_BRUSH( RelativePath, ... ) FSlateBoxBrush( RootToContentDir( RelativePath, TEXT(".png") ), __VA_ARGS__ )
