World Browser: added dummy icons

diff --git a/Engine/Source/Editor/EditorStyle/Private/SlateEditorStyle.cpp b/Engine/Source/Editor/EditorStyle/Private/SlateEditorStyle.cpp
--- a/Engine/Source/Editor/EditorStyle/Private/SlateEditorStyle.cpp
+++ b/Engine/Source/Editor/EditorStyle/Private/SlateEditorStyle.cpp
@@ -1013,7 +1013,6 @@ void FSlateEditorStyle::FStyle::SetupGeneralStyles()
 
 	// Levels General
 	{
-		Set( "Level.Icon16x", new IMAGE_BRUSH( "Icons/levels_16x", Icon16x16 ) );
 		Set( "Level.VisibleIcon16x", new IMAGE_BRUSH( "Icons/icon_levels_visible_16px", Icon16x16 ) );
 		Set( "Level.VisibleHighlightIcon16x", new IMAGE_BRUSH( "Icons/icon_levels_visible_hi_16px", Icon16x16 ) );
 		Set( "Level.NotVisibleIcon16x", new IMAGE_BRUSH( "Icons/icon_levels_invisible_16px", Icon16x16 ) );
@@ -1051,22 +1050,22 @@ void FSlateEditorStyle::FStyle::SetupGeneralStyles()
 
 	// World Browser
 	{
-		Set( "WorldBrowser.AddLayer", new IMAGE_BRUSH( "Icons/icon_add_40x", Icon16x16 ) );
-		Set( "WorldBrowser.SimulationViewPositon", new IMAGE_BRUSH( "Icons/icon_variable_used_16x", Icon16x16 ) );
-		Set( "WorldBrowser.MouseLocation", new IMAGE_BRUSH( "Icons/cursor_cardinal_cross", Icon16x16 ) );
-		Set( "WorldBrowser.MarqueeRectSize", new IMAGE_BRUSH( "Icons/icon_ViewMode_VisualisationGBuffer_16px", Icon16x16 ) );
-		Set( "WorldBrowser.WorldSize", new IMAGE_BRUSH( "Icons/icon_world_16x", Icon16x16 ) );
-		Set( "WorldBrowser.WorldOrigin", new IMAGE_BRUSH( "Icons/icon_axis_16px", Icon16x16 ) );
-		Set( "WorldBrowser.WorldViewScale", new IMAGE_BRUSH( "Icons/ExploreLayerContents", Icon16x16 ) );
+		Set( "WorldBrowser.AddLayer", new IMAGE_BRUSH( "Icons/icon_levels_addlayer_16x", Icon16x16 ) );
+		Set( "WorldBrowser.SimulationViewPositon", new IMAGE_BRUSH( "Icons/icon_levels_simulationviewpos_16x", Icon16x16 ) );
+		Set( "WorldBrowser.MouseLocation", new IMAGE_BRUSH( "Icons/icon_levels_mouselocation_16x", Icon16x16 ) );
+		Set( "WorldBrowser.MarqueeRectSize", new IMAGE_BRUSH( "Icons/icon_levels_marqueerectsize_16x", Icon16x16 ) );
+		Set( "WorldBrowser.WorldSize", new IMAGE_BRUSH( "Icons/icon_levels_worldsize_16x", Icon16x16 ) );
+		Set( "WorldBrowser.WorldOrigin", new IMAGE_BRUSH( "Icons/icon_levels_worldorigin_16x", Icon16x16 ) );
 		Set( "WorldBrowser.DirectionXPositive", new IMAGE_BRUSH( "Icons/icon_PanRight", Icon16x16 ) );
 		Set( "WorldBrowser.DirectionXNegative", new IMAGE_BRUSH( "Icons/icon_PanLeft", Icon16x16 ) );
 		Set( "WorldBrowser.DirectionYPositive", new IMAGE_BRUSH( "Icons/icon_PanUp", Icon16x16 ) );
 		Set( "WorldBrowser.DirectionYNegative", new IMAGE_BRUSH( "Icons/icon_PanDown", Icon16x16 ) );
 		Set( "WorldBrowser.LevelStreamingAlwaysLoaded", new FSlateNoResource() );
 		Set( "WorldBrowser.LevelStreamingKismet", new IMAGE_BRUSH( "Icons/ActorIcons/ComponentMobilityStationary_7x16px", Icon7x16 ) );
-		Set( "WorldBrowser.HierarchyButtonBrush", new IMAGE_BRUSH( "Icons/levels_16x", Icon16x16 ) );
-		Set( "WorldBrowser.DetailsButtonBrush", new IMAGE_BRUSH( "NewsFeed/SettingsButton", Icon16x16 ) );
-		Set( "WorldBrowser.CompositionButtonBrush", new IMAGE_BRUSH( "Icons/icon_volume_LevelGrid_16x", Icon16x16 ) );
+		Set( "WorldBrowser.LevelsMenuBrush", new IMAGE_BRUSH( "Icons/icon_levels_levelsmenu_40x", Icon25x25 ) );
+		Set( "WorldBrowser.HierarchyButtonBrush", new IMAGE_BRUSH( "Icons/icon_levels_hierarchybutton_16x", Icon16x16 ) );
+		Set( "WorldBrowser.DetailsButtonBrush", new IMAGE_BRUSH( "Icons/icon_levels_detailsbutton_16x", Icon16x16 ) );
+		Set( "WorldBrowser.CompositionButtonBrush", new IMAGE_BRUSH( "Icons/icon_levels_compositionbutton_16x", Icon16x16 ) );
 		
 		Set( "WorldBrowser.StatusBarText", FTextBlockStyle(NormalText)
 				.SetFont( TTF_FONT( "Fonts/Roboto-BoldCondensed", 12 ) )
diff --git a/Engine/Source/Editor/WorldBrowser/Private/SWorldHierarchy.cpp b/Engine/Source/Editor/WorldBrowser/Private/SWorldHierarchy.cpp
--- a/Engine/Source/Editor/WorldBrowser/Private/SWorldHierarchy.cpp
+++ b/Engine/Source/Editor/WorldBrowser/Private/SWorldHierarchy.cpp
@@ -566,7 +566,7 @@ void SWorldHierarchy::OnBrowseWorld(UWorld* InWorld)
 							.VAlign(VAlign_Center)
 							[
 								SNew(SImage)
-								.Image(FEditorStyle::GetBrush("ContentBrowser.NewAsset"))
+								.Image(this, &SWorldHierarchy::GetLevelsMenuBrush)
 							]
 
 							// Text
@@ -645,6 +645,11 @@ FReply SWorldHierarchy::OnSummonDetails()
 	return FReply::Handled();
 }
 
+const FSlateBrush* SWorldHierarchy::GetLevelsMenuBrush() const
+{
+	return FEditorStyle::GetBrush("WorldBrowser.LevelsMenuBrush");
+}
+
 const FSlateBrush* SWorldHierarchy::GetSummonDetailsBrush() const
 {
 	return FEditorStyle::GetBrush("WorldBrowser.DetailsButtonBrush");
diff --git a/Engine/Source/Editor/WorldBrowser/Private/SWorldHierarchy.h b/Engine/Source/Editor/WorldBrowser/Private/SWorldHierarchy.h
--- a/Engine/Source/Editor/WorldBrowser/Private/SWorldHierarchy.h
+++ b/Engine/Source/Editor/WorldBrowser/Private/SWorldHierarchy.h
@@ -21,6 +21,9 @@ class SWorldHierarchy
 private:
 	void OnBrowseWorld(UWorld* InWorld);
 
+	/**  */
+	const FSlateBrush* GetLevelsMenuBrush() const;
+
 	/**  */
 	FReply OnSummonDetails();
 	const FSlateBrush* GetSummonDetailsBrush() const;
