Hooked up new Behavior Tree icons.

Added editor-only functions to allow nodes to override their in-editor icon display.
Added Blueprint icon to show which nodes use Blueprint for their logic.
Hooked up new icons to Editor style.
Added missing Blackboard icon.

TTP# 337682 - [ART REQUEST] New Icons for behaviour trees

diff --git a/Engine/Source/Editor/BehaviorTreeEditor/Classes/BehaviorTreeGraphNode.h b/Engine/Source/Editor/BehaviorTreeEditor/Classes/BehaviorTreeGraphNode.h
--- a/Engine/Source/Editor/BehaviorTreeEditor/Classes/BehaviorTreeGraphNode.h
+++ b/Engine/Source/Editor/BehaviorTreeEditor/Classes/BehaviorTreeGraphNode.h
@@ -80,6 +80,9 @@ class UBehaviorTreeGraphNode : public UEdGraphNode
 	/** gets icon resource name for title bar */
 	virtual FName GetNameIcon() const;
 
+	/** Check if node instance uses blueprint for its implementation */
+	bool UsesBlueprint() const;
+
 	/** check if node has any errors, used for assigning colors on graph */
 	virtual bool HasErrors() const;
 
diff --git a/Engine/Source/Editor/BehaviorTreeEditor/Classes/BehaviorTreeGraphNode_Composite.h b/Engine/Source/Editor/BehaviorTreeEditor/Classes/BehaviorTreeGraphNode_Composite.h
--- a/Engine/Source/Editor/BehaviorTreeEditor/Classes/BehaviorTreeGraphNode_Composite.h
+++ b/Engine/Source/Editor/BehaviorTreeEditor/Classes/BehaviorTreeGraphNode_Composite.h
@@ -11,9 +11,6 @@ class UBehaviorTreeGraphNode_Composite : public UBehaviorTreeGraphNode
 	virtual FText GetNodeTitle(ENodeTitleType::Type TitleType) const OVERRIDE;
 	virtual bool RefreshNodeClass() OVERRIDE{ return false; }
 
-	/** gets icon resource name for title bar */
-	virtual FName GetNameIcon() const OVERRIDE;
-
 	/** Gets a list of actions that can be done to this particular node */
 	virtual void GetContextMenuActions(const FGraphNodeContextMenuBuilder& Context) const OVERRIDE;
 
diff --git a/Engine/Source/Editor/BehaviorTreeEditor/Classes/BehaviorTreeGraphNode_CompositeDecorator.h b/Engine/Source/Editor/BehaviorTreeEditor/Classes/BehaviorTreeGraphNode_CompositeDecorator.h
--- a/Engine/Source/Editor/BehaviorTreeEditor/Classes/BehaviorTreeGraphNode_CompositeDecorator.h
+++ b/Engine/Source/Editor/BehaviorTreeEditor/Classes/BehaviorTreeGraphNode_CompositeDecorator.h
@@ -25,8 +25,6 @@ class UBehaviorTreeGraphNode_CompositeDecorator : public UBehaviorTreeGraphNode
 
 	FString GetNodeTypeDescription() const;
 	virtual FText GetNodeTitle(ENodeTitleType::Type TitleType) const OVERRIDE;
-	/** gets icon resource name for title bar */
-	virtual FName GetNameIcon() const OVERRIDE;
 	virtual void AllocateDefaultPins() OVERRIDE;
 	virtual FString GetDescription() const OVERRIDE;
 	virtual void PostPlacedNewNode() OVERRIDE;
diff --git a/Engine/Source/Editor/BehaviorTreeEditor/Classes/BehaviorTreeGraphNode_Decorator.h b/Engine/Source/Editor/BehaviorTreeEditor/Classes/BehaviorTreeGraphNode_Decorator.h
--- a/Engine/Source/Editor/BehaviorTreeEditor/Classes/BehaviorTreeGraphNode_Decorator.h
+++ b/Engine/Source/Editor/BehaviorTreeEditor/Classes/BehaviorTreeGraphNode_Decorator.h
@@ -9,8 +9,6 @@ class UBehaviorTreeGraphNode_Decorator : public UBehaviorTreeGraphNode
 	GENERATED_UCLASS_BODY()
 
 	virtual FText GetNodeTitle(ENodeTitleType::Type TitleType) const OVERRIDE;
-	/** gets icon resource name for title bar */
-	virtual FName GetNameIcon() const OVERRIDE;
 	virtual void AllocateDefaultPins() OVERRIDE;
 	void CollectDecoratorData(TArray<class UBTDecorator*>& NodeInstances, TArray<struct FBTDecoratorLogic>& Operations) const;
 
diff --git a/Engine/Source/Editor/BehaviorTreeEditor/Classes/BehaviorTreeGraphNode_Service.h b/Engine/Source/Editor/BehaviorTreeEditor/Classes/BehaviorTreeGraphNode_Service.h
--- a/Engine/Source/Editor/BehaviorTreeEditor/Classes/BehaviorTreeGraphNode_Service.h
+++ b/Engine/Source/Editor/BehaviorTreeEditor/Classes/BehaviorTreeGraphNode_Service.h
@@ -9,8 +9,6 @@ class UBehaviorTreeGraphNode_Service : public UBehaviorTreeGraphNode
 	GENERATED_UCLASS_BODY()
 
 	virtual FText GetNodeTitle(ENodeTitleType::Type TitleType) const OVERRIDE;
-	/** gets icon resource name for title bar */
-	virtual FName GetNameIcon() const OVERRIDE;
 	virtual void AllocateDefaultPins() OVERRIDE;
 	virtual bool IsSubNode() const OVERRIDE;
 };
diff --git a/Engine/Source/Editor/BehaviorTreeEditor/Classes/BehaviorTreeGraphNode_Task.h b/Engine/Source/Editor/BehaviorTreeEditor/Classes/BehaviorTreeGraphNode_Task.h
--- a/Engine/Source/Editor/BehaviorTreeEditor/Classes/BehaviorTreeGraphNode_Task.h
+++ b/Engine/Source/Editor/BehaviorTreeEditor/Classes/BehaviorTreeGraphNode_Task.h
@@ -11,8 +11,6 @@ class UBehaviorTreeGraphNode_Task : public UBehaviorTreeGraphNode
 
 	virtual void AllocateDefaultPins() OVERRIDE;
 	virtual FText GetNodeTitle(ENodeTitleType::Type TitleType) const OVERRIDE;
-	/** gets icon resource name for title bar */
-	virtual FName GetNameIcon() const OVERRIDE;
 	/** Gets a list of actions that can be done to this particular node */
 	virtual void GetContextMenuActions(const FGraphNodeContextMenuBuilder& Context) const OVERRIDE;
 
diff --git a/Engine/Source/Editor/BehaviorTreeEditor/Private/BehaviorTreeGraphNode.cpp b/Engine/Source/Editor/BehaviorTreeEditor/Private/BehaviorTreeGraphNode.cpp
--- a/Engine/Source/Editor/BehaviorTreeEditor/Private/BehaviorTreeGraphNode.cpp
+++ b/Engine/Source/Editor/BehaviorTreeEditor/Private/BehaviorTreeGraphNode.cpp
@@ -393,7 +393,14 @@ void UBehaviorTreeGraphNode::ClearDebuggerState()
 
 FName UBehaviorTreeGraphNode::GetNameIcon() const
 {
-	return FName("BTEditor.Graph.BTNode.Icon");
+	UBTNode* BTNodeInstance = Cast<UBTNode>(NodeInstance);
+	return BTNodeInstance != nullptr ? BTNodeInstance->GetNodeIconName() : FName("BTEditor.Graph.BTNode.Icon");
+}
+
+bool UBehaviorTreeGraphNode::UsesBlueprint() const
+{
+	UBTNode* BTNodeInstance = Cast<UBTNode>(NodeInstance);
+	return BTNodeInstance != nullptr ? BTNodeInstance->UsesBlueprint() : false;
 }
 
 bool UBehaviorTreeGraphNode::RefreshNodeClass()
diff --git a/Engine/Source/Editor/BehaviorTreeEditor/Private/BehaviorTreeGraphNode_Composite.cpp b/Engine/Source/Editor/BehaviorTreeEditor/Private/BehaviorTreeGraphNode_Composite.cpp
--- a/Engine/Source/Editor/BehaviorTreeEditor/Private/BehaviorTreeGraphNode_Composite.cpp
+++ b/Engine/Source/Editor/BehaviorTreeEditor/Private/BehaviorTreeGraphNode_Composite.cpp
@@ -19,12 +19,6 @@ FText UBehaviorTreeGraphNode_Composite::GetNodeTitle(ENodeTitleType::Type TitleT
 	return Super::GetNodeTitle(TitleType);
 }
 
-FName UBehaviorTreeGraphNode_Composite::GetNameIcon() const
-{
-	return FName("BTEditor.Graph.BTNode.Composite.Icon");
-}
-
-
 void UBehaviorTreeGraphNode_Composite::GetContextMenuActions(const FGraphNodeContextMenuBuilder& Context) const
 {
 	AddContextMenuActionsDecorators(Context);
diff --git a/Engine/Source/Editor/BehaviorTreeEditor/Private/BehaviorTreeGraphNode_CompositeDecorator.cpp b/Engine/Source/Editor/BehaviorTreeEditor/Private/BehaviorTreeGraphNode_CompositeDecorator.cpp
--- a/Engine/Source/Editor/BehaviorTreeEditor/Private/BehaviorTreeGraphNode_CompositeDecorator.cpp
+++ b/Engine/Source/Editor/BehaviorTreeEditor/Private/BehaviorTreeGraphNode_CompositeDecorator.cpp
@@ -37,11 +37,6 @@ FText UBehaviorTreeGraphNode_CompositeDecorator::GetNodeTitle(ENodeTitleType::Ty
 	return FText::FromString(CompositeName.Len() ? CompositeName : GetNodeTypeDescription());
 }
 
-FName UBehaviorTreeGraphNode_CompositeDecorator::GetNameIcon() const
-{
-	return FName("BTEditor.Graph.BTNode.CompositeDecorator.Icon");
-}
-
 FString UBehaviorTreeGraphNode_CompositeDecorator::GetDescription() const
 {
 	return CachedDescription;
diff --git a/Engine/Source/Editor/BehaviorTreeEditor/Private/BehaviorTreeGraphNode_Decorator.cpp b/Engine/Source/Editor/BehaviorTreeEditor/Private/BehaviorTreeGraphNode_Decorator.cpp
--- a/Engine/Source/Editor/BehaviorTreeEditor/Private/BehaviorTreeGraphNode_Decorator.cpp
+++ b/Engine/Source/Editor/BehaviorTreeEditor/Private/BehaviorTreeGraphNode_Decorator.cpp
@@ -28,11 +28,6 @@ FText UBehaviorTreeGraphNode_Decorator::GetNodeTitle(ENodeTitleType::Type TitleT
 	return Super::GetNodeTitle(TitleType);
 }
 
-FName UBehaviorTreeGraphNode_Decorator::GetNameIcon() const
-{
-	return FName("BTEditor.Graph.BTNode.Decorator.Icon");
-}
-
 bool UBehaviorTreeGraphNode_Decorator::IsSubNode() const
 {
 	return true;
diff --git a/Engine/Source/Editor/BehaviorTreeEditor/Private/BehaviorTreeGraphNode_Service.cpp b/Engine/Source/Editor/BehaviorTreeEditor/Private/BehaviorTreeGraphNode_Service.cpp
--- a/Engine/Source/Editor/BehaviorTreeEditor/Private/BehaviorTreeGraphNode_Service.cpp
+++ b/Engine/Source/Editor/BehaviorTreeEditor/Private/BehaviorTreeGraphNode_Service.cpp
@@ -29,11 +29,6 @@ FText UBehaviorTreeGraphNode_Service::GetNodeTitle(ENodeTitleType::Type TitleTyp
 	return Super::GetNodeTitle(TitleType);
 }
 
-FName UBehaviorTreeGraphNode_Service::GetNameIcon() const
-{
-	return FName("BTEditor.Graph.BTNode.Service.Icon");
-}
-
 bool UBehaviorTreeGraphNode_Service::IsSubNode() const
 {
 	return true;
diff --git a/Engine/Source/Editor/BehaviorTreeEditor/Private/BehaviorTreeGraphNode_Task.cpp b/Engine/Source/Editor/BehaviorTreeEditor/Private/BehaviorTreeGraphNode_Task.cpp
--- a/Engine/Source/Editor/BehaviorTreeEditor/Private/BehaviorTreeGraphNode_Task.cpp
+++ b/Engine/Source/Editor/BehaviorTreeEditor/Private/BehaviorTreeGraphNode_Task.cpp
@@ -26,11 +26,6 @@ FText UBehaviorTreeGraphNode_Task::GetNodeTitle(ENodeTitleType::Type TitleType)
 	return Super::GetNodeTitle(TitleType);
 }
 
-FName UBehaviorTreeGraphNode_Task::GetNameIcon() const
-{
-	return FName("BTEditor.Graph.BTNode.Task.Icon");
-}
-
 void UBehaviorTreeGraphNode_Task::GetContextMenuActions(const FGraphNodeContextMenuBuilder& Context) const
 {
 	AddContextMenuActionsDecorators(Context);
diff --git a/Engine/Source/Editor/BehaviorTreeEditor/Private/SGraphNode_BehaviorTree.cpp b/Engine/Source/Editor/BehaviorTreeEditor/Private/SGraphNode_BehaviorTree.cpp
--- a/Engine/Source/Editor/BehaviorTreeEditor/Private/SGraphNode_BehaviorTree.cpp
+++ b/Engine/Source/Editor/BehaviorTreeEditor/Private/SGraphNode_BehaviorTree.cpp
@@ -349,8 +349,12 @@ void SGraphNode_BehaviorTree::UpdateGraphNode()
 									+SHorizontalBox::Slot()
 									.AspectRatio()
 									[
-										SNew(SImage)
-										.Image(this, &SGraphNode_BehaviorTree::GetNameIcon)
+										SNew(SOverlay)
+										+SOverlay::Slot()
+										[
+											SNew(SImage)
+											.Image(this, &SGraphNode_BehaviorTree::GetNameIcon)
+										]
 									]
 									+SHorizontalBox::Slot()
 									.Padding(FMargin(4.0f, 0.0f, 4.0f, 0.0f))
@@ -416,6 +420,15 @@ void SGraphNode_BehaviorTree::UpdateGraphNode()
 						.HeightOverride(4)
 					]
 				]
+				//drag marker overlay
+				+SOverlay::Slot()
+				.HAlign(HAlign_Left)
+				.VAlign(VAlign_Top)
+				[
+					SNew(SImage)
+					.Image(FEditorStyle::GetBrush(TEXT("BTEditor.Graph.BTNode.Blueprint")))
+					.Visibility(this, &SGraphNode_BehaviorTree::GetBlueprintIconVisibility)
+				]
 			]
 		];
 
@@ -908,18 +921,15 @@ FString SGraphNode_BehaviorTree::GetPreviewCornerText() const
 }
 
 const FSlateBrush* SGraphNode_BehaviorTree::GetNameIcon() const
+{	
+	UBehaviorTreeGraphNode* BTGraphNode = Cast<UBehaviorTreeGraphNode>(GraphNode);
+	return BTGraphNode != nullptr ? FEditorStyle::GetBrush(BTGraphNode->GetNameIcon()) : FEditorStyle::GetBrush(TEXT("BTEditor.Graph.BTNode.Icon"));
+}
+
+EVisibility SGraphNode_BehaviorTree::GetBlueprintIconVisibility() const
 {
-	UBehaviorTreeGraphNode_CompositeDecorator* CompDecorator = Cast<UBehaviorTreeGraphNode_CompositeDecorator>(GraphNode);
-	UBehaviorTreeGraphNode_Decorator* Decorator = Cast<UBehaviorTreeGraphNode_Decorator>(GraphNode);
-	UBTDecorator* DecoratorInstance = Decorator ? Cast<UBTDecorator>(Decorator->NodeInstance) : NULL;
-	
-	if ((DecoratorInstance && DecoratorInstance->GetFlowAbortMode() != EBTFlowAbortMode::None) ||
-		(CompDecorator && CompDecorator->bCanAbortFlow))
-	{
-		return FEditorStyle::GetBrush( TEXT("BTEditor.Graph.FlowControl.Icon") );
-	}
 	UBehaviorTreeGraphNode* BTGraphNode = Cast<UBehaviorTreeGraphNode>(GraphNode);
-	return BTGraphNode != NULL ? FEditorStyle::GetBrush(BTGraphNode->GetNameIcon()) : FEditorStyle::GetBrush(TEXT("BTEditor.Graph.BTNode.Icon"));
+	return (BTGraphNode != nullptr && BTGraphNode->UsesBlueprint()) ? EVisibility::Visible : EVisibility::Collapsed;
 }
 
 void SGraphNode_BehaviorTree::GetOverlayBrushes(bool bSelected, const FVector2D WidgetSize, TArray<FOverlayBrushInfo>& Brushes) const
diff --git a/Engine/Source/Editor/BehaviorTreeEditor/Private/SGraphNode_BehaviorTree.h b/Engine/Source/Editor/BehaviorTreeEditor/Private/SGraphNode_BehaviorTree.h
--- a/Engine/Source/Editor/BehaviorTreeEditor/Private/SGraphNode_BehaviorTree.h
+++ b/Engine/Source/Editor/BehaviorTreeEditor/Private/SGraphNode_BehaviorTree.h
@@ -99,6 +99,7 @@ protected:
 
 	virtual FString GetPreviewCornerText() const;
 	virtual const FSlateBrush* GetNameIcon() const;
+	virtual EVisibility GetBlueprintIconVisibility() const;
 
 	FString GetPinTooltip(UEdGraphPin* GraphPinObj) const;
 };
diff --git a/Engine/Source/Editor/EditorStyle/Private/SlateEditorStyle.cpp b/Engine/Source/Editor/EditorStyle/Private/SlateEditorStyle.cpp
--- a/Engine/Source/Editor/EditorStyle/Private/SlateEditorStyle.cpp
+++ b/Engine/Source/Editor/EditorStyle/Private/SlateEditorStyle.cpp
@@ -2777,16 +2777,34 @@ void FSlateEditorStyle::FStyle::SetupGraphEditorStyles()
 		// Behavior Tree Editor
 		{
 			Set( "BTEditor.Graph.BTNode.Body", new BOX_BRUSH( "/BehaviorTree/BTNode_ColorSpill", FMargin(16.f/64.f, 25.f/64.f, 16.f/64.f, 16.f/64.f) ) );
-			Set( "BTEditor.Graph.BTNode.Icon", new IMAGE_BRUSH( "/Persona/StateMachineEditor/State_Node_Icon_32x", Icon16x16 ) );
-
-			Set( "BTEditor.Graph.BTNode.Composite.Icon", new IMAGE_BRUSH( "/BehaviorTree/Composite_Node_Icon_32x", Icon16x16 ) );
-			Set( "BTEditor.Graph.BTNode.CompositeDecorator.Icon", new IMAGE_BRUSH( "/BehaviorTree/CompositeDecorator_Node_Icon_32x", Icon16x16 ) );
-			Set( "BTEditor.Graph.BTNode.Decorator.Icon", new IMAGE_BRUSH( "/BehaviorTree/Decorator_Node_Icon_32x", Icon16x16 ) );
-			Set( "BTEditor.Graph.BTNode.Root.Icon", new IMAGE_BRUSH( "/BehaviorTree/Root_Node_Icon_32x", Icon16x16 ) );
-			Set( "BTEditor.Graph.BTNode.Service.Icon", new IMAGE_BRUSH( "/BehaviorTree/Service_Node_Icon_32x", Icon16x16 ) );
-			Set( "BTEditor.Graph.BTNode.Task.Icon", new IMAGE_BRUSH( "/BehaviorTree/Task_Node_Icon_32x", Icon16x16 ) );
-
-			Set( "BTEditor.Graph.FlowControl.Icon", new IMAGE_BRUSH( "/Persona/StateMachineEditor/Conduit_Node_Icon_32x", Icon16x16 ) );
+			Set( "BTEditor.Graph.BTNode.Icon", new IMAGE_BRUSH( "/BehaviorTree/Icons/Run_Behaviour_24x", Icon16x16 ) );
+
+			Set( "BTEditor.Graph.BTNode.Root.Icon", new IMAGE_BRUSH( "/BehaviorTree/Icons/Run_Behaviour_24x", Icon24x24 ) );
+			Set( "BTEditor.Graph.BTNode.Composite.Selector.Icon", new IMAGE_BRUSH( "/BehaviorTree/Icons/Selector_24x", Icon24x24 ) );
+			Set( "BTEditor.Graph.BTNode.Composite.Sequence.Icon", new IMAGE_BRUSH( "/BehaviorTree/Icons/Sequence_24x", Icon24x24 ) );
+			Set( "BTEditor.Graph.BTNode.Composite.SimpleParallel.Icon", new IMAGE_BRUSH( "/BehaviorTree/Icons/Simple_Parallel_24x", Icon24x24 ) );
+			Set( "BTEditor.Graph.BTNode.Decorator.Blackboard.Icon", new IMAGE_BRUSH( "/BehaviorTree/Icons/Blackboard_24x", Icon24x24 ) );
+			Set( "BTEditor.Graph.BTNode.Decorator.CompareBlackboardEntries.Icon", new IMAGE_BRUSH( "/BehaviorTree/Icons/Compare_Blackboard_Entries_24x", Icon24x24 ) );
+			Set( "BTEditor.Graph.BTNode.Decorator.ConeCheck.Icon", new IMAGE_BRUSH( "/BehaviorTree/Icons/Cone_Check_24x", Icon24x24 ) );
+			Set( "BTEditor.Graph.BTNode.Decorator.Cooldown.Icon", new IMAGE_BRUSH( "/BehaviorTree/Icons/Cooldown_24x", Icon24x24 ) );
+			Set( "BTEditor.Graph.BTNode.Decorator.DoesPathExist.Icon", new IMAGE_BRUSH( "/BehaviorTree/Icons/Does_Path_Exist_24x", Icon24x24 ) );
+			Set( "BTEditor.Graph.BTNode.Decorator.ForceSuccess.Icon", new IMAGE_BRUSH( "/BehaviorTree/Icons/Force_Success_24x", Icon24x24 ) );
+			Set( "BTEditor.Graph.BTNode.Decorator.KeepInCone.Icon", new IMAGE_BRUSH( "/BehaviorTree/Icons/Keep_In_Cone_24x", Icon24x24 ) );
+			Set( "BTEditor.Graph.BTNode.Decorator.Loop.Icon", new IMAGE_BRUSH( "/BehaviorTree/Icons/Loop_24x", Icon24x24 ) );
+			Set( "BTEditor.Graph.BTNode.Decorator.Optional.Icon", new IMAGE_BRUSH( "/BehaviorTree/Icons/Optional_24x", Icon24x24 ) );
+			Set( "BTEditor.Graph.BTNode.Decorator.ReachedMoveGoal.Icon", new IMAGE_BRUSH( "/BehaviorTree/Icons/Reached_Move_Goal_24x", Icon24x24 ) );
+			Set( "BTEditor.Graph.BTNode.Decorator.TimeLimit.Icon", new IMAGE_BRUSH( "/BehaviorTree/Icons/Time_Limit_24x", Icon24x24 ) );
+			Set( "BTEditor.Graph.BTNode.Service.Icon", new IMAGE_BRUSH( "/BehaviorTree/Icons/Service_24x", Icon24x24 ) );
+			Set( "BTEditor.Graph.BTNode.Service.DefaultFocus.Icon", new IMAGE_BRUSH( "/BehaviorTree/Icons/Default_Focus_Service_24x", Icon24x24 ) );
+			Set( "BTEditor.Graph.BTNode.Task.Icon", new IMAGE_BRUSH( "/BehaviorTree/Icons/Task_24x", Icon24x24 ) );
+			Set( "BTEditor.Graph.BTNode.Task.MakeNoise.Icon", new IMAGE_BRUSH( "/BehaviorTree/Icons/Make_Noise_24x", Icon24x24 ) );
+			Set( "BTEditor.Graph.BTNode.Task.MoveDirectlyToward.Icon", new IMAGE_BRUSH( "/BehaviorTree/Icons/Move_Directly_Toward_24x", Icon24x24 ) );
+			Set( "BTEditor.Graph.BTNode.Task.MoveTo.Icon", new IMAGE_BRUSH( "/BehaviorTree/Icons/Move_To_24x", Icon24x24 ) );
+			Set( "BTEditor.Graph.BTNode.Task.PlaySound.Icon", new IMAGE_BRUSH( "/BehaviorTree/Icons/Play_Sound_24x", Icon24x24 ) );
+			Set( "BTEditor.Graph.BTNode.Task.RunBehavior.Icon", new IMAGE_BRUSH( "/BehaviorTree/Icons/Run_Behaviour_24x", Icon24x24 ) );
+			Set( "BTEditor.Graph.BTNode.Task.RunEQSQuery.Icon", new IMAGE_BRUSH( "/BehaviorTree/Icons/EQS_24x", Icon24x24 ) );
+			Set( "BTEditor.Graph.BTNode.Task.Wait.Icon", new IMAGE_BRUSH( "/BehaviorTree/Icons/Wait_24x", Icon24x24 ) );
+			Set( "BTEditor.Graph.BTNode.Blueprint", new IMAGE_BRUSH( "/BehaviorTree/Icons/Blueprint_Referencer_16x", Icon16x16 ) );
 
 			Set( "BTEditor.Debugger.BackOver", new IMAGE_BRUSH( "Icons/icon_step_back_40x", Icon40x40 ) );
 			Set( "BTEditor.Debugger.BackOver.Small", new IMAGE_BRUSH( "Icons/icon_step_back_40x", Icon20x20 ) );
diff --git a/Engine/Source/Runtime/AIModule/Classes/BehaviorTree/BTNode.h b/Engine/Source/Runtime/AIModule/Classes/BehaviorTree/BTNode.h
--- a/Engine/Source/Runtime/AIModule/Classes/BehaviorTree/BTNode.h
+++ b/Engine/Source/Runtime/AIModule/Classes/BehaviorTree/BTNode.h
@@ -115,6 +115,14 @@ class AIMODULE_API UBTNode : public UObject
 	/** @return string containing description of this node with all setup values */
 	virtual FString GetStaticDescription() const;
 
+#if WITH_EDITOR
+	/** Get the name of the icon used to display this node in the editor */
+	virtual FName GetNodeIconName() const;
+
+	/** Get whether this node is using a blueprint for its logic */
+	virtual bool UsesBlueprint() const;
+#endif
+
 protected:
 
 	/** node name */
diff --git a/Engine/Source/Runtime/AIModule/Classes/BehaviorTree/BTService.h b/Engine/Source/Runtime/AIModule/Classes/BehaviorTree/BTService.h
--- a/Engine/Source/Runtime/AIModule/Classes/BehaviorTree/BTService.h
+++ b/Engine/Source/Runtime/AIModule/Classes/BehaviorTree/BTService.h
@@ -53,4 +53,8 @@ protected:
 	/** called when auxiliary node becomes active
 	 * this function should be considered as const (don't modify state of object) if node is not instanced! */
 	virtual void OnBecomeRelevant(UBehaviorTreeComponent* OwnerComp, uint8* NodeMemory) OVERRIDE;
+
+#if WITH_EDITOR
+	virtual FName GetNodeIconName() const override;
+#endif // WITH_EDITOR
 };
diff --git a/Engine/Source/Runtime/AIModule/Classes/BehaviorTree/BTTaskNode.h b/Engine/Source/Runtime/AIModule/Classes/BehaviorTree/BTTaskNode.h
--- a/Engine/Source/Runtime/AIModule/Classes/BehaviorTree/BTTaskNode.h
+++ b/Engine/Source/Runtime/AIModule/Classes/BehaviorTree/BTTaskNode.h
@@ -35,6 +35,10 @@ class AIMODULE_API UBTTaskNode : public UBTNode
 	 * this function should be considered as const (don't modify state of object) if node is not instanced! */
 	virtual EBTNodeResult::Type AbortTask(class UBehaviorTreeComponent* OwnerComp, uint8* NodeMemory);
 
+#if WITH_EDITOR
+	virtual FName GetNodeIconName() const override;
+#endif // WITH_EDITOR
+
 	/** message observer's hook */
 	void ReceivedMessage(UBrainComponent* BrainComp, const struct FAIMessage& Message);
 
diff --git a/Engine/Source/Runtime/AIModule/Classes/BehaviorTree/Composites/BTComposite_Selector.h b/Engine/Source/Runtime/AIModule/Classes/BehaviorTree/Composites/BTComposite_Selector.h
--- a/Engine/Source/Runtime/AIModule/Classes/BehaviorTree/Composites/BTComposite_Selector.h
+++ b/Engine/Source/Runtime/AIModule/Classes/BehaviorTree/Composites/BTComposite_Selector.h
@@ -10,4 +10,8 @@ class UBTComposite_Selector : public UBTCompositeNode
 	GENERATED_UCLASS_BODY()
 
 	int32 GetNextChildHandler(struct FBehaviorTreeSearchData& SearchData, int32 PrevChild, EBTNodeResult::Type LastResult) const;
+
+#if WITH_EDITOR
+	virtual FName GetNodeIconName() const override;
+#endif
 };
diff --git a/Engine/Source/Runtime/AIModule/Classes/BehaviorTree/Composites/BTComposite_Sequence.h b/Engine/Source/Runtime/AIModule/Classes/BehaviorTree/Composites/BTComposite_Sequence.h
--- a/Engine/Source/Runtime/AIModule/Classes/BehaviorTree/Composites/BTComposite_Sequence.h
+++ b/Engine/Source/Runtime/AIModule/Classes/BehaviorTree/Composites/BTComposite_Sequence.h
@@ -13,5 +13,6 @@ class UBTComposite_Sequence : public UBTCompositeNode
 
 #if WITH_EDITOR
 	virtual bool CanAbortLowerPriority() const OVERRIDE;
+	virtual FName GetNodeIconName() const override;
 #endif
 };
diff --git a/Engine/Source/Runtime/AIModule/Classes/BehaviorTree/Composites/BTComposite_SimpleParallel.h b/Engine/Source/Runtime/AIModule/Classes/BehaviorTree/Composites/BTComposite_SimpleParallel.h
--- a/Engine/Source/Runtime/AIModule/Classes/BehaviorTree/Composites/BTComposite_SimpleParallel.h
+++ b/Engine/Source/Runtime/AIModule/Classes/BehaviorTree/Composites/BTComposite_SimpleParallel.h
@@ -62,5 +62,6 @@ class UBTComposite_SimpleParallel : public UBTCompositeNode
 #if WITH_EDITOR
 	virtual bool CanAbortLowerPriority() const OVERRIDE;
 	virtual bool CanAbortSelf() const OVERRIDE;
+	virtual FName GetNodeIconName() const override;
 #endif // WITH_EDITOR
 };
diff --git a/Engine/Source/Runtime/AIModule/Classes/BehaviorTree/Decorators/BTDecorator_BlackboardBase.h b/Engine/Source/Runtime/AIModule/Classes/BehaviorTree/Decorators/BTDecorator_BlackboardBase.h
--- a/Engine/Source/Runtime/AIModule/Classes/BehaviorTree/Decorators/BTDecorator_BlackboardBase.h
+++ b/Engine/Source/Runtime/AIModule/Classes/BehaviorTree/Decorators/BTDecorator_BlackboardBase.h
@@ -14,6 +14,10 @@ class AIMODULE_API UBTDecorator_BlackboardBase : public UBTDecorator
 	/** notify about change in blackboard keys */
 	virtual void OnBlackboardChange(const class UBlackboardComponent* Blackboard, uint8 ChangedKeyID);
 
+#if WITH_EDITOR
+	virtual FName GetNodeIconName() const override;
+#endif
+
 	/** get name of selected blackboard key */
 	FName GetSelectedBlackboardKey() const;
 
diff --git a/Engine/Source/Runtime/AIModule/Classes/BehaviorTree/Decorators/BTDecorator_BlueprintBase.h b/Engine/Source/Runtime/AIModule/Classes/BehaviorTree/Decorators/BTDecorator_BlueprintBase.h
--- a/Engine/Source/Runtime/AIModule/Classes/BehaviorTree/Decorators/BTDecorator_BlueprintBase.h
+++ b/Engine/Source/Runtime/AIModule/Classes/BehaviorTree/Decorators/BTDecorator_BlueprintBase.h
@@ -36,6 +36,11 @@ class AIMODULE_API UBTDecorator_BlueprintBase : public UBTDecorator
 	virtual void DescribeRuntimeValues(const class UBehaviorTreeComponent* OwnerComp, uint8* NodeMemory, EBTDescriptionVerbosity::Type Verbosity, TArray<FString>& Values) const OVERRIDE;
 	virtual bool CalculateRawConditionValue(class UBehaviorTreeComponent* OwnerComp, uint8* NodeMemory) const OVERRIDE;
 
+#if WITH_EDITOR
+	virtual FName GetNodeIconName() const override;
+	virtual bool UsesBlueprint() const override;
+#endif
+
 protected:
 
 	/** blackboard key names that should be observed */
diff --git a/Engine/Source/Runtime/AIModule/Classes/BehaviorTree/Decorators/BTDecorator_CompareBBEntries.h b/Engine/Source/Runtime/AIModule/Classes/BehaviorTree/Decorators/BTDecorator_CompareBBEntries.h
--- a/Engine/Source/Runtime/AIModule/Classes/BehaviorTree/Decorators/BTDecorator_CompareBBEntries.h
+++ b/Engine/Source/Runtime/AIModule/Classes/BehaviorTree/Decorators/BTDecorator_CompareBBEntries.h
@@ -47,4 +47,8 @@ public:
 	virtual void OnBlackboardChange(const class UBlackboardComponent* Blackboard, uint8 ChangedKeyID);	
 	virtual void OnBecomeRelevant(UBehaviorTreeComponent* OwnerComp, uint8* NodeMemory) OVERRIDE;
 	virtual void OnCeaseRelevant(UBehaviorTreeComponent* OwnerComp, uint8* NodeMemory) OVERRIDE;
+
+#if WITH_EDITOR
+	virtual FName GetNodeIconName() const override;
+#endif // WITH_EDITOR
 };
diff --git a/Engine/Source/Runtime/AIModule/Classes/BehaviorTree/Decorators/BTDecorator_ConeCheck.h b/Engine/Source/Runtime/AIModule/Classes/BehaviorTree/Decorators/BTDecorator_ConeCheck.h
--- a/Engine/Source/Runtime/AIModule/Classes/BehaviorTree/Decorators/BTDecorator_ConeCheck.h
+++ b/Engine/Source/Runtime/AIModule/Classes/BehaviorTree/Decorators/BTDecorator_ConeCheck.h
@@ -39,6 +39,10 @@ class AIMODULE_API UBTDecorator_ConeCheck : public UBTDecorator
 	virtual void DescribeRuntimeValues(const class UBehaviorTreeComponent* OwnerComp, uint8* NodeMemory, EBTDescriptionVerbosity::Type Verbosity, TArray<FString>& Values) const OVERRIDE;
 	virtual FString GetStaticDescription() const OVERRIDE;
 
+#if WITH_EDITOR
+	virtual FName GetNodeIconName() const override;
+#endif // WITH_EDITOR
+
 protected:
 
 	virtual bool CalculateRawConditionValue(class UBehaviorTreeComponent* OwnerComp, uint8* NodeMemory) const OVERRIDE;
diff --git a/Engine/Source/Runtime/AIModule/Classes/BehaviorTree/Decorators/BTDecorator_Cooldown.h b/Engine/Source/Runtime/AIModule/Classes/BehaviorTree/Decorators/BTDecorator_Cooldown.h
--- a/Engine/Source/Runtime/AIModule/Classes/BehaviorTree/Decorators/BTDecorator_Cooldown.h
+++ b/Engine/Source/Runtime/AIModule/Classes/BehaviorTree/Decorators/BTDecorator_Cooldown.h
@@ -29,6 +29,10 @@ class AIMODULE_API UBTDecorator_Cooldown : public UBTDecorator
 	virtual void DescribeRuntimeValues(const class UBehaviorTreeComponent* OwnerComp, uint8* NodeMemory, EBTDescriptionVerbosity::Type Verbosity, TArray<FString>& Values) const OVERRIDE;
 	virtual FString GetStaticDescription() const OVERRIDE;
 
+#if WITH_EDITOR
+	virtual FName GetNodeIconName() const override;
+#endif // WITH_EDITOR
+
 protected:
 
 	virtual void OnNodeDeactivation(struct FBehaviorTreeSearchData& SearchData, EBTNodeResult::Type NodeResult) OVERRIDE;
diff --git a/Engine/Source/Runtime/AIModule/Classes/BehaviorTree/Decorators/BTDecorator_DoesPathExist.h b/Engine/Source/Runtime/AIModule/Classes/BehaviorTree/Decorators/BTDecorator_DoesPathExist.h
--- a/Engine/Source/Runtime/AIModule/Classes/BehaviorTree/Decorators/BTDecorator_DoesPathExist.h
+++ b/Engine/Source/Runtime/AIModule/Classes/BehaviorTree/Decorators/BTDecorator_DoesPathExist.h
@@ -46,4 +46,8 @@ public:
 	virtual bool CalculateRawConditionValue(class UBehaviorTreeComponent* OwnerComp, uint8* NodeMemory) const OVERRIDE;
 	virtual FString GetStaticDescription() const OVERRIDE;
 	virtual void InitializeFromAsset(class UBehaviorTree* Asset) OVERRIDE;
+
+#if WITH_EDITOR
+	virtual FName GetNodeIconName() const override;
+#endif // WITH_EDITOR
 };
diff --git a/Engine/Source/Runtime/AIModule/Classes/BehaviorTree/Decorators/BTDecorator_ForceSuccess.h b/Engine/Source/Runtime/AIModule/Classes/BehaviorTree/Decorators/BTDecorator_ForceSuccess.h
--- a/Engine/Source/Runtime/AIModule/Classes/BehaviorTree/Decorators/BTDecorator_ForceSuccess.h
+++ b/Engine/Source/Runtime/AIModule/Classes/BehaviorTree/Decorators/BTDecorator_ForceSuccess.h
@@ -21,4 +21,8 @@ class UBTDecorator_ForceSuccess : public UBTDecorator
 protected:
 
 	virtual void OnNodeProcessed(struct FBehaviorTreeSearchData& SearchData, EBTNodeResult::Type& NodeResult);
+
+#if WITH_EDITOR
+	virtual FName GetNodeIconName() const override;
+#endif // WITH_EDITOR
 };
diff --git a/Engine/Source/Runtime/AIModule/Classes/BehaviorTree/Decorators/BTDecorator_KeepInCone.h b/Engine/Source/Runtime/AIModule/Classes/BehaviorTree/Decorators/BTDecorator_KeepInCone.h
--- a/Engine/Source/Runtime/AIModule/Classes/BehaviorTree/Decorators/BTDecorator_KeepInCone.h
+++ b/Engine/Source/Runtime/AIModule/Classes/BehaviorTree/Decorators/BTDecorator_KeepInCone.h
@@ -43,6 +43,10 @@ class AIMODULE_API UBTDecorator_KeepInCone : public UBTDecorator
 	virtual void DescribeRuntimeValues(const class UBehaviorTreeComponent* OwnerComp, uint8* NodeMemory, EBTDescriptionVerbosity::Type Verbosity, TArray<FString>& Values) const OVERRIDE;
 	virtual FString GetStaticDescription() const OVERRIDE;
 
+#if WITH_EDITOR
+	virtual FName GetNodeIconName() const override;
+#endif // WITH_EDITOR
+
 protected:
 
 	virtual void OnBecomeRelevant(UBehaviorTreeComponent* OwnerComp, uint8* NodeMemory) OVERRIDE;
diff --git a/Engine/Source/Runtime/AIModule/Classes/BehaviorTree/Decorators/BTDecorator_Loop.h b/Engine/Source/Runtime/AIModule/Classes/BehaviorTree/Decorators/BTDecorator_Loop.h
--- a/Engine/Source/Runtime/AIModule/Classes/BehaviorTree/Decorators/BTDecorator_Loop.h
+++ b/Engine/Source/Runtime/AIModule/Classes/BehaviorTree/Decorators/BTDecorator_Loop.h
@@ -25,6 +25,10 @@ class AIMODULE_API UBTDecorator_Loop : public UBTDecorator
 	virtual void DescribeRuntimeValues(const class UBehaviorTreeComponent* OwnerComp, uint8* NodeMemory, EBTDescriptionVerbosity::Type Verbosity, TArray<FString>& Values) const OVERRIDE;
 	virtual FString GetStaticDescription() const OVERRIDE;
 
+#if WITH_EDITOR
+	virtual FName GetNodeIconName() const override;
+#endif // WITH_EDITOR
+
 protected:
 
 	virtual void OnNodeActivation(struct FBehaviorTreeSearchData& SearchData) OVERRIDE;
diff --git a/Engine/Source/Runtime/AIModule/Classes/BehaviorTree/Decorators/BTDecorator_Optional.h b/Engine/Source/Runtime/AIModule/Classes/BehaviorTree/Decorators/BTDecorator_Optional.h
--- a/Engine/Source/Runtime/AIModule/Classes/BehaviorTree/Decorators/BTDecorator_Optional.h
+++ b/Engine/Source/Runtime/AIModule/Classes/BehaviorTree/Decorators/BTDecorator_Optional.h
@@ -12,4 +12,8 @@ class UBTDecorator_Optional : public UBTDecorator
 protected:
 
 	virtual void OnNodeProcessed(struct FBehaviorTreeSearchData& SearchData, EBTNodeResult::Type& NodeResult);
+
+#if WITH_EDITOR
+	virtual FName GetNodeIconName() const override;
+#endif // WITH_EDITOR
 };
diff --git a/Engine/Source/Runtime/AIModule/Classes/BehaviorTree/Decorators/BTDecorator_ReachedMoveGoal.h b/Engine/Source/Runtime/AIModule/Classes/BehaviorTree/Decorators/BTDecorator_ReachedMoveGoal.h
--- a/Engine/Source/Runtime/AIModule/Classes/BehaviorTree/Decorators/BTDecorator_ReachedMoveGoal.h
+++ b/Engine/Source/Runtime/AIModule/Classes/BehaviorTree/Decorators/BTDecorator_ReachedMoveGoal.h
@@ -10,4 +10,8 @@ class AIMODULE_API UBTDecorator_ReachedMoveGoal : public UBTDecorator
 	GENERATED_UCLASS_BODY()
 
 	virtual bool CalculateRawConditionValue(class UBehaviorTreeComponent* OwnerComp, uint8* NodeMemory) const OVERRIDE;
+
+#if WITH_EDITOR
+	virtual FName GetNodeIconName() const override;
+#endif // WITH_EDITOR
 };
diff --git a/Engine/Source/Runtime/AIModule/Classes/BehaviorTree/Decorators/BTDecorator_TimeLimit.h b/Engine/Source/Runtime/AIModule/Classes/BehaviorTree/Decorators/BTDecorator_TimeLimit.h
--- a/Engine/Source/Runtime/AIModule/Classes/BehaviorTree/Decorators/BTDecorator_TimeLimit.h
+++ b/Engine/Source/Runtime/AIModule/Classes/BehaviorTree/Decorators/BTDecorator_TimeLimit.h
@@ -15,6 +15,10 @@ class AIMODULE_API UBTDecorator_TimeLimit : public UBTDecorator
 	virtual void DescribeRuntimeValues(const class UBehaviorTreeComponent* OwnerComp, uint8* NodeMemory, EBTDescriptionVerbosity::Type Verbosity, TArray<FString>& Values) const OVERRIDE;
 	virtual FString GetStaticDescription() const OVERRIDE;
 
+#if WITH_EDITOR
+	virtual FName GetNodeIconName() const override;
+#endif // WITH_EDITOR
+
 protected:
 
 	virtual void OnBecomeRelevant(UBehaviorTreeComponent* OwnerComp, uint8* NodeMemory) OVERRIDE;
diff --git a/Engine/Source/Runtime/AIModule/Classes/BehaviorTree/Services/BTService_BlueprintBase.h b/Engine/Source/Runtime/AIModule/Classes/BehaviorTree/Services/BTService_BlueprintBase.h
--- a/Engine/Source/Runtime/AIModule/Classes/BehaviorTree/Services/BTService_BlueprintBase.h
+++ b/Engine/Source/Runtime/AIModule/Classes/BehaviorTree/Services/BTService_BlueprintBase.h
@@ -22,6 +22,10 @@ class AIMODULE_API UBTService_BlueprintBase : public UBTService
 	virtual FString GetStaticDescription() const OVERRIDE;
 	virtual void DescribeRuntimeValues(const class UBehaviorTreeComponent* OwnerComp, uint8* NodeMemory, EBTDescriptionVerbosity::Type Verbosity, TArray<FString>& Values) const OVERRIDE;
 
+#if WITH_EDITOR
+	virtual bool UsesBlueprint() const override;
+#endif
+
 protected:
 
 	/** properties with runtime values, stored only in class default object */
diff --git a/Engine/Source/Runtime/AIModule/Classes/BehaviorTree/Services/BTService_DefaultFocus.h b/Engine/Source/Runtime/AIModule/Classes/BehaviorTree/Services/BTService_DefaultFocus.h
--- a/Engine/Source/Runtime/AIModule/Classes/BehaviorTree/Services/BTService_DefaultFocus.h
+++ b/Engine/Source/Runtime/AIModule/Classes/BehaviorTree/Services/BTService_DefaultFocus.h
@@ -27,4 +27,8 @@ class UBTService_DefaultFocus : public UBTService_BlackboardBase
 
 	virtual FString GetStaticDescription() const OVERRIDE;
 	virtual void DescribeRuntimeValues(const class UBehaviorTreeComponent* OwnerComp, uint8* NodeMemory, EBTDescriptionVerbosity::Type Verbosity, TArray<FString>& Values) const OVERRIDE;
+
+#if WITH_EDITOR
+	virtual FName GetNodeIconName() const override;
+#endif // WITH_EDITOR
 };
diff --git a/Engine/Source/Runtime/AIModule/Classes/BehaviorTree/Tasks/BTTask_BlueprintBase.h b/Engine/Source/Runtime/AIModule/Classes/BehaviorTree/Tasks/BTTask_BlueprintBase.h
--- a/Engine/Source/Runtime/AIModule/Classes/BehaviorTree/Tasks/BTTask_BlueprintBase.h
+++ b/Engine/Source/Runtime/AIModule/Classes/BehaviorTree/Tasks/BTTask_BlueprintBase.h
@@ -25,6 +25,10 @@ class AIMODULE_API UBTTask_BlueprintBase : public UBTTaskNode
 	virtual FString GetStaticDescription() const OVERRIDE;
 	virtual void DescribeRuntimeValues(const class UBehaviorTreeComponent* OwnerComp, uint8* NodeMemory, EBTDescriptionVerbosity::Type Verbosity, TArray<FString>& Values) const OVERRIDE;
 
+#if WITH_EDITOR
+	virtual bool UsesBlueprint() const override;
+#endif
+
 protected:
 
 	/** temporary variable for ReceiveExecute(Abort)-FinishExecute(Abort) chain */
diff --git a/Engine/Source/Runtime/AIModule/Classes/BehaviorTree/Tasks/BTTask_MakeNoise.h b/Engine/Source/Runtime/AIModule/Classes/BehaviorTree/Tasks/BTTask_MakeNoise.h
--- a/Engine/Source/Runtime/AIModule/Classes/BehaviorTree/Tasks/BTTask_MakeNoise.h
+++ b/Engine/Source/Runtime/AIModule/Classes/BehaviorTree/Tasks/BTTask_MakeNoise.h
@@ -13,4 +13,8 @@ class UBTTask_MakeNoise : public UBTTaskNode
 	float Loudnes;
 
 	virtual EBTNodeResult::Type ExecuteTask(UBehaviorTreeComponent* OwnerComp, uint8* NodeMemory) OVERRIDE;
+
+#if WITH_EDITOR
+	virtual FName GetNodeIconName() const override;
+#endif // WITH_EDITOR
 };
diff --git a/Engine/Source/Runtime/AIModule/Classes/BehaviorTree/Tasks/BTTask_MoveDirectlyToward.h b/Engine/Source/Runtime/AIModule/Classes/BehaviorTree/Tasks/BTTask_MoveDirectlyToward.h
--- a/Engine/Source/Runtime/AIModule/Classes/BehaviorTree/Tasks/BTTask_MoveDirectlyToward.h
+++ b/Engine/Source/Runtime/AIModule/Classes/BehaviorTree/Tasks/BTTask_MoveDirectlyToward.h
@@ -36,4 +36,8 @@ class AIMODULE_API UBTTask_MoveDirectlyToward : public UBTTask_BlackboardBase
 	virtual uint16 GetInstanceMemorySize() const OVERRIDE;
 	virtual void DescribeRuntimeValues(const class UBehaviorTreeComponent* OwnerComp, uint8* NodeMemory, EBTDescriptionVerbosity::Type Verbosity, TArray<FString>& Values) const OVERRIDE;
 	virtual FString GetStaticDescription() const OVERRIDE;
+
+#if WITH_EDITOR
+	virtual FName GetNodeIconName() const override;
+#endif // WITH_EDITOR
 };
diff --git a/Engine/Source/Runtime/AIModule/Classes/BehaviorTree/Tasks/BTTask_MoveTo.h b/Engine/Source/Runtime/AIModule/Classes/BehaviorTree/Tasks/BTTask_MoveTo.h
--- a/Engine/Source/Runtime/AIModule/Classes/BehaviorTree/Tasks/BTTask_MoveTo.h
+++ b/Engine/Source/Runtime/AIModule/Classes/BehaviorTree/Tasks/BTTask_MoveTo.h
@@ -37,6 +37,10 @@ class AIMODULE_API UBTTask_MoveTo : public UBTTask_BlackboardBase
 	virtual void DescribeRuntimeValues(const class UBehaviorTreeComponent* OwnerComp, uint8* NodeMemory, EBTDescriptionVerbosity::Type Verbosity, TArray<FString>& Values) const OVERRIDE;
 	virtual FString GetStaticDescription() const OVERRIDE;
 
+#if WITH_EDITOR
+	virtual FName GetNodeIconName() const override;
+#endif // WITH_EDITOR
+
 protected:
 
 	EBTNodeResult::Type PerformMoveTask(class UBehaviorTreeComponent* OwnerComp, uint8* NodeMemory);
diff --git a/Engine/Source/Runtime/AIModule/Classes/BehaviorTree/Tasks/BTTask_PlaySound.h b/Engine/Source/Runtime/AIModule/Classes/BehaviorTree/Tasks/BTTask_PlaySound.h
--- a/Engine/Source/Runtime/AIModule/Classes/BehaviorTree/Tasks/BTTask_PlaySound.h
+++ b/Engine/Source/Runtime/AIModule/Classes/BehaviorTree/Tasks/BTTask_PlaySound.h
@@ -17,4 +17,8 @@ class UBTTask_PlaySound : public UBTTaskNode
 
 	virtual EBTNodeResult::Type ExecuteTask(UBehaviorTreeComponent* OwnerComp, uint8* NodeMemory) OVERRIDE;
 	virtual FString GetStaticDescription() const OVERRIDE;
+
+#if WITH_EDITOR
+	virtual FName GetNodeIconName() const override;
+#endif // WITH_EDITOR
 };
diff --git a/Engine/Source/Runtime/AIModule/Classes/BehaviorTree/Tasks/BTTask_RunBehavior.h b/Engine/Source/Runtime/AIModule/Classes/BehaviorTree/Tasks/BTTask_RunBehavior.h
--- a/Engine/Source/Runtime/AIModule/Classes/BehaviorTree/Tasks/BTTask_RunBehavior.h
+++ b/Engine/Source/Runtime/AIModule/Classes/BehaviorTree/Tasks/BTTask_RunBehavior.h
@@ -27,6 +27,10 @@ class UBTTask_RunBehavior : public UBTTaskNode
 	virtual uint16 GetInstanceMemorySize() const OVERRIDE;
 	virtual FString GetStaticDescription() const OVERRIDE;
 
+#if WITH_EDITOR
+	virtual FName GetNodeIconName() const override;
+#endif // WITH_EDITOR
+
 	/** called on instance startup, prepares root level nodes to use */
 	void InjectNodes(UBehaviorTreeComponent* OwnerComp, uint8* NodeMemory, int32& InstancedIndex) const;
 
diff --git a/Engine/Source/Runtime/AIModule/Classes/BehaviorTree/Tasks/BTTask_RunEQSQuery.h b/Engine/Source/Runtime/AIModule/Classes/BehaviorTree/Tasks/BTTask_RunEQSQuery.h
--- a/Engine/Source/Runtime/AIModule/Classes/BehaviorTree/Tasks/BTTask_RunEQSQuery.h
+++ b/Engine/Source/Runtime/AIModule/Classes/BehaviorTree/Tasks/BTTask_RunEQSQuery.h
@@ -36,6 +36,7 @@ class UBTTask_RunEQSQuery : public UBTTask_BlackboardBase
 #if WITH_EDITOR
 	/** prepare query params */
 	virtual void PostEditChangeProperty(struct FPropertyChangedEvent& PropertyChangedEvent) OVERRIDE;
+	virtual FName GetNodeIconName() const override;
 #endif
 
 protected:
diff --git a/Engine/Source/Runtime/AIModule/Classes/BehaviorTree/Tasks/BTTask_Wait.h b/Engine/Source/Runtime/AIModule/Classes/BehaviorTree/Tasks/BTTask_Wait.h
--- a/Engine/Source/Runtime/AIModule/Classes/BehaviorTree/Tasks/BTTask_Wait.h
+++ b/Engine/Source/Runtime/AIModule/Classes/BehaviorTree/Tasks/BTTask_Wait.h
@@ -23,6 +23,10 @@ class UBTTask_Wait : public UBTTaskNode
 	virtual void DescribeRuntimeValues(const class UBehaviorTreeComponent* OwnerComp, uint8* NodeMemory, EBTDescriptionVerbosity::Type Verbosity, TArray<FString>& Values) const OVERRIDE;
 	virtual FString GetStaticDescription() const OVERRIDE;
 
+#if WITH_EDITOR
+	virtual FName GetNodeIconName() const override;
+#endif // WITH_EDITOR
+
 protected:
 
 	virtual void TickTask(class UBehaviorTreeComponent* OwnerComp, uint8* NodeMemory, float DeltaSeconds) OVERRIDE;
diff --git a/Engine/Source/Runtime/AIModule/Private/BehaviorTree/BTNode.cpp b/Engine/Source/Runtime/AIModule/Private/BehaviorTree/BTNode.cpp
--- a/Engine/Source/Runtime/AIModule/Private/BehaviorTree/BTNode.cpp
+++ b/Engine/Source/Runtime/AIModule/Private/BehaviorTree/BTNode.cpp
@@ -156,4 +156,18 @@ FString UBTNode::GetStaticDescription() const
 void UBTNode::DescribeRuntimeValues(const class UBehaviorTreeComponent* OwnerComp, uint8* NodeMemory, EBTDescriptionVerbosity::Type Verbosity, TArray<FString>& Values) const
 {
 	// nothing stored in memory for base class
-}
\ No newline at end of file
+}
+
+#if WITH_EDITOR
+
+FName UBTNode::GetNodeIconName() const
+{
+	return NAME_None;
+}
+
+bool UBTNode::UsesBlueprint() const
+{
+	return false;
+}
+
+#endif
\ No newline at end of file
diff --git a/Engine/Source/Runtime/AIModule/Private/BehaviorTree/BTService.cpp b/Engine/Source/Runtime/AIModule/Private/BehaviorTree/BTService.cpp
--- a/Engine/Source/Runtime/AIModule/Private/BehaviorTree/BTService.cpp
+++ b/Engine/Source/Runtime/AIModule/Private/BehaviorTree/BTService.cpp
@@ -34,3 +34,11 @@ FString UBTService::GetStaticDescription() const
 	return FString::Printf(TEXT("%s: tick every %s"), *UBehaviorTreeTypes::GetShortTypeName(this), *IntervalDesc);
 }
 
+#if WITH_EDITOR
+
+FName UBTService::GetNodeIconName() const
+{
+	return FName("BTEditor.Graph.BTNode.Service.Icon");
+}
+
+#endif // WITH_EDITOR
diff --git a/Engine/Source/Runtime/AIModule/Private/BehaviorTree/BTTaskNode.cpp b/Engine/Source/Runtime/AIModule/Private/BehaviorTree/BTTaskNode.cpp
--- a/Engine/Source/Runtime/AIModule/Private/BehaviorTree/BTTaskNode.cpp
+++ b/Engine/Source/Runtime/AIModule/Private/BehaviorTree/BTTaskNode.cpp
@@ -101,3 +101,12 @@ void UBTTaskNode::StopWaitingForMessages(class UBehaviorTreeComponent* OwnerComp
 	// messages delegates should be called on node instances (if they exists)
 	OwnerComp->UnregisterMessageObserversFrom(this);
 }
+
+#if WITH_EDITOR
+
+FName UBTTaskNode::GetNodeIconName() const
+{
+	return FName("BTEditor.Graph.BTNode.Task.Icon");
+}
+
+#endif	// WITH_EDITOR
\ No newline at end of file
diff --git a/Engine/Source/Runtime/AIModule/Private/BehaviorTree/Composites/BTComposite_Selector.cpp b/Engine/Source/Runtime/AIModule/Private/BehaviorTree/Composites/BTComposite_Selector.cpp
--- a/Engine/Source/Runtime/AIModule/Private/BehaviorTree/Composites/BTComposite_Selector.cpp
+++ b/Engine/Source/Runtime/AIModule/Private/BehaviorTree/Composites/BTComposite_Selector.cpp
@@ -29,3 +29,12 @@ int32 UBTComposite_Selector::GetNextChildHandler(struct FBehaviorTreeSearchData&
 
 	return NextChildIdx;
 }
+
+#if WITH_EDITOR
+
+FName UBTComposite_Selector::GetNodeIconName() const
+{
+	return FName("BTEditor.Graph.BTNode.Composite.Selector.Icon");
+}
+
+#endif
\ No newline at end of file
diff --git a/Engine/Source/Runtime/AIModule/Private/BehaviorTree/Composites/BTComposite_Sequence.cpp b/Engine/Source/Runtime/AIModule/Private/BehaviorTree/Composites/BTComposite_Sequence.cpp
--- a/Engine/Source/Runtime/AIModule/Private/BehaviorTree/Composites/BTComposite_Sequence.cpp
+++ b/Engine/Source/Runtime/AIModule/Private/BehaviorTree/Composites/BTComposite_Sequence.cpp
@@ -38,4 +38,9 @@ bool UBTComposite_Sequence::CanAbortLowerPriority() const
 	return false;
 }
 
+FName UBTComposite_Sequence::GetNodeIconName() const
+{
+	return FName("BTEditor.Graph.BTNode.Composite.Sequence.Icon");
+}
+
 #endif
diff --git a/Engine/Source/Runtime/AIModule/Private/BehaviorTree/Composites/BTComposite_SimpleParallel.cpp b/Engine/Source/Runtime/AIModule/Private/BehaviorTree/Composites/BTComposite_SimpleParallel.cpp
--- a/Engine/Source/Runtime/AIModule/Private/BehaviorTree/Composites/BTComposite_SimpleParallel.cpp
+++ b/Engine/Source/Runtime/AIModule/Private/BehaviorTree/Composites/BTComposite_SimpleParallel.cpp
@@ -152,4 +152,9 @@ bool UBTComposite_SimpleParallel::CanAbortSelf() const
 	return false;
 }
 
+FName UBTComposite_SimpleParallel::GetNodeIconName() const
+{
+	return FName("BTEditor.Graph.BTNode.Composite.SimpleParallel.Icon");
+}
+
 #endif // WITH_EDITOR
diff --git a/Engine/Source/Runtime/AIModule/Private/BehaviorTree/Decorators/BTDecorator_BlackboardBase.cpp b/Engine/Source/Runtime/AIModule/Private/BehaviorTree/Decorators/BTDecorator_BlackboardBase.cpp
--- a/Engine/Source/Runtime/AIModule/Private/BehaviorTree/Decorators/BTDecorator_BlackboardBase.cpp
+++ b/Engine/Source/Runtime/AIModule/Private/BehaviorTree/Decorators/BTDecorator_BlackboardBase.cpp
@@ -52,3 +52,12 @@ void UBTDecorator_BlackboardBase::OnBlackboardChange(const UBlackboardComponent*
 		BehaviorComp->RequestExecution(this);		
 	}
 }
+
+#if WITH_EDITOR
+
+FName UBTDecorator_BlackboardBase::GetNodeIconName() const
+{
+	return FName("BTEditor.Graph.BTNode.Decorator.Blackboard.Icon");
+}
+
+#endif	// WITH_EDITOR
\ No newline at end of file
diff --git a/Engine/Source/Runtime/AIModule/Private/BehaviorTree/Decorators/BTDecorator_BlueprintBase.cpp b/Engine/Source/Runtime/AIModule/Private/BehaviorTree/Decorators/BTDecorator_BlueprintBase.cpp
--- a/Engine/Source/Runtime/AIModule/Private/BehaviorTree/Decorators/BTDecorator_BlueprintBase.cpp
+++ b/Engine/Source/Runtime/AIModule/Private/BehaviorTree/Decorators/BTDecorator_BlueprintBase.cpp
@@ -193,3 +193,24 @@ void UBTDecorator_BlueprintBase::OnBlackboardChange(const UBlackboardComponent*
 		BehaviorComp->RequestExecution(this);		
 	}
 }
+
+#if WITH_EDITOR
+
+FName UBTDecorator_BlueprintBase::GetNodeIconName() const
+{
+	if(bImplementsReceiveConditionCheck)
+	{
+		return FName("BTEditor.Graph.BTNode.Decorator.Conditional.Icon");
+	}
+	else
+	{
+		return FName("BTEditor.Graph.BTNode.Decorator.NonConditional.Icon");
+	}
+}
+
+bool UBTDecorator_BlueprintBase::UsesBlueprint() const
+{
+	return true;
+}
+
+#endif // WITH_EDITOR
diff --git a/Engine/Source/Runtime/AIModule/Private/BehaviorTree/Decorators/BTDecorator_CompareBBEntries.cpp b/Engine/Source/Runtime/AIModule/Private/BehaviorTree/Decorators/BTDecorator_CompareBBEntries.cpp
--- a/Engine/Source/Runtime/AIModule/Private/BehaviorTree/Decorators/BTDecorator_CompareBBEntries.cpp
+++ b/Engine/Source/Runtime/AIModule/Private/BehaviorTree/Decorators/BTDecorator_CompareBBEntries.cpp
@@ -93,3 +93,12 @@ void UBTDecorator_CompareBBEntries::OnBlackboardChange(const UBlackboardComponen
 		BehaviorComp->RequestExecution(this);		
 	}
 }
+
+#if WITH_EDITOR
+
+FName UBTDecorator_CompareBBEntries::GetNodeIconName() const
+{
+	return FName("BTEditor.Graph.BTNode.Decorator.CompareBlackboardEntries.Icon");
+}
+
+#endif	// WITH_EDITOR
\ No newline at end of file
diff --git a/Engine/Source/Runtime/AIModule/Private/BehaviorTree/Decorators/BTDecorator_ConeCheck.cpp b/Engine/Source/Runtime/AIModule/Private/BehaviorTree/Decorators/BTDecorator_ConeCheck.cpp
--- a/Engine/Source/Runtime/AIModule/Private/BehaviorTree/Decorators/BTDecorator_ConeCheck.cpp
+++ b/Engine/Source/Runtime/AIModule/Private/BehaviorTree/Decorators/BTDecorator_ConeCheck.cpp
@@ -6,7 +6,7 @@
 
 UBTDecorator_ConeCheck::UBTDecorator_ConeCheck(const class FPostConstructInitializeProperties& PCIP) : Super(PCIP)
 {
-	NodeName = "Keep in Cone";
+	NodeName = "Cone Check";
 
 	// accept only actors and vectors
 	ConeOrigin.AddObjectFilter(this, AActor::StaticClass());
@@ -134,3 +134,12 @@ uint16 UBTDecorator_ConeCheck::GetInstanceMemorySize() const
 {
 	return sizeof(TNodeInstanceMemory);
 }
+
+#if WITH_EDITOR
+
+FName UBTDecorator_ConeCheck::GetNodeIconName() const
+{
+	return FName("BTEditor.Graph.BTNode.Decorator.ConeCheck.Icon");
+}
+
+#endif	// WITH_EDITOR
\ No newline at end of file
diff --git a/Engine/Source/Runtime/AIModule/Private/BehaviorTree/Decorators/BTDecorator_Cooldown.cpp b/Engine/Source/Runtime/AIModule/Private/BehaviorTree/Decorators/BTDecorator_Cooldown.cpp
--- a/Engine/Source/Runtime/AIModule/Private/BehaviorTree/Decorators/BTDecorator_Cooldown.cpp
+++ b/Engine/Source/Runtime/AIModule/Private/BehaviorTree/Decorators/BTDecorator_Cooldown.cpp
@@ -81,3 +81,12 @@ uint16 UBTDecorator_Cooldown::GetInstanceMemorySize() const
 {
 	return sizeof(FBTCooldownDecoratorMemory);
 }
+
+#if WITH_EDITOR
+
+FName UBTDecorator_Cooldown::GetNodeIconName() const
+{
+	return FName("BTEditor.Graph.BTNode.Decorator.Cooldown.Icon");
+}
+
+#endif	// WITH_EDITOR
\ No newline at end of file
diff --git a/Engine/Source/Runtime/AIModule/Private/BehaviorTree/Decorators/BTDecorator_DoesPathExist.cpp b/Engine/Source/Runtime/AIModule/Private/BehaviorTree/Decorators/BTDecorator_DoesPathExist.cpp
--- a/Engine/Source/Runtime/AIModule/Private/BehaviorTree/Decorators/BTDecorator_DoesPathExist.cpp
+++ b/Engine/Source/Runtime/AIModule/Private/BehaviorTree/Decorators/BTDecorator_DoesPathExist.cpp
@@ -92,3 +92,12 @@ FString UBTDecorator_DoesPathExist::GetStaticDescription() const
 		*BlackboardKeyB.SelectedKeyName.ToString(),
 		*PathTypeEnum->GetEnumName(PathQueryType));
 }
+
+#if WITH_EDITOR
+
+FName UBTDecorator_DoesPathExist::GetNodeIconName() const
+{
+	return FName("BTEditor.Graph.BTNode.Decorator.DoesPathExist.Icon");
+}
+
+#endif	// WITH_EDITOR
\ No newline at end of file
diff --git a/Engine/Source/Runtime/AIModule/Private/BehaviorTree/Decorators/BTDecorator_ForceSuccess.cpp b/Engine/Source/Runtime/AIModule/Private/BehaviorTree/Decorators/BTDecorator_ForceSuccess.cpp
--- a/Engine/Source/Runtime/AIModule/Private/BehaviorTree/Decorators/BTDecorator_ForceSuccess.cpp
+++ b/Engine/Source/Runtime/AIModule/Private/BehaviorTree/Decorators/BTDecorator_ForceSuccess.cpp
@@ -17,3 +17,12 @@ void UBTDecorator_ForceSuccess::OnNodeProcessed(struct FBehaviorTreeSearchData&
 {
 	NodeResult = EBTNodeResult::Succeeded;
 }
+
+#if WITH_EDITOR
+
+FName UBTDecorator_ForceSuccess::GetNodeIconName() const
+{
+	return FName("BTEditor.Graph.BTNode.Decorator.ForceSuccess.Icon");
+}
+
+#endif	// WITH_EDITOR
\ No newline at end of file
diff --git a/Engine/Source/Runtime/AIModule/Private/BehaviorTree/Decorators/BTDecorator_KeepInCone.cpp b/Engine/Source/Runtime/AIModule/Private/BehaviorTree/Decorators/BTDecorator_KeepInCone.cpp
--- a/Engine/Source/Runtime/AIModule/Private/BehaviorTree/Decorators/BTDecorator_KeepInCone.cpp
+++ b/Engine/Source/Runtime/AIModule/Private/BehaviorTree/Decorators/BTDecorator_KeepInCone.cpp
@@ -127,3 +127,12 @@ uint16 UBTDecorator_KeepInCone::GetInstanceMemorySize() const
 {
 	return sizeof(TNodeInstanceMemory);
 }
+
+#if WITH_EDITOR
+
+FName UBTDecorator_KeepInCone::GetNodeIconName() const
+{
+	return FName("BTEditor.Graph.BTNode.Decorator.KeepInCone.Icon");
+}
+
+#endif	// WITH_EDITOR
\ No newline at end of file
diff --git a/Engine/Source/Runtime/AIModule/Private/BehaviorTree/Decorators/BTDecorator_Loop.cpp b/Engine/Source/Runtime/AIModule/Private/BehaviorTree/Decorators/BTDecorator_Loop.cpp
--- a/Engine/Source/Runtime/AIModule/Private/BehaviorTree/Decorators/BTDecorator_Loop.cpp
+++ b/Engine/Source/Runtime/AIModule/Private/BehaviorTree/Decorators/BTDecorator_Loop.cpp
@@ -61,3 +61,12 @@ uint16 UBTDecorator_Loop::GetInstanceMemorySize() const
 {
 	return sizeof(FBTLoopDecoratorMemory);
 }
+
+#if WITH_EDITOR
+
+FName UBTDecorator_Loop::GetNodeIconName() const
+{
+	return FName("BTEditor.Graph.BTNode.Decorator.Loop.Icon");
+}
+
+#endif	// WITH_EDITOR
\ No newline at end of file
diff --git a/Engine/Source/Runtime/AIModule/Private/BehaviorTree/Decorators/BTDecorator_Optional.cpp b/Engine/Source/Runtime/AIModule/Private/BehaviorTree/Decorators/BTDecorator_Optional.cpp
--- a/Engine/Source/Runtime/AIModule/Private/BehaviorTree/Decorators/BTDecorator_Optional.cpp
+++ b/Engine/Source/Runtime/AIModule/Private/BehaviorTree/Decorators/BTDecorator_Optional.cpp
@@ -17,3 +17,12 @@ void UBTDecorator_Optional::OnNodeProcessed(struct FBehaviorTreeSearchData& Sear
 {
 	NodeResult = EBTNodeResult::Optional;
 }
+
+#if WITH_EDITOR
+
+FName UBTDecorator_Optional::GetNodeIconName() const
+{
+	return FName("BTEditor.Graph.BTNode.Decorator.Optional.Icon");
+}
+
+#endif	// WITH_EDITOR
\ No newline at end of file
diff --git a/Engine/Source/Runtime/AIModule/Private/BehaviorTree/Decorators/BTDecorator_ReachedMoveGoal.cpp b/Engine/Source/Runtime/AIModule/Private/BehaviorTree/Decorators/BTDecorator_ReachedMoveGoal.cpp
--- a/Engine/Source/Runtime/AIModule/Private/BehaviorTree/Decorators/BTDecorator_ReachedMoveGoal.cpp
+++ b/Engine/Source/Runtime/AIModule/Private/BehaviorTree/Decorators/BTDecorator_ReachedMoveGoal.cpp
@@ -20,3 +20,12 @@ bool UBTDecorator_ReachedMoveGoal::CalculateRawConditionValue(class UBehaviorTre
 	const bool bReachedGoal = AIOwner && AIOwner->PathFollowingComponent && AIOwner->PathFollowingComponent->DidMoveReachGoal();
 	return bReachedGoal;
 }
+
+#if WITH_EDITOR
+
+FName UBTDecorator_ReachedMoveGoal::GetNodeIconName() const
+{
+	return FName("BTEditor.Graph.BTNode.Decorator.ReachedMoveGoal.Icon");
+}
+
+#endif	// WITH_EDITOR
\ No newline at end of file
diff --git a/Engine/Source/Runtime/AIModule/Private/BehaviorTree/Decorators/BTDecorator_TimeLimit.cpp b/Engine/Source/Runtime/AIModule/Private/BehaviorTree/Decorators/BTDecorator_TimeLimit.cpp
--- a/Engine/Source/Runtime/AIModule/Private/BehaviorTree/Decorators/BTDecorator_TimeLimit.cpp
+++ b/Engine/Source/Runtime/AIModule/Private/BehaviorTree/Decorators/BTDecorator_TimeLimit.cpp
@@ -51,3 +51,12 @@ void UBTDecorator_TimeLimit::DescribeRuntimeValues(const class UBehaviorTreeComp
 			*FString::SanitizeFloat(TimeLeft)));
 	}
 }
+
+#if WITH_EDITOR
+
+FName UBTDecorator_TimeLimit::GetNodeIconName() const
+{
+	return FName("BTEditor.Graph.BTNode.Decorator.TimeLimit.Icon");
+}
+
+#endif	// WITH_EDITOR
\ No newline at end of file
diff --git a/Engine/Source/Runtime/AIModule/Private/BehaviorTree/Services/BTService_BlueprintBase.cpp b/Engine/Source/Runtime/AIModule/Private/BehaviorTree/Services/BTService_BlueprintBase.cpp
--- a/Engine/Source/Runtime/AIModule/Private/BehaviorTree/Services/BTService_BlueprintBase.cpp
+++ b/Engine/Source/Runtime/AIModule/Private/BehaviorTree/Services/BTService_BlueprintBase.cpp
@@ -100,3 +100,12 @@ void UBTService_BlueprintBase::DescribeRuntimeValues(const class UBehaviorTreeCo
 		BlueprintNodeHelpers::DescribeRuntimeValues(this, CDO->PropertyData, Values);
 	}
 }
+
+#if WITH_EDITOR
+
+bool UBTService_BlueprintBase::UsesBlueprint() const
+{
+	return true;
+}
+
+#endif // WITH_EDITOR
\ No newline at end of file
diff --git a/Engine/Source/Runtime/AIModule/Private/BehaviorTree/Services/BTService_DefaultFocus.cpp b/Engine/Source/Runtime/AIModule/Private/BehaviorTree/Services/BTService_DefaultFocus.cpp
--- a/Engine/Source/Runtime/AIModule/Private/BehaviorTree/Services/BTService_DefaultFocus.cpp
+++ b/Engine/Source/Runtime/AIModule/Private/BehaviorTree/Services/BTService_DefaultFocus.cpp
@@ -99,3 +99,12 @@ uint16 UBTService_DefaultFocus::GetInstanceMemorySize() const
 {
 	return sizeof(FBTDefaultFocusMemory);
 }
+
+#if WITH_EDITOR
+
+FName UBTService_DefaultFocus::GetNodeIconName() const
+{
+	return FName("BTEditor.Graph.BTNode.Service.DefaultFocus.Icon");
+}
+
+#endif	// WITH_EDITOR
\ No newline at end of file
diff --git a/Engine/Source/Runtime/AIModule/Private/BehaviorTree/Tasks/BTTask_BlueprintBase.cpp b/Engine/Source/Runtime/AIModule/Private/BehaviorTree/Tasks/BTTask_BlueprintBase.cpp
--- a/Engine/Source/Runtime/AIModule/Private/BehaviorTree/Tasks/BTTask_BlueprintBase.cpp
+++ b/Engine/Source/Runtime/AIModule/Private/BehaviorTree/Tasks/BTTask_BlueprintBase.cpp
@@ -153,3 +153,12 @@ void UBTTask_BlueprintBase::DescribeRuntimeValues(const class UBehaviorTreeCompo
 		BlueprintNodeHelpers::DescribeRuntimeValues(this, CDO->PropertyData, Values);
 	}
 }
+
+#if WITH_EDITOR
+
+bool UBTTask_BlueprintBase::UsesBlueprint() const
+{
+	return true;
+}
+
+#endif // WITH_EDITOR
\ No newline at end of file
diff --git a/Engine/Source/Runtime/AIModule/Private/BehaviorTree/Tasks/BTTask_MakeNoise.cpp b/Engine/Source/Runtime/AIModule/Private/BehaviorTree/Tasks/BTTask_MakeNoise.cpp
--- a/Engine/Source/Runtime/AIModule/Private/BehaviorTree/Tasks/BTTask_MakeNoise.cpp
+++ b/Engine/Source/Runtime/AIModule/Private/BehaviorTree/Tasks/BTTask_MakeNoise.cpp
@@ -26,3 +26,11 @@ EBTNodeResult::Type UBTTask_MakeNoise::ExecuteTask(UBehaviorTreeComponent* Owner
 	return EBTNodeResult::Failed;
 }
 
+#if WITH_EDITOR
+
+FName UBTTask_MakeNoise::GetNodeIconName() const
+{
+	return FName("BTEditor.Graph.BTNode.Task.MakeNoise.Icon");
+}
+
+#endif	// WITH_EDITOR
\ No newline at end of file
diff --git a/Engine/Source/Runtime/AIModule/Private/BehaviorTree/Tasks/BTTask_MoveDirectlyToward.cpp b/Engine/Source/Runtime/AIModule/Private/BehaviorTree/Tasks/BTTask_MoveDirectlyToward.cpp
--- a/Engine/Source/Runtime/AIModule/Private/BehaviorTree/Tasks/BTTask_MoveDirectlyToward.cpp
+++ b/Engine/Source/Runtime/AIModule/Private/BehaviorTree/Tasks/BTTask_MoveDirectlyToward.cpp
@@ -108,3 +108,12 @@ uint16 UBTTask_MoveDirectlyToward::GetInstanceMemorySize() const
 {
 	return sizeof(FBTMoveDirectlyTowardMemory);
 }
+
+#if WITH_EDITOR
+
+FName UBTTask_MoveDirectlyToward::GetNodeIconName() const
+{
+	return FName("BTEditor.Graph.BTNode.Task.MoveDirectlyToward.Icon");
+}
+
+#endif	// WITH_EDITOR
\ No newline at end of file
diff --git a/Engine/Source/Runtime/AIModule/Private/BehaviorTree/Tasks/BTTask_MoveTo.cpp b/Engine/Source/Runtime/AIModule/Private/BehaviorTree/Tasks/BTTask_MoveTo.cpp
--- a/Engine/Source/Runtime/AIModule/Private/BehaviorTree/Tasks/BTTask_MoveTo.cpp
+++ b/Engine/Source/Runtime/AIModule/Private/BehaviorTree/Tasks/BTTask_MoveTo.cpp
@@ -164,3 +164,12 @@ uint16 UBTTask_MoveTo::GetInstanceMemorySize() const
 {
 	return sizeof(FBTMoveToTaskMemory);
 }
+
+#if WITH_EDITOR
+
+FName UBTTask_MoveTo::GetNodeIconName() const
+{
+	return FName("BTEditor.Graph.BTNode.Task.MoveTo.Icon");
+}
+
+#endif	// WITH_EDITOR
\ No newline at end of file
diff --git a/Engine/Source/Runtime/AIModule/Private/BehaviorTree/Tasks/BTTask_PlaySound.cpp b/Engine/Source/Runtime/AIModule/Private/BehaviorTree/Tasks/BTTask_PlaySound.cpp
--- a/Engine/Source/Runtime/AIModule/Private/BehaviorTree/Tasks/BTTask_PlaySound.cpp
+++ b/Engine/Source/Runtime/AIModule/Private/BehaviorTree/Tasks/BTTask_PlaySound.cpp
@@ -30,3 +30,12 @@ FString UBTTask_PlaySound::GetStaticDescription() const
 {
 	return FString::Printf(TEXT("%s: '%s'"), *Super::GetStaticDescription(), SoundToPlay ? *SoundToPlay->GetName() : TEXT(""));
 }
+
+#if WITH_EDITOR
+
+FName UBTTask_PlaySound::GetNodeIconName() const
+{
+	return FName("BTEditor.Graph.BTNode.Task.PlaySound.Icon");
+}
+
+#endif	// WITH_EDITOR
\ No newline at end of file
diff --git a/Engine/Source/Runtime/AIModule/Private/BehaviorTree/Tasks/BTTask_RunBehavior.cpp b/Engine/Source/Runtime/AIModule/Private/BehaviorTree/Tasks/BTTask_RunBehavior.cpp
--- a/Engine/Source/Runtime/AIModule/Private/BehaviorTree/Tasks/BTTask_RunBehavior.cpp
+++ b/Engine/Source/Runtime/AIModule/Private/BehaviorTree/Tasks/BTTask_RunBehavior.cpp
@@ -224,3 +224,12 @@ void UBTTask_RunBehavior::InjectNodes(UBehaviorTreeComponent* OwnerComp, uint8*
 		}
 	}
 }
+
+#if WITH_EDITOR
+
+FName UBTTask_RunBehavior::GetNodeIconName() const
+{
+	return FName("BTEditor.Graph.BTNode.Task.RunBehavior.Icon");
+}
+
+#endif	// WITH_EDITOR
\ No newline at end of file
diff --git a/Engine/Source/Runtime/AIModule/Private/BehaviorTree/Tasks/BTTask_RunEQSQuery.cpp b/Engine/Source/Runtime/AIModule/Private/BehaviorTree/Tasks/BTTask_RunEQSQuery.cpp
--- a/Engine/Source/Runtime/AIModule/Private/BehaviorTree/Tasks/BTTask_RunEQSQuery.cpp
+++ b/Engine/Source/Runtime/AIModule/Private/BehaviorTree/Tasks/BTTask_RunEQSQuery.cpp
@@ -141,4 +141,10 @@ void UBTTask_RunEQSQuery::PostEditChangeProperty(struct FPropertyChangedEvent& P
 		}
 	}
 }
+
+FName UBTTask_RunEQSQuery::GetNodeIconName() const
+{
+	return FName("BTEditor.Graph.BTNode.Task.RunEQSQuery.Icon");
+}
+
 #endif
diff --git a/Engine/Source/Runtime/AIModule/Private/BehaviorTree/Tasks/BTTask_Wait.cpp b/Engine/Source/Runtime/AIModule/Private/BehaviorTree/Tasks/BTTask_Wait.cpp
--- a/Engine/Source/Runtime/AIModule/Private/BehaviorTree/Tasks/BTTask_Wait.cpp
+++ b/Engine/Source/Runtime/AIModule/Private/BehaviorTree/Tasks/BTTask_Wait.cpp
@@ -50,3 +50,12 @@ uint16 UBTTask_Wait::GetInstanceMemorySize() const
 {
 	return sizeof(FBTWaitTaskMemory);
 }
+
+#if WITH_EDITOR
+
+FName UBTTask_Wait::GetNodeIconName() const
+{
+	return FName("BTEditor.Graph.BTNode.Task.Wait.Icon");
+}
+
+#endif	// WITH_EDITOR
\ No newline at end of file
