Altered hyperlink style in the details panel to make it more obvious that it is clickable.
#ttp 324844 - Details/Outliner: Make it more obvious how to get to Blueprints or Header files from actors
#branch UE4
#proj Editor.EditorStyle, Editor.PropertyEditor
#add Added new BlueHyperlink style to EditorStyle.
#change Changed font size in hyperlink text styles EditorStyle.
#change Changed HyperlinkStyle used by SDetailNameArea to provide an underlined hyperlink even when not hovered.

diff --git a/Engine/Source/Editor/EditorStyle/Private/SlateEditorStyle.cpp b/Engine/Source/Editor/EditorStyle/Private/SlateEditorStyle.cpp
--- a/Engine/Source/Editor/EditorStyle/Private/SlateEditorStyle.cpp
+++ b/Engine/Source/Editor/EditorStyle/Private/SlateEditorStyle.cpp
@@ -554,9 +554,20 @@ void FSlateEditorStyle::FStyle::SetupGeneralStyles()
 			.SetPadding(FMargin(0.0f));
 		Set("DarkHyperlink", DarkHyperlink);
 
+		// Edit BP Hyperlink
+		FButtonStyle EditBPHyperlinkButton = FButtonStyle()
+			.SetNormal(BORDER_BRUSH("Old/HyperlinkDotted", FMargin(0, 0, 0, 3 / 16.0f), FLinearColor(0.25f, 0.5f, 1.0f)))
+			.SetPressed(FSlateNoResource())
+			.SetHovered(BORDER_BRUSH("Old/HyperlinkUnderline", FMargin(0, 0, 0, 3 / 16.0f), FLinearColor(0.25f, 0.5f, 1.0f)));
+		FHyperlinkStyle EditBPHyperlink = FHyperlinkStyle()
+			.SetUnderlineStyle(EditBPHyperlinkButton)
+			.SetTextStyle(NormalText)
+			.SetPadding(FMargin(0.0f));
+		Set("EditBPHyperlink", EditBPHyperlink);
+
 		// Visible on hover hyper link
 		FButtonStyle HoverOnlyHyperlinkButton = FButtonStyle()
-			.SetNormal (FSlateNoResource() )
+			.SetNormal(FSlateNoResource() )
 			.SetPressed(FSlateNoResource() )
 			.SetHovered(BORDER_BRUSH( "Old/HyperlinkUnderline", FMargin(0,0,0,3/16.0f) ) );
 		FHyperlinkStyle HoverOnlyHyperlink = FHyperlinkStyle()
@@ -2464,8 +2475,8 @@ void FSlateEditorStyle::FStyle::SetupPropertyEditorStyles()
 		Set( "DetailsView.AdvancedDropdownBorder.Open", new IMAGE_BRUSH( "Common/ScrollBoxShadowTop", FVector2D(64,8) ) );
 		Set( "DetailsView.CategoryFontStyle", TTF_CORE_FONT( "Fonts/Roboto-Bold", 10 ) );
 
-		Set( "DetailsView.EditBlueprintHyperlinkStyle", FTextBlockStyle(NormalText) .SetFont( TTF_CORE_FONT( "Fonts/Roboto-Regular", 8 ) ).SetColorAndOpacity( FLinearColor( 0.25f, 0.5f, 1.0f ) ) );
-		Set( "DetailsView.GoToCodeHyperlinkStyle", FTextBlockStyle(NormalText) .SetFont( TTF_CORE_FONT( "Fonts/Roboto-Regular", 8 ) ).SetColorAndOpacity( FLinearColor( 0.7f, 0.7f, 0.7f ) ) );
+		Set( "DetailsView.EditBlueprintHyperlinkStyle", FTextBlockStyle(NormalText) .SetFont( TTF_CORE_FONT( "Fonts/Roboto-Regular", 10 ) ).SetColorAndOpacity( FLinearColor( 0.25f, 0.5f, 1.0f ) ) );
+		Set( "DetailsView.GoToCodeHyperlinkStyle", FTextBlockStyle(NormalText) .SetFont( TTF_CORE_FONT( "Fonts/Roboto-Regular", 10 ) ).SetColorAndOpacity( FLinearColor( 0.7f, 0.7f, 0.7f ) ) );
 
 		Set( "DetailsView.TreeView.TableRow", FTableRowStyle()
 			.SetEvenRowBackgroundBrush( FSlateNoResource() )
diff --git a/Engine/Source/Editor/PropertyEditor/Private/SDetailNameArea.cpp b/Engine/Source/Editor/PropertyEditor/Private/SDetailNameArea.cpp
--- a/Engine/Source/Editor/PropertyEditor/Private/SDetailNameArea.cpp
+++ b/Engine/Source/Editor/PropertyEditor/Private/SDetailNameArea.cpp
@@ -196,10 +196,10 @@ void SDetailNameArea::BuildObjectNameAreaSelectionLabel( TSharedRef< SHorizontal
 					.FillWidth( 1.0f )
 					[
 						SNew(SHyperlink)
-							.Style(FEditorStyle::Get(), "HoverOnlyHyperlink")
+							.Style(FEditorStyle::Get(), "EditBPHyperlink")
 							.TextStyle(FEditorStyle::Get(), "DetailsView.EditBlueprintHyperlinkStyle")
 							.OnNavigate(this, &SDetailNameArea::OnEditBlueprintClicked, BlueprintPtr, ObjectWeakPtr)
-							.Text(FText::Format(LOCTEXT("EditBlueprint", "(Edit {0})"), FText::FromString( Blueprint->GetName() ) ))
+							.Text(FText::Format(LOCTEXT("EditBlueprint", "Edit {0}"), FText::FromString( Blueprint->GetName() ) ))
 							.ToolTipText(LOCTEXT("EditBlueprint_ToolTip", "Click to edit the blueprint"))
 					];
 			}		
@@ -224,11 +224,11 @@ void SDetailNameArea::BuildObjectNameAreaSelectionLabel( TSharedRef< SHorizontal
 						.FillWidth( 1.0f )
 						[
 							SNew(SHyperlink)
-								.Style(FEditorStyle::Get(), "HoverOnlyHyperlink")
+								.Style(FCoreStyle::Get(), "Hyperlink")
 								.TextStyle(FEditorStyle::Get(), "DetailsView.GoToCodeHyperlinkStyle")
 								.OnNavigate_Static(&Local::OnEditCodeClicked, ClassHeaderPath)
 //								.Text(LOCTEXT("GoToCode", "(C++)" ))
-								.Text(FText::Format(LOCTEXT("GoToCode", "({0})" ), FText::FromString(FPaths::GetCleanFilename( *ClassHeaderPath ) ) ) )
+								.Text(FText::Format(LOCTEXT("GoToCode", "{0}" ), FText::FromString(FPaths::GetCleanFilename( *ClassHeaderPath ) ) ) )
 								.ToolTipText(LOCTEXT("GoToCode_ToolTip", "Click to open this source file in a text editor"))
 						];
 				}
