Paper2D: Hook up some asset and actor icons

diff --git a/Engine/Source/Editor/EditorStyle/Private/SlateEditorStyle.cpp b/Engine/Source/Editor/EditorStyle/Private/SlateEditorStyle.cpp
--- a/Engine/Source/Editor/EditorStyle/Private/SlateEditorStyle.cpp
+++ b/Engine/Source/Editor/EditorStyle/Private/SlateEditorStyle.cpp
@@ -4055,69 +4055,93 @@ void FSlateEditorStyle::FStyle::SetupClassIconsAndThumbnails()
 #if WITH_EDITOR
 	// Actor Classes Outliner
 	{
-		Set( "ClassIcon.Emitter", new IMAGE_BRUSH( "Icons/ActorIcons/Emitter_16x", Icon16x16 ) );
-		Set( "ClassIcon.Light", new IMAGE_BRUSH( "Icons/ActorIcons/LightActor_16x", Icon16x16 ) );
-		Set( "ClassIcon.MatineeActor", new IMAGE_BRUSH( "Icons/ActorIcons/MatineeActor_16x", Icon16x16 ) );
-		Set( "ClassIcon.Brush", new IMAGE_BRUSH( "Icons/ActorIcons/Brush_16x", Icon16x16 ) );
-		Set( "ClassIcon.BrushAdditive", new IMAGE_BRUSH( "Icons/ActorIcons/Brush_Add_16x", Icon16x16 ) );
-		Set( "ClassIcon.BrushSubtractive", new IMAGE_BRUSH( "Icons/ActorIcons/Brush_Subtract_16x", Icon16x16 ) );
-		Set( "ClassIcon.Volume", new IMAGE_BRUSH( "Icons/ActorIcons/Volume_16x", Icon16x16 ) );
-		Set( "ClassIcon.SkeletalPhysicsActor", new IMAGE_BRUSH( "Icons/ActorIcons/SkeletalPhysicsActor_16x", Icon16x16 ) );
-		Set( "ClassIcon.GroupActor", new IMAGE_BRUSH( "Icons/ActorIcons/GroupActor_16x", Icon16x16 ) );
-		Set( "ClassIcon.VectorFieldVolume", new IMAGE_BRUSH( "Icons/ActorIcons/VectorFieldVolume_16x", Icon16x16 ) );
-		Set( "ClassIcon.Deleted", new IMAGE_BRUSH( "Icons/ActorIcons/DeletedActor_16px", Icon16x16 ) );
-		Set( "ClassIcon.DestructibleActor", new IMAGE_BRUSH( "Icons/ActorIcons/Destructable_16px", Icon16x16 ) );
-		Set( "ClassIcon.DefaultPawn", new IMAGE_BRUSH( "Icons/ActorIcons/DefaultPawn_16px", Icon16x16 ) );
-		Set( "ClassIcon.Pawn", new IMAGE_BRUSH( "Icons/ActorIcons/Pawn_16px", Icon16x16 ) );
-		Set( "ClassIcon.Note", new IMAGE_BRUSH( "Icons/ActorIcons/Note_16px", Icon16x16 ) );
-		Set( "ClassIcon.Character", new IMAGE_BRUSH( "Icons/ActorIcons/Character_16px", Icon16x16 ) );
+		Set("ClassIcon.Emitter", new IMAGE_BRUSH("Icons/ActorIcons/Emitter_16x", Icon16x16));
+		Set("ClassIcon.Light", new IMAGE_BRUSH("Icons/ActorIcons/LightActor_16x", Icon16x16));
+		Set("ClassIcon.MatineeActor", new IMAGE_BRUSH("Icons/ActorIcons/MatineeActor_16x", Icon16x16));
+		Set("ClassIcon.Brush", new IMAGE_BRUSH("Icons/ActorIcons/Brush_16x", Icon16x16));
+		Set("ClassIcon.BrushAdditive", new IMAGE_BRUSH("Icons/ActorIcons/Brush_Add_16x", Icon16x16));
+		Set("ClassIcon.BrushSubtractive", new IMAGE_BRUSH("Icons/ActorIcons/Brush_Subtract_16x", Icon16x16));
+		Set("ClassIcon.Volume", new IMAGE_BRUSH("Icons/ActorIcons/Volume_16x", Icon16x16));
+		Set("ClassIcon.SkeletalPhysicsActor", new IMAGE_BRUSH("Icons/ActorIcons/SkeletalPhysicsActor_16x", Icon16x16));
+		Set("ClassIcon.GroupActor", new IMAGE_BRUSH("Icons/ActorIcons/GroupActor_16x", Icon16x16));
+		Set("ClassIcon.VectorFieldVolume", new IMAGE_BRUSH("Icons/ActorIcons/VectorFieldVolume_16x", Icon16x16));
+		Set("ClassIcon.Deleted", new IMAGE_BRUSH("Icons/ActorIcons/DeletedActor_16px", Icon16x16));
+		Set("ClassIcon.DestructibleActor", new IMAGE_BRUSH("Icons/ActorIcons/Destructable_16px", Icon16x16));
+		Set("ClassIcon.DefaultPawn", new IMAGE_BRUSH("Icons/ActorIcons/DefaultPawn_16px", Icon16x16));
+		Set("ClassIcon.Pawn", new IMAGE_BRUSH("Icons/ActorIcons/Pawn_16px", Icon16x16));
+		Set("ClassIcon.Note", new IMAGE_BRUSH("Icons/ActorIcons/Note_16px", Icon16x16));
+		Set("ClassIcon.Character", new IMAGE_BRUSH("Icons/ActorIcons/Character_16px", Icon16x16));
 
 		Set("ClassIcon.BlueprintActorBase", new IMAGE_BRUSH("Icons/AssetIcons/Blueprint_16x", Icon16x16));
 		Set("ClassIcon.StaticMeshActor", new IMAGE_BRUSH("Icons/AssetIcons/StaticMesh_16x", Icon16x16));
 		Set("ClassIcon.SkeletalMeshActor", new IMAGE_BRUSH("Icons/AssetIcons/SkeletalMesh_16x", Icon16x16));
 
 		// Component classes
-		Set( "ClassIcon.AudioComponent", new IMAGE_BRUSH( "Icons/ActorIcons/SoundActor_16x", Icon16x16 ) );
-		Set( "ClassIcon.CameraComponent", new IMAGE_BRUSH( "Icons/AssetIcons/CameraActor_16x", Icon16x16 ) );
-		Set( "ClassIcon.BlueprintCore", new IMAGE_BRUSH( "Icons/AssetIcons/Blueprint_16x", Icon16x16 ) );
-		Set( "ClassIcon.BrushComponent", new IMAGE_BRUSH( "Icons/ActorIcons/Brush_16x", Icon16x16 ) );
-		Set( "ClassIcon.DecalComponent", new IMAGE_BRUSH( "Icons/AssetIcons/DecalActor_16x", Icon16x16 ) );
-		Set( "ClassIcon.DirectionalLightComponent", new IMAGE_BRUSH( "Icons/AssetIcons/DirectionalLight_16x", Icon16x16 ) );
-		Set( "ClassIcon.ExponentialHeightFogComponent", new IMAGE_BRUSH( "Icons/AssetIcons/ExponentialHeightFog_16x", Icon16x16 ) );
-		Set( "ClassIcon.LandscapeComponent", new IMAGE_BRUSH( "Icons/AssetIcons/Landscape_16x", Icon16x16 ) );
-		Set( "ClassIcon.LightComponent", new IMAGE_BRUSH( "Icons/ActorIcons/LightActor_16x", Icon16x16 ) );
-		Set( "ClassIcon.ParticleSystemComponent", new IMAGE_BRUSH( "Icons/AssetIcons/ParticleSystem_16x", Icon16x16 ) );
-		Set( "ClassIcon.PointLightComponent", new IMAGE_BRUSH( "Icons/AssetIcons/PointLight_16x", Icon16x16 ) );
-		Set( "ClassIcon.RB_RadialForceComponent", new IMAGE_BRUSH( "Icons/AssetIcons/RadialForceActor_16x", Icon16x16 ) );
-		Set( "ClassIcon.SingleAnimSkeletalComponent", new IMAGE_BRUSH( "Icons/AssetIcons/SkeletalMesh_16x", Icon16x16 ) );		
-		Set( "ClassIcon.SkeletalMeshComponent", new IMAGE_BRUSH( "Icons/AssetIcons/SkeletalMesh_16x", Icon16x16 ) );		
-		Set( "ClassIcon.SpotLightComponent", new IMAGE_BRUSH( "Icons/AssetIcons/SpotLight_16x", Icon16x16 ) );
-		Set( "ClassIcon.StaticMeshComponent", new IMAGE_BRUSH( "Icons/AssetIcons/StaticMesh_16x", Icon16x16 ) );
-		Set( "ClassIcon.VectorFieldComponent", new IMAGE_BRUSH( "Icons/ActorIcons/VectorFieldVolume_16x", Icon16x16 ) );
-		Set( "ClassIcon.ArrowComponent", new IMAGE_BRUSH( "Icons/ActorIcons/Arrow_16px", Icon16x16 ) );
-		Set( "ClassIcon.AtmosphericFogComponent", new IMAGE_BRUSH( "Icons/AssetIcons/AtmosphericFog_16x", Icon16x16 ) );
-		Set( "ClassIcon.BoxComponent", new IMAGE_BRUSH( "Icons/ActorIcons/Box_16px", Icon16x16 ) );
-		Set( "ClassIcon.CapsuleComponent", new IMAGE_BRUSH( "Icons/ActorIcons/Capsule_16px", Icon16x16 ) );
-		Set( "ClassIcon.DestructibleComponent", new IMAGE_BRUSH( "Icons/ActorIcons/Destructable_16px", Icon16x16 ) );
-		Set( "ClassIcon.InstancedStaticMeshComponent", new IMAGE_BRUSH( "Icons/ActorIcons/InstancedStaticMesh_16px", Icon16x16 ) );
-		Set( "ClassIcon.MaterialBillboardComponent", new IMAGE_BRUSH( "Icons/ActorIcons/MaterialSprite_16px", Icon16x16 ) );
-		Set( "ClassIcon.SceneCaptureComponent2D", new IMAGE_BRUSH( "Icons/AssetIcons/SceneCapture2D_16x", Icon16x16 ) );
-		Set( "ClassIcon.SceneCaptureComponent", new IMAGE_BRUSH( "Icons/ActorIcons/SceneCapture_16px", Icon16x16 ) );
-		Set( "ClassIcon.SceneComponent", new IMAGE_BRUSH( "Icons/ActorIcons/Scene_16px", Icon16x16 ) );
-		Set( "ClassIcon.SphereComponent", new IMAGE_BRUSH("Icons/ActorIcons/Sphere_16px", Icon16x16));
-		Set( "ClassIcon.SplineComponent", new IMAGE_BRUSH("Icons/ActorIcons/Spline_16px", Icon16x16));
-		Set( "ClassIcon.BillboardComponent", new IMAGE_BRUSH( "Icons/ActorIcons/SpriteComponent_16px", Icon16x16 ) );
-		Set( "ClassIcon.TextRenderComponent", new IMAGE_BRUSH( "Icons/AssetIcons/TextRenderActor_16x", Icon16x16 ) );
-		Set( "ClassIcon.TimelineComponent", new IMAGE_BRUSH("Icons/ActorIcons/TimelineComponent_16px", Icon16x16));
-		Set( "ClassIcon.ChildActorComponent", new IMAGE_BRUSH( "Icons/ActorIcons/ChildActorComponent_16px", Icon16x16 ) );
-		Set( "ClassIcon.ComponentMobilityStaticPip", new IMAGE_BRUSH( "Icons/ActorIcons/ComponentMobilityStationary_7x16px", Icon7x16, FLinearColor(0.f, 0.f, 0.f, 0.f) ) );
-		Set( "ClassIcon.ComponentMobilityStationaryPip", new IMAGE_BRUSH( "Icons/ActorIcons/ComponentMobilityStationary_7x16px", Icon7x16 ) );
-		Set( "ClassIcon.ComponentMobilityMovablePip", new IMAGE_BRUSH( "Icons/ActorIcons/ComponentMobilityMovable_7x16px", Icon7x16 ) );
-		Set( "ClassIcon.MovableMobilityIcon", new IMAGE_BRUSH( "Icons/ActorIcons/Light_Movable_16x", Icon16x16));
-		Set( "ClassIcon.StationaryMobilityIcon", new IMAGE_BRUSH( "Icons/ActorIcons/Light_Adjustable_16x", Icon16x16));
-		Set( "ClassIcon.ComponentMobilityHeaderIcon", new IMAGE_BRUSH( "Icons/ActorIcons/ComponentMobilityHeader_7x16", Icon7x16));
-		
-		
+		Set("ClassIcon.AudioComponent", new IMAGE_BRUSH("Icons/ActorIcons/SoundActor_16x", Icon16x16));
+		Set("ClassIcon.CameraComponent", new IMAGE_BRUSH("Icons/AssetIcons/CameraActor_16x", Icon16x16));
+		Set("ClassIcon.BlueprintCore", new IMAGE_BRUSH("Icons/AssetIcons/Blueprint_16x", Icon16x16));
+		Set("ClassIcon.BrushComponent", new IMAGE_BRUSH("Icons/ActorIcons/Brush_16x", Icon16x16));
+		Set("ClassIcon.DecalComponent", new IMAGE_BRUSH("Icons/AssetIcons/DecalActor_16x", Icon16x16));
+		Set("ClassIcon.DirectionalLightComponent", new IMAGE_BRUSH("Icons/AssetIcons/DirectionalLight_16x", Icon16x16));
+		Set("ClassIcon.ExponentialHeightFogComponent", new IMAGE_BRUSH("Icons/AssetIcons/ExponentialHeightFog_16x", Icon16x16));
+		Set("ClassIcon.LandscapeComponent", new IMAGE_BRUSH("Icons/AssetIcons/Landscape_16x", Icon16x16));
+		Set("ClassIcon.LightComponent", new IMAGE_BRUSH("Icons/ActorIcons/LightActor_16x", Icon16x16));
+		Set("ClassIcon.ParticleSystemComponent", new IMAGE_BRUSH("Icons/AssetIcons/ParticleSystem_16x", Icon16x16));
+		Set("ClassIcon.PointLightComponent", new IMAGE_BRUSH("Icons/AssetIcons/PointLight_16x", Icon16x16));
+		Set("ClassIcon.RB_RadialForceComponent", new IMAGE_BRUSH("Icons/AssetIcons/RadialForceActor_16x", Icon16x16));
+		Set("ClassIcon.SingleAnimSkeletalComponent", new IMAGE_BRUSH("Icons/AssetIcons/SkeletalMesh_16x", Icon16x16));
+		Set("ClassIcon.SkeletalMeshComponent", new IMAGE_BRUSH("Icons/AssetIcons/SkeletalMesh_16x", Icon16x16));
+		Set("ClassIcon.SpotLightComponent", new IMAGE_BRUSH("Icons/AssetIcons/SpotLight_16x", Icon16x16));
+		Set("ClassIcon.StaticMeshComponent", new IMAGE_BRUSH("Icons/AssetIcons/StaticMesh_16x", Icon16x16));
+		Set("ClassIcon.VectorFieldComponent", new IMAGE_BRUSH("Icons/ActorIcons/VectorFieldVolume_16x", Icon16x16));
+		Set("ClassIcon.ArrowComponent", new IMAGE_BRUSH("Icons/ActorIcons/Arrow_16px", Icon16x16));
+		Set("ClassIcon.AtmosphericFogComponent", new IMAGE_BRUSH("Icons/AssetIcons/AtmosphericFog_16x", Icon16x16));
+		Set("ClassIcon.BoxComponent", new IMAGE_BRUSH("Icons/ActorIcons/Box_16px", Icon16x16));
+		Set("ClassIcon.CapsuleComponent", new IMAGE_BRUSH("Icons/ActorIcons/Capsule_16px", Icon16x16));
+		Set("ClassIcon.DestructibleComponent", new IMAGE_BRUSH("Icons/ActorIcons/Destructable_16px", Icon16x16));
+		Set("ClassIcon.InstancedStaticMeshComponent", new IMAGE_BRUSH("Icons/ActorIcons/InstancedStaticMesh_16px", Icon16x16));
+		Set("ClassIcon.MaterialBillboardComponent", new IMAGE_BRUSH("Icons/ActorIcons/MaterialSprite_16px", Icon16x16));
+		Set("ClassIcon.SceneCaptureComponent2D", new IMAGE_BRUSH("Icons/AssetIcons/SceneCapture2D_16x", Icon16x16));
+		Set("ClassIcon.SceneCaptureComponent", new IMAGE_BRUSH("Icons/ActorIcons/SceneCapture_16px", Icon16x16));
+		Set("ClassIcon.SceneComponent", new IMAGE_BRUSH("Icons/ActorIcons/Scene_16px", Icon16x16));
+		Set("ClassIcon.SphereComponent", new IMAGE_BRUSH("Icons/ActorIcons/Sphere_16px", Icon16x16));
+		Set("ClassIcon.SplineComponent", new IMAGE_BRUSH("Icons/ActorIcons/Spline_16px", Icon16x16));
+		Set("ClassIcon.BillboardComponent", new IMAGE_BRUSH("Icons/ActorIcons/SpriteComponent_16px", Icon16x16));
+		Set("ClassIcon.TextRenderComponent", new IMAGE_BRUSH("Icons/AssetIcons/TextRenderActor_16x", Icon16x16));
+		Set("ClassIcon.TimelineComponent", new IMAGE_BRUSH("Icons/ActorIcons/TimelineComponent_16px", Icon16x16));
+		Set("ClassIcon.ChildActorComponent", new IMAGE_BRUSH("Icons/ActorIcons/ChildActorComponent_16px", Icon16x16));
+		Set("ClassIcon.ComponentMobilityStaticPip", new IMAGE_BRUSH("Icons/ActorIcons/ComponentMobilityStationary_7x16px", Icon7x16, FLinearColor(0.f, 0.f, 0.f, 0.f)));
+		Set("ClassIcon.ComponentMobilityStationaryPip", new IMAGE_BRUSH("Icons/ActorIcons/ComponentMobilityStationary_7x16px", Icon7x16));
+		Set("ClassIcon.ComponentMobilityMovablePip", new IMAGE_BRUSH("Icons/ActorIcons/ComponentMobilityMovable_7x16px", Icon7x16));
+		Set("ClassIcon.MovableMobilityIcon", new IMAGE_BRUSH("Icons/ActorIcons/Light_Movable_16x", Icon16x16));
+		Set("ClassIcon.StationaryMobilityIcon", new IMAGE_BRUSH("Icons/ActorIcons/Light_Adjustable_16x", Icon16x16));
+		Set("ClassIcon.ComponentMobilityHeaderIcon", new IMAGE_BRUSH("Icons/ActorIcons/ComponentMobilityHeader_7x16", Icon7x16));
+
+		//@TODO: PAPER2D: Defined here until it is possible to define these in a plugin
+		{
+			Set("ClassIcon.PaperTerrainActor", new IMAGE_BRUSH("Icons/ActorIcons/Paper2DTerrainActor_16x", Icon16x16));
+			Set("ClassIcon.PaperTerrainComponent", new IMAGE_BRUSH("Icons/ActorIcons/Paper2DTerrainActor_16x", Icon16x16));
+
+			Set("ClassIcon.PaperTileMapActor", new IMAGE_BRUSH("Icons/ActorIcons/Paper2DTileMapActor_16x", Icon16x16));
+
+			Set("ClassIcon.PaperRenderActor", new IMAGE_BRUSH("Icons/AssetIcons/Paper2DSprite_16x", Icon16x16));
+			Set("ClassIcon.PaperRenderComponent", new IMAGE_BRUSH("Icons/AssetIcons/Paper2DSprite_16x", Icon16x16));
+			Set("ClassIcon.PaperSprite", new IMAGE_BRUSH("Icons/AssetIcons/Paper2DSprite_16x", Icon16x16));
+			Set("ClassThumbnail.PaperSprite", new IMAGE_BRUSH("Icons/AssetIcons/Paper2DSprite_64x", Icon64x64));
+
+			Set("ClassIcon.PaperFlipbookActor", new IMAGE_BRUSH("Icons/AssetIcons/Paper2DSpriteFlipbook16x", Icon16x16));
+			Set("ClassIcon.PaperAnimatedRenderComponent", new IMAGE_BRUSH("Icons/AssetIcons/Paper2DSpriteFlipbook16x", Icon16x16));
+			Set("ClassIcon.PaperFlipbook", new IMAGE_BRUSH("Icons/AssetIcons/Paper2DSpriteFlipbook16x", Icon16x16));
+			Set("ClassThumbnail.PaperFlipbook", new IMAGE_BRUSH("Icons/AssetIcons/Paper2DSpriteFlipbook64x", Icon64x64));
+
+			Set("ClassIcon.PaperSpriteAtlas", new IMAGE_BRUSH("Icons/AssetIcons/Paper2DSpriteAtlasGroup_16x", Icon16x16));
+			Set("ClassThumbnail.PaperSpriteAtlas", new IMAGE_BRUSH("Icons/AssetIcons/Paper2DSpriteAtlasGroup_64x", Icon64x64));
+
+			// ClassIcon and ClassThumbnail for tile sets and/or terrain materials?  Might not be needed/feasible due to them being a UDataAsset
+			//Set("ClassIcon.PaperTerrainMaterial", new IMAGE_BRUSH("Icons/AssetIcons/Paper2DSpriteAtlasGroup_16x", Icon16x16));
+			//Set("ClassThumbnail.PaperTerrainMaterial", new IMAGE_BRUSH("Icons/AssetIcons/Paper2DSpriteAtlasGroup_64x", Icon64x64));
+		}
+
 		// Factory classes
 		Set( "ClassIcon.ActorFactoryBoxVolume", new IMAGE_BRUSH( "Icons/icon_volume_Box_16x", Icon16x16 ) );
 		Set( "ClassIcon.ActorFactoryCylinderVolume", new IMAGE_BRUSH( "Icons/icon_volume_cylinder_16x", Icon16x16 ) );
