Added Wind Directional Source to editor style.

An icon exists for this actor type, but it was not hooked up in the style set.

diff --git a/Engine/Source/Editor/EditorStyle/Private/SlateEditorStyle.cpp b/Engine/Source/Editor/EditorStyle/Private/SlateEditorStyle.cpp
--- a/Engine/Source/Editor/EditorStyle/Private/SlateEditorStyle.cpp
+++ b/Engine/Source/Editor/EditorStyle/Private/SlateEditorStyle.cpp
@@ -4256,6 +4256,7 @@ void FSlateEditorStyle::FStyle::SetupClassIconsAndThumbnails()
 			TEXT("TouchInterface"),
 			TEXT("UserDefinedEnum"),
 			TEXT("UserDefinedStruct"),
+			TEXT("WindDirectionalSource"),
 		};
 
 		const TCHAR* Path = TEXT("Icons/AssetIcons/%s_%dx");
