TTP 329427
Adding a credits screen to the editor, all the the credits are pulled from Engine\\Content\\Editor\\Credits.rt.

diff --git a/Engine/Source/Editor/EditorStyle/Private/SlateEditorStyle.cpp b/Engine/Source/Editor/EditorStyle/Private/SlateEditorStyle.cpp
--- a/Engine/Source/Editor/EditorStyle/Private/SlateEditorStyle.cpp
+++ b/Engine/Source/Editor/EditorStyle/Private/SlateEditorStyle.cpp
@@ -1295,6 +1295,72 @@ void FSlateEditorStyle::FStyle::SetupGeneralStyles()
 	}
 #endif // WITH_EDITOR || IS_PROGRAM
 
+#if WITH_EDITOR
+	// Credits screen
+	{
+		Set("Credits.Button", FButtonStyle(NoBorder)
+			.SetNormal(FSlateNoResource())
+			.SetPressed(BOX_BRUSH("Common/RoundedSelection_16x", 4.0f / 16.0f, SelectionColor_Pressed))
+			.SetHovered(BOX_BRUSH("Common/RoundedSelection_16x", 4.0f / 16.0f, SelectionColor))
+			);
+
+		Set("Credits.Pause", new IMAGE_BRUSH("Icons/PauseCredits", Icon20x20));
+		Set("Credits.Play", new IMAGE_BRUSH("Icons/PlayCredits", Icon20x20));
+
+		FLinearColor EditorOrange = FLinearColor(0.728f, 0.364f, 0.003f);
+
+		FTextBlockStyle CreditsNormal = FTextBlockStyle(NormalText)
+			.SetFont(TTF_CORE_FONT("Fonts/Roboto-Normal", 16))
+			.SetShadowOffset(FVector2D::UnitVector);
+
+		Set("Credits.Normal", CreditsNormal);
+
+		Set("Credits.Strong", FTextBlockStyle(CreditsNormal)
+			.SetFont(TTF_CORE_FONT("Fonts/Roboto-Bold", 16))
+			.SetShadowOffset(FVector2D::UnitVector));
+
+		Set("Credits.H1", FTextBlockStyle(CreditsNormal)
+			.SetColorAndOpacity(EditorOrange)
+			.SetFont(TTF_CORE_FONT("Fonts/Roboto-Bold", 36))
+			.SetShadowOffset(FVector2D::UnitVector));
+
+		Set("Credits.H2", FTextBlockStyle(CreditsNormal)
+			.SetColorAndOpacity(EditorOrange)
+			.SetFont(TTF_CORE_FONT("Fonts/Roboto-Bold", 30))
+			.SetShadowOffset(FVector2D::UnitVector));
+
+		Set("Credits.H3", FTextBlockStyle(CreditsNormal)
+			.SetFont(TTF_CORE_FONT("Fonts/Roboto-Bold", 24))
+			.SetShadowOffset(FVector2D::UnitVector));
+
+		Set("Credits.H4", FTextBlockStyle(CreditsNormal)
+			.SetFont(TTF_CORE_FONT("Fonts/Roboto-Bold", 18))
+			.SetShadowOffset(FVector2D::UnitVector));
+
+		Set("Credits.H5", FTextBlockStyle(CreditsNormal)
+			.SetFont(TTF_CORE_FONT("Fonts/Roboto-Bold", 12))
+			.SetShadowOffset(FVector2D::UnitVector));
+
+		Set("Credits.H6", FTextBlockStyle(CreditsNormal)
+			.SetFont(TTF_CORE_FONT("Fonts/Roboto-Bold", 6))
+			.SetShadowOffset(FVector2D::UnitVector));
+
+		FTextBlockStyle LinkText = FTextBlockStyle(NormalText)
+			.SetColorAndOpacity(EditorOrange)
+			.SetShadowOffset(FVector2D::UnitVector);
+		FButtonStyle HoverOnlyHyperlinkButton = FButtonStyle()
+			.SetNormal(FSlateNoResource())
+			.SetPressed(FSlateNoResource())
+			.SetHovered(BORDER_BRUSH("Old/HyperlinkUnderline", FMargin(0, 0, 0, 3 / 16.0f)));
+		FHyperlinkStyle HoverOnlyHyperlink = FHyperlinkStyle()
+			.SetUnderlineStyle(HoverOnlyHyperlinkButton)
+			.SetTextStyle(LinkText)
+			.SetPadding(FMargin(0.0f));
+
+		Set("Credits.Hyperlink", HoverOnlyHyperlink);
+	}
+#endif // WITH_EDITOR || IS_PROGRAM
+
 	// New Level Dialog
 #if WITH_EDITOR || IS_PROGRAM
 	{
@@ -2428,7 +2494,8 @@ void FSlateEditorStyle::FStyle::SetupPropertyEditorStyles()
 		Set( "PropertyWindow.Locked", new IMAGE_BRUSH( "Icons/padlock_locked_16x", Icon16x16 ) );
 		Set( "PropertyWindow.Unlocked", new IMAGE_BRUSH( "Icons/padlock_unlocked_16x", Icon16x16 ) );
 		Set( "PropertyWindow.DiffersFromDefault", new IMAGE_BRUSH( "/PropertyView/DiffersFromDefault_8x8", FVector2D(8,8) ) ) ;
-		Set( "PropertyWindow.NormalFont", TTF_CORE_FONT( "Fonts/Roboto-Regular", 8 ) );
+		
+		Set( "PropertyWindow.NormalFont", TTF_CORE_FONT("Fonts/Roboto-Regular", 8) );
 		Set( "PropertyWindow.BoldFont", TTF_CORE_FONT( "Fonts/Roboto-Bold", 8 ) );
 		Set( "PropertyWindow.ItalicFont", TTF_FONT( "Fonts/Roboto-Italic", 8 ) );
 		Set( "PropertyWindow.FilterFont", TTF_CORE_FONT( "Fonts/Roboto-Regular", 10 ) );
@@ -3290,6 +3357,7 @@ void FSlateEditorStyle::FStyle::SetupLevelEditorStyle()
 		Set("MainFrame.VisitSupportWebSite", new IMAGE_BRUSH("Icons/Help/icon_Help_support_16x", Icon16x16));
 		Set("MainFrame.VisitEpicGamesDotCom", new IMAGE_BRUSH("Icons/Help/icon_Help_epic_16x", Icon16x16));
 		Set("MainFrame.AboutUnrealEd", new IMAGE_BRUSH("Icons/Help/icon_Help_unreal_16x", Icon16x16));
+		Set("MainFrame.CreditsUnrealEd", new IMAGE_BRUSH("Icons/Help/icon_Help_credits_16x", Icon16x16));
 
 		const FLinearColor IconColor = FLinearColor::Black;
 		Set( "EditorViewport.TranslateMode", new IMAGE_BRUSH( "Icons/icon_translateb_16x", Icon16x16 ) );
@@ -3641,7 +3709,7 @@ void FSlateEditorStyle::FStyle::SetupLevelEditorStyle()
 		const FString SmallRoundedButtonEnd(TEXT("Common/SmallRoundedToggleRight"));
 
 		const FLinearColor NormalColor(0.15, 0.15, 0.15, 1);
-		
+
 		Set("Property.ToggleButton", FCheckBoxStyle()
 			.SetCheckBoxType(ESlateCheckBoxType::ToggleButton)
 			.SetUncheckedImage(BOX_BRUSH(*SmallRoundedButton, FMargin(7.f / 16.f), NormalColor))
diff --git a/Engine/Source/Editor/MainFrame/Private/Frame/MainFrameActions.cpp b/Engine/Source/Editor/MainFrame/Private/Frame/MainFrameActions.cpp
--- a/Engine/Source/Editor/MainFrame/Private/Frame/MainFrameActions.cpp
+++ b/Engine/Source/Editor/MainFrame/Private/Frame/MainFrameActions.cpp
@@ -127,6 +127,9 @@ void FMainFrameCommands::RegisterCommands()
 	UI_COMMAND( AboutUnrealEd, "About Editor...", "Displays application credits and copyright information", EUserInterfaceActionType::Button, FInputGesture() );
 	ActionList->MapAction( AboutUnrealEd, FExecuteAction::CreateStatic( &FMainFrameActionCallbacks::AboutUnrealEd_Execute ) );
 
+	UI_COMMAND( CreditsUnrealEd, "Credits", "Displays application credits", EUserInterfaceActionType::Button, FInputGesture() );
+	ActionList->MapAction( CreditsUnrealEd, FExecuteAction::CreateStatic(&FMainFrameActionCallbacks::CreditsUnrealEd_Execute) );
+
 	UI_COMMAND( ResetLayout, "Reset Layout...", "Make a backup of your user settings and reset the layout customizations", EUserInterfaceActionType::Button, FInputGesture() );
 	ActionList->MapAction( ResetLayout, FExecuteAction::CreateStatic( &FMainFrameActionCallbacks::ResetLayout) );
 
@@ -928,6 +931,34 @@ void FMainFrameActionCallbacks::AboutUnrealEd_Execute()
 	}
 }
 
+void FMainFrameActionCallbacks::CreditsUnrealEd_Execute()
+{
+	const FText CreditsWindowTitle = LOCTEXT("CreditsUnrealEditor", "Credits");
+
+	TSharedPtr<SWindow> CreditsWindow =
+		SNew(SWindow)
+		.Title(CreditsWindowTitle)
+		.ClientSize(FVector2D(600.f, 700.f))
+		.SupportsMaximize(false)
+		.SupportsMinimize(false)
+		.SizingRule(ESizingRule::FixedSize)
+		[
+			SNew(SCreditsScreen)
+		];
+
+	IMainFrameModule& MainFrame = FModuleManager::LoadModuleChecked<IMainFrameModule>("MainFrame");
+	TSharedPtr<SWindow> ParentWindow = MainFrame.GetParentWindow();
+
+	if ( ParentWindow.IsValid() )
+	{
+		FSlateApplication::Get().AddModalWindow(CreditsWindow.ToSharedRef(), ParentWindow.ToSharedRef());
+	}
+	else
+	{
+		FSlateApplication::Get().AddWindow(CreditsWindow.ToSharedRef());
+	}
+}
+
 void FMainFrameActionCallbacks::OpenWidgetReflector_Execute()
 {
 	FGlobalTabmanager::Get()->InvokeTab(FTabId("WidgetReflector"));
diff --git a/Engine/Source/Editor/MainFrame/Private/Frame/MainFrameActions.h b/Engine/Source/Editor/MainFrame/Private/Frame/MainFrameActions.h
--- a/Engine/Source/Editor/MainFrame/Private/Frame/MainFrameActions.h
+++ b/Engine/Source/Editor/MainFrame/Private/Frame/MainFrameActions.h
@@ -52,6 +52,7 @@ public:
 	TSharedPtr< FUICommandInfo > VisitSupportWebSite;
 	TSharedPtr< FUICommandInfo > VisitEpicGamesDotCom;
 	TSharedPtr< FUICommandInfo > AboutUnrealEd;
+	TSharedPtr< FUICommandInfo > CreditsUnrealEd;
 
 	TSharedPtr< FUICommandInfo > ResetLayout;
 	TSharedPtr< FUICommandInfo > SaveLayout;
@@ -140,6 +141,8 @@ public:
 
 	static void AboutUnrealEd_Execute();
 
+	static void CreditsUnrealEd_Execute();
+
 	static void OpenWidgetReflector_Execute();
 
 	/** Opens the new project dialog */
diff --git a/Engine/Source/Editor/MainFrame/Private/MainFramePrivatePCH.h b/Engine/Source/Editor/MainFrame/Private/MainFramePrivatePCH.h
--- a/Engine/Source/Editor/MainFrame/Private/MainFramePrivatePCH.h
+++ b/Engine/Source/Editor/MainFrame/Private/MainFramePrivatePCH.h
@@ -10,6 +10,7 @@
  *****************************************************************************/
 
 #include "AboutScreen.h"
+#include "CreditsScreen.h"
 #include "AssetToolsModule.h"
 #include "CrashTracker.h"
 #include "DesktopPlatformModule.h"
diff --git a/Engine/Source/Editor/MainFrame/Private/Menus/MainMenu.cpp b/Engine/Source/Editor/MainFrame/Private/Menus/MainMenu.cpp
--- a/Engine/Source/Editor/MainFrame/Private/Menus/MainMenu.cpp
+++ b/Engine/Source/Editor/MainFrame/Private/Menus/MainMenu.cpp
@@ -230,6 +230,7 @@ void FMainMenu::FillHelpMenu( FMenuBuilder& MenuBuilder, const TSharedRef< FExte
 		const FText AboutWindowTitle = NSLOCTEXT("MainHelpMenu", "AboutUnrealEditor", "About Unreal Editor...");
 
 		MenuBuilder.AddMenuEntry(FMainFrameCommands::Get().AboutUnrealEd, "AboutUnrealEd", AboutWindowTitle);
+		MenuBuilder.AddMenuEntry(FMainFrameCommands::Get().CreditsUnrealEd);
 	}
 	MenuBuilder.EndSection();
 #endif
diff --git a/Engine/Source/Editor/UnrealEd/Private/CreditsScreen.cpp b/Engine/Source/Editor/UnrealEd/Private/CreditsScreen.cpp
new file mode 100644
--- /dev/null
+++ b/Engine/Source/Editor/UnrealEd/Private/CreditsScreen.cpp
@@ -0,0 +1,124 @@
+// Copyright 1998-2014 Epic Games, Inc. All Rights Reserved.
+
+
+#include "UnrealEd.h"
+#include "CreditsScreen.h"
+#include "IDocumentation.h"
+#include "EngineBuildSettings.h"
+
+#define LOCTEXT_NAMESPACE "CreditsScreen"
+
+void SCreditsScreen::Construct(const FArguments& InArgs)
+{
+	PreviousScrollPosition = 0.0f;
+	ScrollPixelsPerSecond = 50.0f;
+	bIsPlaying = true;
+
+	const FString Version = GEngineVersion.ToString(FEngineBuildSettings::IsPerforceBuild() ? EVersionComponent::Branch : EVersionComponent::Patch);
+
+	FString CreditsText;
+	FFileHelper::LoadFileToString(CreditsText, *( FPaths::EngineContentDir() + TEXT("Editor/Credits.rt") ));
+	CreditsText.ReplaceInline(TEXT("%VERSION%"), *Version);
+
+	ChildSlot
+	[
+		SNew(SOverlay)
+
+		+ SOverlay::Slot()
+		[
+			SAssignNew(ScrollBox, SScrollBox)
+			.Style( FEditorStyle::Get(), "ScrollBox" )
+			.OnUserScrolled(this, &SCreditsScreen::HandleUserScrolled)
+
+			+ SScrollBox::Slot()
+			[
+				SNew(SVerticalBox)
+
+				+ SVerticalBox::Slot()
+				.FillHeight(1.0f)
+				[
+					SNew(SRichTextBlock)
+					.Text(FText::FromString(CreditsText))
+					.TextStyle(FEditorStyle::Get(), "Credits.Normal")
+					.DecoratorStyleSet(&FEditorStyle::Get())
+					.AutoWrapText(true)
+					.Justification(ETextJustify::Center)
+					+ SRichTextBlock::HyperlinkDecorator(TEXT("browser"), this, &SCreditsScreen::OnBrowserLinkClicked)
+				]
+			]
+		]
+
+		+ SOverlay::Slot()
+		.VAlign(VAlign_Bottom)
+		[
+			SNew(SHorizontalBox)
+
+			+ SHorizontalBox::Slot()
+			.AutoWidth()
+			[
+				SNew(SButton)
+				.ButtonStyle(FEditorStyle::Get(), "Credits.Button")
+				.OnClicked(this, &SCreditsScreen::HandleTogglePlayPause)
+				[
+					SNew(SImage)
+					.Image(this, &SCreditsScreen::GetTogglePlayPauseBrush)
+				]
+			]
+		]
+	];
+}
+
+void SCreditsScreen::Tick(const FGeometry& AllottedGeometry, const double InCurrentTime, const float InDeltaTime)
+{
+	if ( bIsPlaying )
+	{
+		float NewPixelOffset = ( ScrollPixelsPerSecond * InDeltaTime );
+		ScrollBox->SetScrollOffset(ScrollBox->GetScrollOffset() + NewPixelOffset);
+		PreviousScrollPosition = ScrollBox->GetScrollOffset();
+	}
+}
+
+FReply SCreditsScreen::HandleTogglePlayPause()
+{
+	bIsPlaying = !bIsPlaying;
+
+	return FReply::Handled();
+}
+
+void SCreditsScreen::HandleUserScrolled(float ScrollOffset)
+{
+	// If the user scrolls up, and we're currently playing, then stop playing.
+	if ( bIsPlaying && ScrollOffset < PreviousScrollPosition )
+	{
+		bIsPlaying = false;
+	}
+
+	PreviousScrollPosition = ScrollOffset;
+}
+
+const FSlateBrush* SCreditsScreen::GetTogglePlayPauseBrush() const
+{
+	static FName PauseIcon(TEXT("Credits.Pause"));
+	static FName PlayIcon(TEXT("Credits.Play"));
+
+	if ( bIsPlaying )
+	{
+		return FEditorStyle::GetBrush(PauseIcon);
+	}
+	else
+	{
+		return FEditorStyle::GetBrush(PlayIcon);
+	}
+}
+
+void SCreditsScreen::OnBrowserLinkClicked(const FSlateHyperlinkRun::FMetadata& Metadata)
+{
+	const FString* url = Metadata.Find(TEXT("href"));
+
+	if ( url != NULL )
+	{
+		FPlatformProcess::LaunchURL(**url, NULL, NULL);
+	}
+}
+
+#undef LOCTEXT_NAMESPACE
\ No newline at end of file
diff --git a/Engine/Source/Editor/UnrealEd/Public/CreditsScreen.h b/Engine/Source/Editor/UnrealEd/Public/CreditsScreen.h
new file mode 100644
--- /dev/null
+++ b/Engine/Source/Editor/UnrealEd/Public/CreditsScreen.h
@@ -0,0 +1,48 @@
+// Copyright 1998-2014 Epic Games, Inc. All Rights Reserved.
+
+#pragma once
+
+/**
+ * Credit screen widget that displays a scrolling list contributors.  
+ */
+class SCreditsScreen : public SCompoundWidget
+{
+public:
+	SLATE_BEGIN_ARGS( SCreditsScreen )
+		{}
+	SLATE_END_ARGS()
+
+	/**
+	 * Constructs the credits screen widgets                   
+	 */
+	UNREALED_API void Construct(const FArguments& InArgs);
+
+	UNREALED_API virtual void Tick(const FGeometry& AllottedGeometry, const double InCurrentTime, const float InDeltaTime) OVERRIDE;
+
+private:
+
+	/** Handles the user clicking the play/pause toggle button. */
+	FReply HandleTogglePlayPause();
+
+	/** Handles when the user scrolls so that we can stop the auto-scrolling when they scroll backwards. */
+	void HandleUserScrolled(float ScrollOffset);
+
+	/** Handles the user clicking links in the credits */
+	void OnBrowserLinkClicked(const FSlateHyperlinkRun::FMetadata& Metadata);
+
+	/** Gets the current brush (play or pause) icon for the play/pause button */
+	const FSlateBrush* GetTogglePlayPauseBrush() const;
+
+	/** The widget that scrolls the credits text */
+	TSharedPtr<SScrollBox> ScrollBox;
+
+	/** The auto scroll rate in pixels per second */
+	float ScrollPixelsPerSecond;
+
+	/** The last recorded scroll position so we can detect the user scrolling up */
+	float PreviousScrollPosition;
+
+	/** If we are playing the auto scroll effect */
+	bool bIsPlaying;
+};
+
