Levels: Changed blueprint type streaming levels icon from yellow dot to blue

diff --git a/Engine/Source/Editor/EditorStyle/Private/SlateEditorStyle.cpp b/Engine/Source/Editor/EditorStyle/Private/SlateEditorStyle.cpp
--- a/Engine/Source/Editor/EditorStyle/Private/SlateEditorStyle.cpp
+++ b/Engine/Source/Editor/EditorStyle/Private/SlateEditorStyle.cpp
@@ -1073,7 +1073,7 @@ void FSlateEditorStyle::FStyle::SetupGeneralStyles()
 		Set( "WorldBrowser.DirectionYPositive", new IMAGE_BRUSH( "Icons/icon_PanUp", Icon16x16 ) );
 		Set( "WorldBrowser.DirectionYNegative", new IMAGE_BRUSH( "Icons/icon_PanDown", Icon16x16 ) );
 		Set( "WorldBrowser.LevelStreamingAlwaysLoaded", new FSlateNoResource() );
-		Set( "WorldBrowser.LevelStreamingKismet", new IMAGE_BRUSH( "Icons/ActorIcons/ComponentMobilityStationary_7x16px", Icon7x16 ) );
+		Set( "WorldBrowser.LevelStreamingBlueprint", new IMAGE_BRUSH( "Icons/icon_levels_blueprinttype_7x16", Icon7x16 ) );
 		Set( "WorldBrowser.LevelsMenuBrush", new IMAGE_BRUSH( "Icons/icon_levels_levelsmenu_40x", Icon25x25 ) );
 		Set( "WorldBrowser.HierarchyButtonBrush", new IMAGE_BRUSH( "Icons/icon_levels_hierarchybutton_16x", Icon16x16 ) );
 		Set( "WorldBrowser.DetailsButtonBrush", new IMAGE_BRUSH( "Icons/icon_levels_detailsbutton_16x", Icon16x16 ) );
diff --git a/Engine/Source/Editor/WorldBrowser/Private/SWorldHierarchyItem.cpp b/Engine/Source/Editor/WorldBrowser/Private/SWorldHierarchyItem.cpp
--- a/Engine/Source/Editor/WorldBrowser/Private/SWorldHierarchyItem.cpp
+++ b/Engine/Source/Editor/WorldBrowser/Private/SWorldHierarchyItem.cpp
@@ -17,7 +17,7 @@ void SWorldHierarchyItem::Construct(const FArguments& InArgs, TSharedRef<STableV
 	HighlightText = InArgs._HighlightText;
 
 	StreamingLevelAlwaysLoadedBrush = FEditorStyle::GetBrush("WorldBrowser.LevelStreamingAlwaysLoaded");
-	StreamingLevelKismetBrush = FEditorStyle::GetBrush("WorldBrowser.LevelStreamingKismet");
+	StreamingLevelKismetBrush = FEditorStyle::GetBrush("WorldBrowser.LevelStreamingBlueprint");
 
 	SMultiColumnTableRow<TSharedPtr<FLevelModel>>::Construct(
 		FSuperRowType::FArguments()
