Updated icons and path for Persona

diff --git a/Engine/Source/Editor/EditorStyle/Private/SlateEditorStyle.cpp b/Engine/Source/Editor/EditorStyle/Private/SlateEditorStyle.cpp
--- a/Engine/Source/Editor/EditorStyle/Private/SlateEditorStyle.cpp
+++ b/Engine/Source/Editor/EditorStyle/Private/SlateEditorStyle.cpp
@@ -3928,30 +3928,30 @@ void FSlateEditorStyle::FStyle::SetupPersonaStyle()
 			);
 
 		// persona commands
-		Set("Persona.ChangeSkeletonPreviewMesh", new IMAGE_BRUSH(TEXT("Icons/icon_ShowSkeletalMeshes_40x"), Icon40x40));
-		Set("Persona.ChangeSkeletonPreviewMesh.Small", new IMAGE_BRUSH(TEXT("Icons/icon_ShowSkeletalMeshes_40x"), Icon20x20));
-		Set("Persona.AnimNotifyWindow", new IMAGE_BRUSH(TEXT("Icons/icon_ShowNavigation_40x"), Icon40x40));
-		Set("Persona.AnimNotifyWindow.Small", new IMAGE_BRUSH(TEXT("Icons/icon_ShowNavigation_40x"), Icon20x20));
-		Set("Persona.RetargetSourceMgr", new IMAGE_BRUSH(TEXT("Icons/icon_PlayDefaultPlayerStart_40x"), Icon40x40));
-		Set("Persona.RetargetSourceMgr.Small", new IMAGE_BRUSH(TEXT("Icons/icon_PlayDefaultPlayerStart_40x"), Icon20x20));
-		Set("Persona.ImportMesh", new IMAGE_BRUSH(TEXT("Icons/icon_PhAT_Prismatic_40x"), Icon40x40));
-		Set("Persona.ImportMesh.Small", new IMAGE_BRUSH(TEXT("Icons/icon_PhAT_Prismatic_40x"), Icon20x20));
-		Set("Persona.ReimportMesh", new IMAGE_BRUSH(TEXT("Icons/icon_Placement_RecentlyPlaced_40x"), Icon40x40));
-		Set("Persona.ReimportMesh.Small", new IMAGE_BRUSH(TEXT("Icons/icon_Placement_RecentlyPlaced_40x"), Icon20x20));
-		Set("Persona.ImportLODs", new IMAGE_BRUSH(TEXT("Icons/icon_StaticMeshEd_AdditionalData_40x"), Icon40x40));
-		Set("Persona.ImportLODs.Small", new IMAGE_BRUSH(TEXT("Icons/icon_StaticMeshEd_AdditionalData_40x"), Icon20x20));
-		Set("Persona.AddBodyPart", new IMAGE_BRUSH(TEXT("Icons/icon_Placement_AllContent_40x"), Icon40x40));
-		Set("Persona.AddBodyPart.Small", new IMAGE_BRUSH(TEXT("Icons/icon_Placement_AllContent_40x"), Icon20x20));
-		Set("Persona.ImportAnimation", new IMAGE_BRUSH(TEXT("Icons/icon_PhAT_Prismatic_40x"), Icon40x40));
-		Set("Persona.ImportAnimation.Small", new IMAGE_BRUSH(TEXT("Icons/icon_PhAT_Prismatic_40x"), Icon20x20));
-		Set("Persona.ReimportAnimation", new IMAGE_BRUSH(TEXT("Icons/icon_PhAT_PhysMat_40x"), Icon40x40));
-		Set("Persona.ReimportAnimation.Small", new IMAGE_BRUSH(TEXT("Icons/icon_PhAT_PhysMat_40x"), Icon20x20));
-		Set("Persona.ApplyCompression", new IMAGE_BRUSH(TEXT("Icons/icon_Cascade_LowestLOD_40x"), Icon40x40));
-		Set("Persona.ApplyCompression.Small", new IMAGE_BRUSH(TEXT("Icons/icon_Cascade_LowestLOD_40x"), Icon20x20));
-		Set("Persona.ExportToFBX", new IMAGE_BRUSH(TEXT("Icons/icon_DestructibleMeshEd_ImportFBX_40x"), Icon40x40));
-		Set("Persona.ExportToFBX.Small", new IMAGE_BRUSH(TEXT("Icons/icon_DestructibleMeshEd_ImportFBX_40x"), Icon20x20));
-		Set("Persona.CreateAsset", new IMAGE_BRUSH(TEXT("Icons/icon_PhAT_NewBody_40x"), Icon40x40));
-		Set("Persona.CreateAsset.Small", new IMAGE_BRUSH(TEXT("Icons/icon_PhAT_NewBody_40x"), Icon20x20));
+		Set("Persona.ChangeSkeletonPreviewMesh", new IMAGE_BRUSH(TEXT("Icons/icon_Persona_PreviewMesh_40x"), Icon40x40));
+		Set("Persona.ChangeSkeletonPreviewMesh.Small", new IMAGE_BRUSH(TEXT("Icons/icon_Persona_PreviewMesh_40x"), Icon20x20));
+		Set("Persona.AnimNotifyWindow", new IMAGE_BRUSH(TEXT("Icons/icon_Persona_AnimNotift_40x"), Icon40x40));
+		Set("Persona.AnimNotifyWindow.Small", new IMAGE_BRUSH(TEXT("Icons/icon_Persona_AnimNotift_40x"), Icon20x20));
+		Set("Persona.RetargetSourceMgr", new IMAGE_BRUSH(TEXT("Icons/icon_Persona_Retarget_40x"), Icon40x40));
+		Set("Persona.RetargetSourceMgr.Small", new IMAGE_BRUSH(TEXT("Icons/icon_Persona_Retarget_40x"), Icon20x20));
+		Set("Persona.ImportMesh", new IMAGE_BRUSH(TEXT("Icons/icon_Persona_ImportMesh_40x"), Icon40x40));
+		Set("Persona.ImportMesh.Small", new IMAGE_BRUSH(TEXT("Icons/icon_Persona_ImportMesh_40x"), Icon20x20));
+		Set("Persona.ReimportMesh", new IMAGE_BRUSH(TEXT("Icons/icon_Persona_ReimportMesh_40x"), Icon40x40));
+		Set("Persona.ReimportMesh.Small", new IMAGE_BRUSH(TEXT("Icons/icon_Persona_ReimportMesh_40x"), Icon20x20));
+		Set("Persona.ImportLODs", new IMAGE_BRUSH(TEXT("Icons/icon_Persona_ImportLODs_40x"), Icon40x40));
+		Set("Persona.ImportLODs.Small", new IMAGE_BRUSH(TEXT("Icons/icon_Persona_ImportLODs_40x"), Icon20x20));
+//		Set("Persona.AddBodyPart", new IMAGE_BRUSH(TEXT("Icons/icon_Placement_AllContent_40x"), Icon40x40));
+//		Set("Persona.AddBodyPart.Small", new IMAGE_BRUSH(TEXT("Icons/icon_Placement_AllContent_40x"), Icon20x20));
+		Set("Persona.ImportAnimation", new IMAGE_BRUSH(TEXT("Icons/icon_Persona_ImportAnim_40x"), Icon40x40));
+		Set("Persona.ImportAnimation.Small", new IMAGE_BRUSH(TEXT("Icons/icon_Persona_ImportAnim_40x"), Icon20x20));
+		Set("Persona.ReimportAnimation", new IMAGE_BRUSH(TEXT("Icons/icon_Persona_ReimportAnim_40x"), Icon40x40));
+		Set("Persona.ReimportAnimation.Small", new IMAGE_BRUSH(TEXT("Icons/icon_Persona_ReimportAnim_40x"), Icon20x20));
+		Set("Persona.ApplyCompression", new IMAGE_BRUSH(TEXT("Icons/icon_Persona_Compression_40x"), Icon40x40));
+		Set("Persona.ApplyCompression.Small", new IMAGE_BRUSH(TEXT("Icons/icon_Persona_Compression_40x"), Icon20x20));
+		Set("Persona.ExportToFBX", new IMAGE_BRUSH(TEXT("Icons/icon_Persona_ExportFBX_40x"), Icon40x40));
+		Set("Persona.ExportToFBX.Small", new IMAGE_BRUSH(TEXT("Icons/icon_Persona_ExportFBX_40x"), Icon20x20));
+		Set("Persona.CreateAsset", new IMAGE_BRUSH(TEXT("Icons/icon_Persona_CreateAsset_40x"), Icon40x40));
+		Set("Persona.CreateAsset.Small", new IMAGE_BRUSH(TEXT("Icons/icon_Persona_CreateAsset_40x"), Icon20x20));
 
 		// persona extras
 		Set("Persona.ConvertAnimationGraph", new IMAGE_BRUSH("Old/Graph/ConvertIcon", Icon40x40));
