Added landscape icon style to EditorStyle

diff --git a/Engine/Source/Editor/EditorStyle/Private/SlateEditorStyle.cpp b/Engine/Source/Editor/EditorStyle/Private/SlateEditorStyle.cpp
--- a/Engine/Source/Editor/EditorStyle/Private/SlateEditorStyle.cpp
+++ b/Engine/Source/Editor/EditorStyle/Private/SlateEditorStyle.cpp
@@ -4328,6 +4328,7 @@ void FSlateEditorStyle::FStyle::SetupClassIconsAndThumbnails()
 			TEXT("ForceFeedbackEffect"),
 			TEXT("InterpData"),
 			TEXT("KillZVolume"),
+			TEXT("Landscape"),
 			TEXT("LevelStreamingVolume"),
 			TEXT("LightmassCharacterIndirectDetailVolume"),
 			TEXT("LightmassImportanceVolume"),
