Icon for live preview button in material editor.

TTP 342063 - MATERIAL EDITOR: Small Icons: Live Preview icon missing

diff --git a/Engine/Source/Editor/EditorStyle/Private/SlateEditorStyle.cpp b/Engine/Source/Editor/EditorStyle/Private/SlateEditorStyle.cpp
--- a/Engine/Source/Editor/EditorStyle/Private/SlateEditorStyle.cpp
+++ b/Engine/Source/Editor/EditorStyle/Private/SlateEditorStyle.cpp
@@ -4784,6 +4784,8 @@ void FSlateEditorStyle::FStyle::SetupToolkitStyles()
 		Set( "MaterialEditor.ToggleRealtimeExpressions.Small", new IMAGE_BRUSH( "Icons/icon_MatEd_LiveNodes_40x", Icon20x20 ) );
 		Set( "MaterialEditor.AlwaysRefreshAllPreviews", new IMAGE_BRUSH( "Icons/icon_MatEd_Refresh_40x", Icon40x40 ) );
 		Set( "MaterialEditor.AlwaysRefreshAllPreviews.Small", new IMAGE_BRUSH( "Icons/icon_MatEd_Refresh_40x", Icon20x20 ) );
+		Set( "MaterialEditor.ToggleLivePreview", new IMAGE_BRUSH("Icons/icon_MatEd_LivePreview_40x", Icon40x40) );
+		Set( "MaterialEditor.ToggleLivePreview.Small", new IMAGE_BRUSH("Icons/icon_MatEd_LivePreview_40x", Icon20x20) );
 		Set( "MaterialEditor.ShowHideConnectors", new IMAGE_BRUSH( "Icons/icon_MatEd_Connectors_40x", Icon40x40 ) );
 		Set( "MaterialEditor.ShowHideConnectors.Small", new IMAGE_BRUSH( "Icons/icon_MatEd_Connectors_40x", Icon20x20 ) );
 		Set( "MaterialEditor.CameraHome", new IMAGE_BRUSH( "Icons/icon_MatEd_Home_40x", Icon40x40 ) );
