Media: Added media asset type icons to Slate style set

diff --git a/Engine/Source/Editor/EditorStyle/Private/SlateEditorStyle.cpp b/Engine/Source/Editor/EditorStyle/Private/SlateEditorStyle.cpp
--- a/Engine/Source/Editor/EditorStyle/Private/SlateEditorStyle.cpp
+++ b/Engine/Source/Editor/EditorStyle/Private/SlateEditorStyle.cpp
@@ -4348,6 +4348,9 @@ void FSlateEditorStyle::FStyle::SetupClassIconsAndThumbnails()
 			TEXT("MaterialFunction"),
 			TEXT("MaterialInstanceConstant"),
 			TEXT("MaterialParameterCollection"),
+			TEXT("MediaAsset"),
+			TEXT("MediaSoundWave"),
+			TEXT("MediaTexture"),
 			TEXT("MultiFont"),
 			TEXT("NavLinkProxy"),
 			TEXT("NavMeshBoundsVolume"),
