Added icon for documentation actor

diff --git a/Engine/Source/Editor/EditorStyle/Private/SlateEditorStyle.cpp b/Engine/Source/Editor/EditorStyle/Private/SlateEditorStyle.cpp
--- a/Engine/Source/Editor/EditorStyle/Private/SlateEditorStyle.cpp
+++ b/Engine/Source/Editor/EditorStyle/Private/SlateEditorStyle.cpp
@@ -4249,6 +4249,7 @@ void FSlateEditorStyle::FStyle::SetupClassIconsAndThumbnails()
 		Set("ClassIcon.Pawn", new IMAGE_BRUSH("Icons/ActorIcons/Pawn_16px", Icon16x16));
 		Set("ClassIcon.Note", new IMAGE_BRUSH("Icons/ActorIcons/Note_16px", Icon16x16));
 		Set("ClassIcon.Character", new IMAGE_BRUSH("Icons/ActorIcons/Character_16px", Icon16x16));
+		Set("ClassIcon.DocumentationActor", new IMAGE_BRUSH("Icons/ActorIcons/DocumentationActor_16x", Icon16x16));
 
 		Set("ClassIcon.BlueprintActorBase", new IMAGE_BRUSH("Icons/AssetIcons/Blueprint_16x", Icon16x16));
 		Set("ClassIcon.StaticMeshActor", new IMAGE_BRUSH("Icons/AssetIcons/StaticMesh_16x", Icon16x16));
