Tutorials 2.0: Adding new tutorial wizard and rich text support

Tutorials can now be created using the right-click asset menu->Misc->Tutorial Blueprint.
Importer allows you to pick a UDN file, which will then be converted into a new Tutorial Asset.

Adapted UDN parser to also output rich text format for the purposes of this importer. This may be deprecated at some point soon, once existing tutorials have been imported.

Rich text editor added to tutorial details customization. Features; font, color, size & hyperlinks for now.
Rich text display also added to content widgets.

Added option to optionally show the widget highlight if only the bubble widget is required.

diff --git a/Engine/Source/Editor/Documentation/Private/DocumentationPage.cpp b/Engine/Source/Editor/Documentation/Private/DocumentationPage.cpp
--- a/Engine/Source/Editor/Documentation/Private/DocumentationPage.cpp
+++ b/Engine/Source/Editor/Documentation/Private/DocumentationPage.cpp
@@ -23,6 +23,7 @@ bool FDocumentationPage::GetExcerptContent( FExcerpt& Excerpt )
 		{
 			Parser->GetExcerptContent( Link, StoredExcerpts[ Index ] );
 			Excerpt.Content = StoredExcerpts[ Index ].Content;
+			Excerpt.RichText = StoredExcerpts[ Index ].RichText;
 			return true;
 		}
 	}
diff --git a/Engine/Source/Editor/Documentation/Private/UDNParser.cpp b/Engine/Source/Editor/Documentation/Private/UDNParser.cpp
--- a/Engine/Source/Editor/Documentation/Private/UDNParser.cpp
+++ b/Engine/Source/Editor/Documentation/Private/UDNParser.cpp
@@ -14,6 +14,19 @@
 
 FName UDNParseErrorLog("UDNParser");
 
+namespace LinkPrefixes
+{
+
+static const FString DocLinkSpecifier( TEXT( "DOCLINK:" ) );
+static const FString TutorialLinkSpecifier( TEXT( "TUTORIALLINK:" ) );
+static const FString HttpLinkSpecifier( TEXT( "http://" ) );
+static const FString HttpsLinkSpecifier( TEXT( "https://" ) );
+
+static const FString CodeLinkSpecifier(TEXT("CODELINK:"));
+static const FString AssetLinkSpecifier(TEXT("ASSETLINK:"));
+
+}
+
 TSharedRef< FUDNParser > FUDNParser::Create( const TSharedPtr< FParserConfiguration >& ParserConfig, const FDocumentationStyle& Style ) 
 {
 	TSharedPtr< FParserConfiguration > FinalParserConfig = ParserConfig;
@@ -689,7 +702,16 @@ void FUDNParser::AppendExcerpt(TSharedPtr<SVerticalBox> Box, TSharedRef<SWidget>
 	];
 }
 
-void FUDNParser::AddContentToExcerpt(TSharedPtr<SVerticalBox> Box, const FString& ContentSource)
+static void AddLineSeperator(FExcerpt& Excerpt)
+{
+	if(!Excerpt.RichText.IsEmpty())
+	{
+		Excerpt.RichText += LINE_TERMINATOR;
+		Excerpt.RichText += LINE_TERMINATOR;
+	}
+}
+
+void FUDNParser::AddContentToExcerpt(TSharedPtr<SVerticalBox> Box, const FString& ContentSource, FExcerpt& Excerpt)
 {
 	if ( !ContentSource.IsEmpty() )
 	{
@@ -699,6 +721,9 @@ void FUDNParser::AddContentToExcerpt(TSharedPtr<SVerticalBox> Box, const FString
 			.TextStyle(FEditorStyle::Get(), Style.ContentStyleName)
 			.WrapTextAt(WrapAt)
 		);
+
+		AddLineSeperator(Excerpt);
+		Excerpt.RichText += FString::Printf(TEXT("<TextStyle FontFamily=\"Roboto\" FontSize=\"11\" FontStyle=\"Regular\" FontColor=\"(R=0.0,G=0.0,B=0.0,A=1.0)\">%s</>"), *ContentSource);
 	}
 }
 
@@ -744,7 +769,7 @@ TSharedRef< SWidget > FUDNParser::GenerateExcerptContent( const FString& Link, F
 
 			if ( ExcerptStack.Num() == 0 )
 			{
-				AddContentToExcerpt(Box, CurrentStringContent);
+				AddContentToExcerpt(Box, CurrentStringContent, Excerpt);
 				break;
 			}
 		}
@@ -817,14 +842,14 @@ TSharedRef< SWidget > FUDNParser::GenerateExcerptContent( const FString& Link, F
 			{
 			case FUDNLine::Whitespace:
 				// Will only apply whitespace for the first empty line
-				AddContentToExcerpt(Box, CurrentStringContent);
+				AddContentToExcerpt(Box, CurrentStringContent, Excerpt);
 				CurrentStringContent.Empty();
 				break;
 			case FUDNLine::Content:
 				CurrentStringContent += Line.AdditionalContent[0];
 				break;
 			case FUDNLine::BoldContent:
-				AddContentToExcerpt(Box, CurrentStringContent);
+				AddContentToExcerpt(Box, CurrentStringContent, Excerpt);
 				CurrentStringContent.Empty();
 
 				AppendExcerpt(Box,
@@ -832,17 +857,20 @@ TSharedRef< SWidget > FUDNParser::GenerateExcerptContent( const FString& Link, F
 					.Text(Line.AdditionalContent[0])
 					.TextStyle(FEditorStyle::Get(), Style.BoldContentStyleName)
 					);
+
+				AddLineSeperator(Excerpt);
+				Excerpt.RichText += FString::Printf(TEXT("<TextStyle FontFamily=\"Roboto\" FontSize=\"11\" FontStyle=\"Bold\" FontColor=\"(R=0.0,G=0.0,B=0.0,A=1.0)\">%s</>"), *Line.AdditionalContent[0]);
 				break;
 			case FUDNLine::NumberedContent:
-				AddContentToExcerpt(Box, CurrentStringContent);
+				AddContentToExcerpt(Box, CurrentStringContent, Excerpt);
 				CurrentStringContent = FString::Printf(TEXT("%i. %s"), CurrentNumbering, *Line.AdditionalContent[0]);
-				AddContentToExcerpt(Box, CurrentStringContent);
+				AddContentToExcerpt(Box, CurrentStringContent, Excerpt);
 				CurrentStringContent.Empty();
 
 				++CurrentNumbering;
 				break;
 			case FUDNLine::HorizontalRule:
-				AddContentToExcerpt(Box, CurrentStringContent);
+				AddContentToExcerpt(Box, CurrentStringContent, Excerpt);
 				CurrentStringContent.Empty();
 
 				Box->AddSlot()
@@ -856,9 +884,12 @@ TSharedRef< SWidget > FUDNParser::GenerateExcerptContent( const FString& Link, F
 							.SeparatorImage(FEditorStyle::GetBrush(Style.SeparatorStyleName))
 						]
 					];
+
+				AddLineSeperator(Excerpt);
+				Excerpt.RichText += TEXT("<hr></>");
 				break;
 			case FUDNLine::Header1:
-				AddContentToExcerpt(Box, CurrentStringContent);
+				AddContentToExcerpt(Box, CurrentStringContent, Excerpt);
 				CurrentStringContent.Empty();
 
 				AppendExcerpt(Box,
@@ -866,9 +897,12 @@ TSharedRef< SWidget > FUDNParser::GenerateExcerptContent( const FString& Link, F
 					.Text(Line.AdditionalContent[0])
 					.TextStyle(FEditorStyle::Get(), Style.Header1StyleName)
 					);
+
+				AddLineSeperator(Excerpt);
+				Excerpt.RichText += FString::Printf(TEXT("<TextStyle FontFamily=\"Roboto\" FontSize=\"32\" FontStyle=\"Bold\" FontColor=\"(R=0.0,G=0.0,B=0.0,A=1.0)\">%s</>"), *Line.AdditionalContent[0]);
 				break;
 			case FUDNLine::Header2:
-				AddContentToExcerpt(Box, CurrentStringContent);
+				AddContentToExcerpt(Box, CurrentStringContent, Excerpt);
 				CurrentStringContent.Empty();
 
 				AppendExcerpt(Box,
@@ -876,9 +910,12 @@ TSharedRef< SWidget > FUDNParser::GenerateExcerptContent( const FString& Link, F
 					.Text(Line.AdditionalContent[0])
 					.TextStyle(FEditorStyle::Get(), Style.Header2StyleName)
 					);
+
+				AddLineSeperator(Excerpt);
+				Excerpt.RichText += FString::Printf(TEXT("<TextStyle FontFamily=\"Roboto\" FontSize=\"24\" FontStyle=\"Bold\" FontColor=\"(R=0.0,G=0.0,B=0.0,A=1.0)\">%s</>"), *Line.AdditionalContent[0]);
 				break;
 			case FUDNLine::Link:
-				AddContentToExcerpt(Box, CurrentStringContent);
+				AddContentToExcerpt(Box, CurrentStringContent, Excerpt);
 				CurrentStringContent.Empty();
 
 				AppendExcerpt(Box,
@@ -888,13 +925,40 @@ TSharedRef< SWidget > FUDNParser::GenerateExcerptContent( const FString& Link, F
 					.UnderlineStyle(FEditorStyle::Get(), Style.HyperlinkButtonStyleName)
 					.OnNavigate( this, &FUDNParser::HandleHyperlinkNavigate, Line.AdditionalContent[1])
 					);
+
+				AddLineSeperator(Excerpt);
+
+				if(Line.AdditionalContent[1].Contains(LinkPrefixes::DocLinkSpecifier))
+				{
+					const FString Link = Line.AdditionalContent[1].RightChop(LinkPrefixes::DocLinkSpecifier.Len());
+					Excerpt.RichText += FString::Printf(TEXT("<a id=\"udn\" href=\"%s\" style=\"%s\">%s</>"), *Link, *Style.HyperlinkStyleName.ToString(), *Line.AdditionalContent[0]);
+				}
+				else if(Line.AdditionalContent[1].Contains(LinkPrefixes::AssetLinkSpecifier))
+				{
+					const FString Link = Line.AdditionalContent[1].RightChop(LinkPrefixes::AssetLinkSpecifier.Len());
+					Excerpt.RichText += FString::Printf(TEXT("<a id=\"asset\" href=\"%s\" style=\"%s\">%s</>"), *Link, *Style.HyperlinkStyleName.ToString(), *Line.AdditionalContent[0]);
+				}
+				else if(Line.AdditionalContent[1].Contains(LinkPrefixes::CodeLinkSpecifier))
+				{
+					const FString Link = Line.AdditionalContent[1].RightChop(LinkPrefixes::CodeLinkSpecifier.Len());
+					Excerpt.RichText += FString::Printf(TEXT("<a id=\"code\" href=\"%s\" style=\"%s\">%s</>"), *Link, *Style.HyperlinkStyleName.ToString(), *Line.AdditionalContent[0]);
+				}
+				else if(Line.AdditionalContent[1].Contains(LinkPrefixes::TutorialLinkSpecifier))
+				{
+					const FString Link = Line.AdditionalContent[1].RightChop(LinkPrefixes::TutorialLinkSpecifier.Len());
+					Excerpt.RichText += FString::Printf(TEXT("<a id=\"tutorial\" href=\"%s\" style=\"%s\">%s</>"), *Link, *Style.HyperlinkStyleName.ToString(), *Line.AdditionalContent[0]);
+				}
+				else
+				{
+					Excerpt.RichText += FString::Printf(TEXT("<a id=\"browser\" href=\"%s\" style=\"%s\">%s</>"), *Link, *Style.HyperlinkStyleName.ToString(), *Line.AdditionalContent[0]);
+				}
 				break;
 			case FUDNLine::Image:
 				ConcatenatedPath = FullPath / TEXT("Images") / Line.AdditionalContent[1];
 				DynamicBrush = GetDynamicBrushFromImagePath(ConcatenatedPath);
 				if (DynamicBrush.IsValid())
 				{
-					AddContentToExcerpt(Box, CurrentStringContent);
+					AddContentToExcerpt(Box, CurrentStringContent, Excerpt);
 					CurrentStringContent.Empty();
 
 					AppendExcerpt(Box,
@@ -905,13 +969,16 @@ TSharedRef< SWidget > FUDNParser::GenerateExcerptContent( const FString& Link, F
 
 					DynamicBrushesUsed.AddUnique(DynamicBrush);
 				}
+
+				AddLineSeperator(Excerpt);
+				Excerpt.RichText += FString::Printf(TEXT("<img src=\"%s\"></>"), *ConcatenatedPath);
 				break;
 			case FUDNLine::ImageLink:
 				ConcatenatedPath = FullPath / TEXT("Images") / Line.AdditionalContent[1];
 				DynamicBrush = GetDynamicBrushFromImagePath(ConcatenatedPath);
 				if (DynamicBrush.IsValid())
 				{
-					AddContentToExcerpt(Box, CurrentStringContent);
+					AddContentToExcerpt(Box, CurrentStringContent, Excerpt);
 					CurrentStringContent.Empty();
 
 					AppendExcerpt(Box,
@@ -928,6 +995,9 @@ TSharedRef< SWidget > FUDNParser::GenerateExcerptContent( const FString& Link, F
 
 					DynamicBrushesUsed.AddUnique(DynamicBrush);
 				}
+
+				AddLineSeperator(Excerpt);
+				Excerpt.RichText += FString::Printf(TEXT("<img src=\"%s\" href=\"%s\"></>"), *ConcatenatedPath, *Line.AdditionalContent[2]);
 				break;
 			default: break;
 			}
diff --git a/Engine/Source/Editor/Documentation/Private/UDNParser.h b/Engine/Source/Editor/Documentation/Private/UDNParser.h
--- a/Engine/Source/Editor/Documentation/Private/UDNParser.h
+++ b/Engine/Source/Editor/Documentation/Private/UDNParser.h
@@ -138,7 +138,7 @@ private:
 	FReply OnImageLinkClicked(FString Payload);
 	
 	/** Adds the content text source to the scrollbox */
-	void AddContentToExcerpt(TSharedPtr<SVerticalBox> Box, const FString& ContentSource);
+	void AddContentToExcerpt(TSharedPtr<SVerticalBox> Box, const FString& ContentSource, FExcerpt& Excerpt);
 	
 	/** Gets the dynamic brush for the given filename */
 	TSharedPtr<FSlateDynamicImageBrush> GetDynamicBrushFromImagePath(FString Filename);
diff --git a/Engine/Source/Editor/Documentation/Public/IDocumentation.h b/Engine/Source/Editor/Documentation/Public/IDocumentation.h
--- a/Engine/Source/Editor/Documentation/Public/IDocumentation.h
+++ b/Engine/Source/Editor/Documentation/Public/IDocumentation.h
@@ -16,6 +16,7 @@ public:
 		, NumberedContentStyleName(TEXT("Documentation.NumberedContent"))
 		, Header1StyleName(TEXT("Documentation.Header1"))
 		, Header2StyleName(TEXT("Documentation.Header2"))
+		, HyperlinkStyleName(TEXT("Documentation.Hyperlink"))
 		, HyperlinkButtonStyleName(TEXT("Documentation.Hyperlink.Button"))
 		, HyperlinkTextStyleName(TEXT("Documentation.Hyperlink.Text"))
 		, SeparatorStyleName(TEXT("Documentation.Separator"))
@@ -93,6 +94,9 @@ public:
 	/** Header2 text style */
 	FName Header2StyleName;
 
+	/** Hyperlink style */
+	FName HyperlinkStyleName;
+
 	/** Hyperlink button style */
 	FName HyperlinkButtonStyleName;
 
diff --git a/Engine/Source/Editor/Documentation/Public/IDocumentationPage.h b/Engine/Source/Editor/Documentation/Public/IDocumentationPage.h
--- a/Engine/Source/Editor/Documentation/Public/IDocumentationPage.h
+++ b/Engine/Source/Editor/Documentation/Public/IDocumentationPage.h
@@ -24,6 +24,8 @@ struct FExcerpt
 	TSharedPtr<SWidget> Content;
 	TMap< FString, FString > Variables;
 	int32 LineNumber;
+	/** Rich text version of the excerpt */
+	FString RichText;
 };
 
 class IDocumentationPage
diff --git a/Engine/Source/Editor/EditorStyle/Private/SlateEditorStyle.cpp b/Engine/Source/Editor/EditorStyle/Private/SlateEditorStyle.cpp
--- a/Engine/Source/Editor/EditorStyle/Private/SlateEditorStyle.cpp
+++ b/Engine/Source/Editor/EditorStyle/Private/SlateEditorStyle.cpp
@@ -2312,7 +2312,7 @@ void FSlateEditorStyle::FStyle::SetupTutorialStyles()
 
 	// Documentation defaults
 	const FTextBlockStyle DocumentationText = FTextBlockStyle(NormalText)
-		.SetFont(TTF_CORE_FONT( "Fonts/Roboto-Regular", 12 ));
+		.SetFont(TTF_CORE_FONT( "Fonts/Roboto-Regular", 11 ));
 
 	const FTextBlockStyle DocumentationHyperlinkText = FTextBlockStyle(DocumentationText)
 		.SetColorAndOpacity(FLinearColor::Blue);
@@ -2320,15 +2320,22 @@ void FSlateEditorStyle::FStyle::SetupTutorialStyles()
 	const FTextBlockStyle DocumentationHeaderText = FTextBlockStyle(NormalText)
 		.SetFont(TTF_FONT("Fonts/Roboto-Black", 32));
 
-	const FButtonStyle DocumentationHyperlink = FButtonStyle()
-			.SetNormal(BORDER_BRUSH( "Old/HyperlinkDotted", FMargin(0,0,0,3/16.0f)))
+	const FButtonStyle DocumentationHyperlinkButton = FButtonStyle()
+			.SetNormal(BORDER_BRUSH( "Old/HyperlinkDotted", FMargin(0,0,0,3/16.0f), FSlateColor::UseForeground()))
 			.SetPressed(FSlateNoResource())
-			.SetHovered(BORDER_BRUSH( "Old/HyperlinkUnderline", FMargin(0,0,0,3/16.0f)));
+			.SetHovered(BORDER_BRUSH( "Old/HyperlinkUnderline", FMargin(0,0,0,3/16.0f), FSlateColor::UseForeground()));
 
 	// Documentation
 	{
 		Set( "Documentation.Content", FTextBlockStyle(DocumentationText) );
-		Set("Documentation.Hyperlink.Button", FButtonStyle(DocumentationHyperlink));
+
+		const FHyperlinkStyle DocumentationHyperlink = FHyperlinkStyle()
+			.SetUnderlineStyle(DocumentationHyperlinkButton)
+			.SetTextStyle(DocumentationText)
+			.SetPadding(FMargin(0.0f));
+		Set("Documentation.Hyperlink", DocumentationHyperlink);
+
+		Set("Documentation.Hyperlink.Button", FButtonStyle(DocumentationHyperlinkButton));
 		Set("Documentation.Hyperlink.Text",   FTextBlockStyle(DocumentationHyperlinkText));
 		Set("Documentation.NumberedContent",  FTextBlockStyle(DocumentationText));
 		Set( "Documentation.BoldContent", FTextBlockStyle(DocumentationText)
@@ -2418,7 +2425,7 @@ void FSlateEditorStyle::FStyle::SetupTutorialStyles()
 		Set("Tutorials.Header2", FTextBlockStyle(TutorialHeaderText)
 			.SetFontSize(24));
 
-		Set( "Tutorials.Hyperlink.Button", FButtonStyle(DocumentationHyperlink)
+		Set( "Tutorials.Hyperlink.Button", FButtonStyle(DocumentationHyperlinkButton)
 			.SetNormal(BORDER_BRUSH( "Old/HyperlinkDotted", FMargin(0,0,0,3/16.0f), FLinearColor::Black))
 			.SetHovered(BORDER_BRUSH( "Old/HyperlinkUnderline", FMargin(0,0,0,3/16.0f), FLinearColor::Black)));
 
@@ -2429,8 +2436,120 @@ void FSlateEditorStyle::FStyle::SetupTutorialStyles()
 			.SetFillImage( BOX_BRUSH( "Common/ProgressBar_NeutralFill", FMargin(5.f/12.f) ) )
 			.SetMarqueeImage( IMAGE_BRUSH( "Common/ProgressBar_Marquee", FVector2D(20,12), FLinearColor::White, ESlateBrushTileType::Horizontal ) )
 			);
+
+		// Default text styles
+		{
+			const FTextBlockStyle RichTextNormal = FTextBlockStyle()
+				.SetFont(TTF_FONT("Fonts/Roboto-Regular", 11))
+				.SetColorAndOpacity(FSlateColor::UseForeground())
+				.SetShadowOffset(FVector2D::ZeroVector)
+				.SetShadowColorAndOpacity(FLinearColor::Black)
+				.SetHighlightColor(FLinearColor(0.02f, 0.3f, 0.0f))
+				.SetHighlightShape(BOX_BRUSH("Common/TextBlockHighlightShape", FMargin(3.f /8.f)));
+			Set( "TutorialEditableText.Editor.Text", RichTextNormal );
+
+			{
+				const FButtonStyle RichTextHyperlinkButton = FButtonStyle()
+					.SetNormal(BORDER_BRUSH("Old/HyperlinkDotted", FMargin(0,0,0,3/16.0f), FSlateColor::UseForeground() ) )
+					.SetPressed(FSlateNoResource() )
+					.SetHovered(BORDER_BRUSH("Old/HyperlinkUnderline", FMargin(0,0,0,3/16.0f), FSlateColor::UseForeground() ) );
+
+				const FHyperlinkStyle RichTextHyperlink = FHyperlinkStyle()
+					.SetUnderlineStyle(RichTextHyperlinkButton)
+					.SetTextStyle(RichTextNormal)
+					.SetPadding(FMargin(0.0f));
+				Set( "TutorialEditableText.Editor.Hyperlink", RichTextHyperlink );
+			}
+		}
+
+		// Toolbar
+		{
+			const FLinearColor NormalColor(FColor(0xffeff3f3));
+			const FLinearColor SelectedColor(FColor(0xffdbe4d5));
+			const FLinearColor HoverColor(FColor(0xffdbe4e4));
+			const FLinearColor DisabledColor(FColor(0xaaaaaa));
+			const FLinearColor TextColor(FColor(0xff2c3e50));
+
+			Set( "TutorialEditableText.RoundedBackground", new BOX_BRUSH( "Common/RoundedSelection_16x", 4.0f/16.0f, FLinearColor(FColor(0xffeff3f3)) ) );
+
+			Set( "TutorialEditableText.Toolbar.HyperlinkImage", new IMAGE_BRUSH( "Tutorials/hyperlink", Icon16x16, TextColor ) );
+
+			Set( "TutorialEditableText.Toolbar.TextColor", TextColor );
+
+			Set( "TutorialEditableText.Toolbar.Text", FTextBlockStyle(NormalText)
+				.SetFont( TTF_FONT( "Fonts/Roboto-Regular", 10 ) )
+				.SetColorAndOpacity( TextColor )
+				);
+
+			Set( "TutorialEditableText.Toolbar.BoldText", FTextBlockStyle(NormalText)
+				.SetFont( TTF_FONT( "Fonts/Roboto-Bold", 10 ) )
+				.SetColorAndOpacity( TextColor )
+				);
+
+			Set( "TutorialEditableText.Toolbar.ItalicText", FTextBlockStyle(NormalText)
+				.SetFont( TTF_FONT( "Fonts/Roboto-Italic", 10 ) )
+				.SetColorAndOpacity( TextColor )
+				);
+
+			Set( "TutorialEditableText.Toolbar.Checkbox", FCheckBoxStyle()
+				.SetCheckBoxType(ESlateCheckBoxType::CheckBox)
+				.SetUncheckedImage( IMAGE_BRUSH( "Common/CheckBox", Icon16x16, FLinearColor::White ) )
+				.SetUncheckedHoveredImage( IMAGE_BRUSH( "Common/CheckBox", Icon16x16, HoverColor ) )
+				.SetUncheckedPressedImage( IMAGE_BRUSH( "Common/CheckBox_Hovered", Icon16x16, HoverColor ) )
+				.SetCheckedImage( IMAGE_BRUSH( "Common/CheckBox_Checked_Hovered", Icon16x16, FLinearColor::White ) )
+				.SetCheckedHoveredImage( IMAGE_BRUSH( "Common/CheckBox_Checked_Hovered", Icon16x16, HoverColor ) )
+				.SetCheckedPressedImage( IMAGE_BRUSH( "Common/CheckBox_Checked", Icon16x16, HoverColor ) )
+				.SetUndeterminedImage( IMAGE_BRUSH( "Common/CheckBox_Undetermined", Icon16x16, FLinearColor::White ) )
+				.SetUndeterminedHoveredImage( IMAGE_BRUSH( "Common/CheckBox_Undetermined_Hovered", Icon16x16, HoverColor ) )
+				.SetUndeterminedPressedImage( IMAGE_BRUSH( "Common/CheckBox_Undetermined_Hovered", Icon16x16, FLinearColor::White ) )
+				);
+
+			Set( "TutorialEditableText.Toolbar.ToggleButtonCheckbox", FCheckBoxStyle()
+				.SetCheckBoxType(ESlateCheckBoxType::ToggleButton)
+				.SetUncheckedImage( BOX_BRUSH("Tutorials/FlatColorSquare", FVector2D(1.0f, 1.0f), FMargin(1), NormalColor ) )
+				.SetUncheckedHoveredImage( BOX_BRUSH("Tutorials/FlatColorSquare", FVector2D(1.0f, 1.0f), FMargin(1), HoverColor ) )
+				.SetUncheckedPressedImage( BOX_BRUSH("Tutorials/FlatColorSquare", FVector2D(1.0f, 1.0f), FMargin(1), HoverColor ) )
+				.SetCheckedImage( BOX_BRUSH("Tutorials/FlatColorSquare",  FVector2D(1.0f, 1.0f), FMargin(1), SelectedColor ) )
+				.SetCheckedHoveredImage( BOX_BRUSH("Tutorials/FlatColorSquare",  FVector2D(1.0f, 1.0f), FMargin(1), HoverColor ) )
+				.SetCheckedPressedImage( BOX_BRUSH("Tutorials/FlatColorSquare",  FVector2D(1.0f, 1.0f), FMargin(1), HoverColor ) )
+				);
+
+			const FButtonStyle Button = FButtonStyle()
+				.SetNormal( BOX_BRUSH("Tutorials/FlatColorSquare", FVector2D(1.0f, 1.0f), FMargin(1), NormalColor ) )
+				.SetHovered( BOX_BRUSH("Tutorials/FlatColorSquare", FVector2D(1.0f, 1.0f), FMargin(1), HoverColor ) )
+				.SetPressed( BOX_BRUSH("Tutorials/FlatColorSquare",  FVector2D(1.0f, 1.0f), FMargin(1), SelectedColor ) )
+				.SetDisabled( BOX_BRUSH("Tutorials/FlatColorSquare",  FVector2D(1.0f, 1.0f), FMargin(1), DisabledColor ) )
+				.SetNormalPadding( FMargin( 2,2,2,2 ) )
+				.SetPressedPadding( FMargin( 2,3,2,1 ) );
+			Set( "TutorialEditableText.Toolbar.Button", Button );
+
+			const FComboButtonStyle ComboButton = FComboButtonStyle()
+				.SetButtonStyle(Button)
+				.SetDownArrowImage(IMAGE_BRUSH("Common/ComboArrow", Icon8x8))
+				.SetMenuBorderBrush(BOX_BRUSH( "Tutorials/FlatColorSquare", FVector2D(1.0f, 1.0f), FMargin(1), NormalColor ))
+				.SetMenuBorderPadding(FMargin(0.0f));
+			Set( "TutorialEditableText.Toolbar.ComboButton", ComboButton	);
+
+			{
+				const FButtonStyle ComboBoxButton = FButtonStyle()
+					.SetNormal( BOX_BRUSH("Tutorials/FlatColorSquare", FVector2D(1.0f, 1.0f), FMargin(1), FLinearColor::White ) )
+					.SetHovered( BOX_BRUSH("Tutorials/FlatColorSquare", FVector2D(1.0f, 1.0f), FMargin(1), FLinearColor::White ) )
+					.SetPressed( BOX_BRUSH("Tutorials/FlatColorSquare",  FVector2D(1.0f, 1.0f), FMargin(1), FLinearColor::White ) )
+					.SetDisabled( BOX_BRUSH("Tutorials/FlatColorSquare",  FVector2D(1.0f, 1.0f), FMargin(1), DisabledColor ) )
+					.SetNormalPadding( FMargin( 2,2,2,2 ) )
+					.SetPressedPadding( FMargin( 2,3,2,1 ) );
+
+				const FComboButtonStyle ComboBoxComboButton = FComboButtonStyle(ComboButton)
+					.SetButtonStyle(ComboBoxButton)
+					.SetMenuBorderPadding(FMargin(1.0));
+
+				Set( "TutorialEditableText.Toolbar.ComboBox", FComboBoxStyle()
+					.SetComboButtonStyle(ComboBoxComboButton)
+					);
+			}
+		}
 	}
-	}
+}
 
 void FSlateEditorStyle::FStyle::SetupPropertyEditorStyles()
 	{
diff --git a/Engine/Source/Editor/IntroTutorials/IntroTutorials.Build.cs b/Engine/Source/Editor/IntroTutorials/IntroTutorials.Build.cs
--- a/Engine/Source/Editor/IntroTutorials/IntroTutorials.Build.cs
+++ b/Engine/Source/Editor/IntroTutorials/IntroTutorials.Build.cs
@@ -34,7 +34,12 @@ namespace UnrealBuildTool.Rules
                     "PlacementMode",
 					"SlateCore",
 					"Settings",
-					"PropertyEditor"
+					"PropertyEditor",
+					"DesktopPlatform",
+					"AssetTools",
+					"SourceCodeAccess",
+					"ContentBrowser",
+					"LevelEditor"
 				}
 			);
 
diff --git a/Engine/Source/Editor/IntroTutorials/Private/EditorTutorialFactory.cpp b/Engine/Source/Editor/IntroTutorials/Private/EditorTutorialFactory.cpp
new file mode 100644
--- /dev/null
+++ b/Engine/Source/Editor/IntroTutorials/Private/EditorTutorialFactory.cpp
@@ -0,0 +1,66 @@
+// Copyright 1998-2014 Epic Games, Inc. All Rights Reserved.
+
+#include "IntroTutorialsPrivatePCH.h"
+#include "EditorTutorialFactory.h"
+#include "EditorTutorial.h"
+#include "KismetEditorUtilities.h"
+#include "SEditorTutorialImporter.h"
+#include "IAssetTypeActions.h"
+
+#define LOCTEXT_NAMESPACE "UEditorTutorialFactory"
+
+UEditorTutorialFactory::UEditorTutorialFactory(const class FPostConstructInitializeProperties& PCIP)
+	: Super(PCIP)
+{
+	bCreateNew = true;
+	bEditAfterNew = true;
+	SupportedClass = UBlueprint::StaticClass();
+}
+
+bool UEditorTutorialFactory::ConfigureProperties()
+{
+	TSharedPtr<SWindow> RootWindow = FGlobalTabmanager::Get()->GetRootWindow();
+	if(RootWindow.IsValid())
+	{
+		// show import dialog
+		TSharedRef<SWindow> Window = SNew(SWindow)
+				.Title(LOCTEXT("WindowTitle", "Import Tutorial"))
+				.SupportsMaximize(false)
+				.SupportsMinimize(false)
+				.SizingRule(ESizingRule::Autosized);	
+
+		ImporterDialog = SNew(SEditorTutorialImporter)
+			.ParentWindow(Window);
+
+		Window->SetContent(ImporterDialog.ToSharedRef());
+
+		FSlateApplication::Get().AddModalWindow(Window, RootWindow.ToSharedRef());
+	}
+
+	return true;
+}
+
+UObject* UEditorTutorialFactory::FactoryCreateNew(UClass* Class, UObject* InParent, FName Name, EObjectFlags Flags, UObject* Context, FFeedbackContext* Warn)
+{
+	UBlueprint* NewBlueprint = FKismetEditorUtilities::CreateBlueprint(UEditorTutorial::StaticClass(), InParent, Name, BPTYPE_Normal, UBlueprint::StaticClass(), UBlueprintGeneratedClass::StaticClass(), NAME_None);
+	
+	if(ImporterDialog.IsValid() && !ImporterDialog->GetImportPath().IsEmpty())
+	{
+		ImporterDialog->Import(CastChecked<UEditorTutorial>(NewBlueprint->GeneratedClass->GetDefaultObject()));
+		ImporterDialog = nullptr;
+	}
+
+	return NewBlueprint;
+}
+
+uint32 UEditorTutorialFactory::GetMenuCategories() const
+{
+	return EAssetTypeCategories::Misc;
+}
+
+FText UEditorTutorialFactory::GetDisplayName() const
+{
+	return LOCTEXT("TutorialMenuEntry", "Tutorial Blueprint");
+}
+
+#undef LOCTEXT_NAMESPACE
\ No newline at end of file
diff --git a/Engine/Source/Editor/IntroTutorials/Private/EditorTutorialFactory.h b/Engine/Source/Editor/IntroTutorials/Private/EditorTutorialFactory.h
new file mode 100644
--- /dev/null
+++ b/Engine/Source/Editor/IntroTutorials/Private/EditorTutorialFactory.h
@@ -0,0 +1,24 @@
+// Copyright 1998-2014 Epic Games, Inc. All Rights Reserved.
+
+#pragma once
+
+#include "EditorTutorialFactory.generated.h"
+
+class SEditorTutorialImporter;
+
+UCLASS(hidecategories=Object, collapsecategories)
+class UEditorTutorialFactory : public UFactory
+{
+	GENERATED_UCLASS_BODY()
+
+	// UFactory Interface
+	virtual bool ConfigureProperties() override;
+	virtual UObject* FactoryCreateNew(UClass* Class, UObject* InParent, FName Name, EObjectFlags Flags, UObject* Context, FFeedbackContext* Warn) override;
+	virtual uint32 GetMenuCategories() const;
+	virtual FText GetDisplayName() const;
+	// End of UFactory Interface
+
+private:
+	/** Importer dialog reference */
+	TSharedPtr<SEditorTutorialImporter> ImporterDialog;
+};
diff --git a/Engine/Source/Editor/IntroTutorials/Private/SEditorTutorialImporter.cpp b/Engine/Source/Editor/IntroTutorials/Private/SEditorTutorialImporter.cpp
new file mode 100644
--- /dev/null
+++ b/Engine/Source/Editor/IntroTutorials/Private/SEditorTutorialImporter.cpp
@@ -0,0 +1,177 @@
+// Copyright 1998-2014 Epic Games, Inc. All Rights Reserved.
+
+#include "IntroTutorialsPrivatePCH.h"
+#include "SEditorTutorialImporter.h"
+#include "IDocumentation.h"
+#include "IDocumentationPage.h"
+#include "EditorTutorial.h"
+#include "DesktopPlatformModule.h"
+
+#define LOCTEXT_NAMESPACE "SEditorTutorialImporter"
+
+void SEditorTutorialImporter::Construct(const FArguments& InArgs)
+{
+	ParentWindow = InArgs._ParentWindow;
+
+	ChildSlot
+	[
+		SNew(SWizard)
+		.CanFinish(true)
+		.ShowPageList(false)
+		.FinishButtonText(LOCTEXT("ImportButtonLabel", "Import"))
+		.FinishButtonToolTip(LOCTEXT("ImportButtonTooltip", "Import the selected UDN file as a tutorial"))
+		.OnFinished(FSimpleDelegate::CreateSP(this, &SEditorTutorialImporter::OnImport))
+		+SWizard::Page()
+		.PageContent()
+		[
+			SNew(SBorder)
+			.BorderImage(FEditorStyle::Get().GetBrush("ToolPanel.GroupBorder"))
+			[
+				SNew(SVerticalBox)
+				+SVerticalBox::Slot()
+				.VAlign(VAlign_Center)
+				[
+					SNew(SVerticalBox)
+					+SVerticalBox::Slot()
+					.AutoHeight()
+					.Padding(5.0f)
+					[
+						SNew(STextBlock)	
+						.Text(LOCTEXT("WizardInstructions", "Choose an engine UDN file to import.\nFiles must be in UE4/Engine/Documentation/Source/."))
+					]
+					+SVerticalBox::Slot()
+					.AutoHeight()
+					.Padding(5.0f)
+					[
+						SNew(SHorizontalBox)
+						+SHorizontalBox::Slot()
+						.FillWidth(1.0f)
+						[
+							SNew(SEditableTextBox)
+							.Text(this, &SEditorTutorialImporter::GetImportPath)
+							.OnTextCommitted(this, &SEditorTutorialImporter::HandleTextCommitted)
+						]
+						+SHorizontalBox::Slot()
+						.AutoWidth()
+						[
+							SNew(SButton)
+							.ButtonStyle( FEditorStyle::Get(), "HoverHintOnly" )
+							.ToolTipText( LOCTEXT( "FileButtonToolTipText", "Choose a file from this computer").ToString() )
+							.OnClicked( FOnClicked::CreateSP(this, &SEditorTutorialImporter::OnPickFile) )
+							.ContentPadding( 2.0f )
+							.ForegroundColor( FSlateColor::UseForeground() )
+							.IsFocusable( false )
+							[
+								SNew( SImage )
+								.Image( FEditorStyle::Get().GetBrush("ExternalImagePicker.PickImageButton") )
+								.ColorAndOpacity( FSlateColor::UseForeground() )
+							]
+						]
+					]
+					+SVerticalBox::Slot()
+					.AutoHeight()
+					.Padding(5.0f)
+					[
+						SNew(STextBlock)	
+						.Text(LOCTEXT("WizardDetails", "Each excerpt of the UDN will form a stage of the imported tutorial."))
+					]
+				]
+			]
+		]
+	];
+}
+
+FText SEditorTutorialImporter::GetImportPath() const
+{
+	return UDNPath;
+}
+
+void SEditorTutorialImporter::Import(UEditorTutorial* InTutorial)
+{
+	if(!UDNPath.IsEmpty())
+	{
+		TSharedRef<IDocumentationPage> Page = IDocumentation::Get()->GetPage(UDNPath.ToString(), TSharedPtr<FParserConfiguration>(), FDocumentationStyle());
+		InTutorial->Title = Page->GetTitle();
+
+		TArray<FExcerpt> Excerpts;
+		Page->GetExcerpts(Excerpts);
+		for(auto& Excerpt : Excerpts)
+		{
+			Page->GetExcerptContent(Excerpt);
+
+			if(Excerpt.RichText.Len() > 0)
+			{
+				FTutorialStage& Stage = InTutorial->Stages[InTutorial->Stages.Add(FTutorialStage())];
+				Stage.Name = *Excerpt.Name;
+				Stage.Content.Type = ETutorialContent::RichText;
+
+				FString RichText;
+				FString* TitleStringPtr = Excerpt.Variables.Find(TEXT("StageTitle"));
+				if(TitleStringPtr != nullptr)
+				{
+					RichText += FString::Printf(TEXT("<TextStyle FontFamily=\"Roboto\" FontSize=\"24\" FontStyle=\"Bold\" FontColor=\"(R=0.0,G=0.0,B=0.0,A=1.0)\">%s</>\n\n"), **TitleStringPtr);
+				}
+				RichText += Excerpt.RichText;
+
+				Stage.Content.Text = FText::FromString(RichText);
+			}
+		}	
+	}
+}
+
+void SEditorTutorialImporter::OnImport()
+{
+	if(ParentWindow.IsValid())
+	{
+		FSlateApplication::Get().RequestDestroyWindow(ParentWindow.Pin().ToSharedRef());
+		ParentWindow.Reset();
+	}
+}
+
+void SEditorTutorialImporter::OnCancel()
+{
+	UDNPath = FText();
+	if(ParentWindow.IsValid())
+	{
+		FSlateApplication::Get().RequestDestroyWindow(ParentWindow.Pin().ToSharedRef());
+		ParentWindow.Reset();
+	}
+}
+
+void SEditorTutorialImporter::HandleTextCommitted(const FText& InText, ETextCommit::Type InCommitType)
+{
+	const FString DocDir = FPaths::ConvertRelativePathToFull(FPaths::EngineDir() / TEXT("Documentation/Source"));
+	FString NewPath = FPaths::ConvertRelativePathToFull(FPaths::GetPath(InText.ToString()));
+	FPaths::NormalizeFilename(NewPath);
+
+	if(NewPath.StartsWith(DocDir))
+	{
+		UDNPath = FText::FromString(NewPath.RightChop(DocDir.Len()));
+	}
+}
+
+FReply SEditorTutorialImporter::OnPickFile()
+{
+	IDesktopPlatform* DesktopPlatform = FDesktopPlatformModule::Get();
+	if ( DesktopPlatform )
+	{
+		TArray<FString> OutFiles;
+		const FString Extension = TEXT("UDN");
+		const FString Filter = FString::Printf(TEXT("%s files (*.%s)|*.%s"), *Extension, *Extension, *Extension);
+		const FString DefaultPath = FPaths::EngineDir() / TEXT("Documentation/Source");
+
+		TSharedPtr<SWindow> ParentWindow = FSlateApplication::Get().FindWidgetWindow(AsShared());
+		void* ParentWindowHandle = (ParentWindow.IsValid() && ParentWindow->GetNativeWindow().IsValid()) ? ParentWindow->GetNativeWindow()->GetOSWindowHandle() : nullptr;
+
+		if (DesktopPlatform->OpenFileDialog(ParentWindowHandle, FText::Format(LOCTEXT("ImagePickerDialogTitle", "Choose a {0} file"), FText::FromString(Extension)).ToString(), DefaultPath, TEXT(""), Filter, EFileDialogFlags::None, OutFiles))
+		{
+			check(OutFiles.Num() == 1);
+
+			HandleTextCommitted(FText::FromString(OutFiles[0]), ETextCommit::Default);
+		}
+	}
+
+	return FReply::Handled();
+}
+
+#undef LOCTEXT_NAMESPACE
\ No newline at end of file
diff --git a/Engine/Source/Editor/IntroTutorials/Private/SEditorTutorialImporter.h b/Engine/Source/Editor/IntroTutorials/Private/SEditorTutorialImporter.h
new file mode 100644
--- /dev/null
+++ b/Engine/Source/Editor/IntroTutorials/Private/SEditorTutorialImporter.h
@@ -0,0 +1,36 @@
+// Copyright 1998-2014 Epic Games, Inc. All Rights Reserved.
+
+#pragma once
+
+class SEditorTutorialImporter : public SCompoundWidget
+{
+public:
+	SLATE_BEGIN_ARGS( SEditorTutorialImporter ) {}
+
+	SLATE_ARGUMENT(TWeakPtr<SWindow>, ParentWindow)
+
+	SLATE_END_ARGS()
+
+	void Construct(const FArguments& InArgs);
+
+	FText GetImportPath() const;
+
+	void Import(UEditorTutorial* InTutorial);
+
+private:
+
+	void OnImport();
+
+	void OnCancel();
+
+	void HandleTextCommitted(const FText& InText, ETextCommit::Type InCommitType);
+
+	FReply OnPickFile();
+
+private:
+	/** The path we will be importing from */
+	FText UDNPath;
+
+	/** Parent window of this widget */
+	TWeakPtr<SWindow> ParentWindow;
+};
diff --git a/Engine/Source/Editor/IntroTutorials/Private/STutorialContent.cpp b/Engine/Source/Editor/IntroTutorials/Private/STutorialContent.cpp
--- a/Engine/Source/Editor/IntroTutorials/Private/STutorialContent.cpp
+++ b/Engine/Source/Editor/IntroTutorials/Private/STutorialContent.cpp
@@ -2,6 +2,9 @@
 
 #include "IntroTutorialsPrivatePCH.h"
 #include "STutorialContent.h"
+#include "STutorialEditableText.h"
+
+#define LOCTEXT_NAMESPACE "STutorialContent"
 
 namespace TutorialConstants
 {
@@ -48,6 +51,7 @@ void STutorialContent::Construct(const FArguments& InArgs, const FTutorialConten
 		.Padding(Anchor.Type != ETutorialAnchorIdentifier::None ? 18.0f : 24.0f)
 		.BorderImage(FEditorStyle::GetBrush("Tutorials.Border"))
 		.Visibility(this, &STutorialContent::GetVisibility)
+		.ForegroundColor(FCoreStyle::Get().GetSlateColor("InvertedForeground"))
 		[
 			SNew(SHorizontalBox)
 			+SHorizontalBox::Slot()
@@ -97,7 +101,7 @@ static void GetAnimationValues(bool bIsIntro, float InAnimationProgress, float&
 
 int32 STutorialContent::OnPaint( const FPaintArgs& Args, const FGeometry& AllottedGeometry, const FSlateRect& MyClippingRect, FSlateWindowElementList& OutDrawElements, int32 LayerId, const FWidgetStyle& InWidgetStyle, bool bParentEnabled ) const
 {
-	if(bIsVisible && Anchor.Type != ETutorialAnchorIdentifier::None)
+	if(bIsVisible && Anchor.Type != ETutorialAnchorIdentifier::None && Anchor.bDrawHighlight)
 	{
 		float AlphaFactor;
 		float PulseFactor;
@@ -174,12 +178,12 @@ TSharedRef<SWidget> STutorialContent::GenerateContentWidget(const FTutorialConte
 		.HyperlinkTextStyle(TEXT("Tutorials.Hyperlink.Text"))
 		.SeparatorStyle(TEXT("Tutorials.Separator"));
 
+	OutDocumentationPage = nullptr;
+
 	switch(InContent.Type)
 	{
 	case ETutorialContent::Text:
 		{
-			OutDocumentationPage = nullptr;
-
 			return SNew(STextBlock)
 				.Visibility(EVisibility::SelfHitTestInvisible)
 				.WrapTextAt(WrapTextAt)
@@ -211,6 +215,37 @@ TSharedRef<SWidget> STutorialContent::GenerateContentWidget(const FTutorialConte
 			}
 		}
 		break;
+
+	case ETutorialContent::RichText:
+		{
+			FTextStyles TextStyles;
+			TextStyles.AvailableFontFamilies.Emplace(MakeShareable(new FTextStyles::FFontFamily(
+				LOCTEXT("RobotoFontFamily", "Roboto"), 
+				TEXT("Roboto"), 
+				FName(*(FPaths::EngineContentDir() / TEXT("Slate/Fonts/Roboto-Regular.ttf"))),
+				FName(*(FPaths::EngineContentDir() / TEXT("Slate/Fonts/Roboto-Bold.ttf"))),
+				FName(*(FPaths::EngineContentDir() / TEXT("Slate/Fonts/Roboto-Italic.ttf"))),
+				FName(*(FPaths::EngineContentDir() / TEXT("Slate/Fonts/Roboto-BoldItalic.ttf")))
+				)));
+
+			TArray< TSharedRef< class ITextDecorator > > Decorators;
+			Decorators.Add(FHyperlinkDecorator::Create(TEXT("browser"), FSlateHyperlinkRun::FOnClick::CreateStatic(&TutorialTextHelpers::OnBrowserLinkClicked)));
+			Decorators.Add(FHyperlinkDecorator::Create(TEXT("udn"), FSlateHyperlinkRun::FOnClick::CreateStatic(&TutorialTextHelpers::OnDocLinkClicked)));
+			Decorators.Add(FHyperlinkDecorator::Create(TEXT("tutorial"), FSlateHyperlinkRun::FOnClick::CreateStatic(&TutorialTextHelpers::OnTutorialLinkClicked)));
+			Decorators.Add(FHyperlinkDecorator::Create(TEXT("code"), FSlateHyperlinkRun::FOnClick::CreateStatic(&TutorialTextHelpers::OnCodeLinkClicked)));
+			Decorators.Add(FHyperlinkDecorator::Create(TEXT("asset"), FSlateHyperlinkRun::FOnClick::CreateStatic(&TutorialTextHelpers::OnAssetLinkClicked)));
+			Decorators.Add(FTextStyleDecorator::Create(&TextStyles));
+
+			return SNew(SRichTextBlock)
+					.TextStyle(FEditorStyle::Get(), "TutorialEditableText.Editor.Text")
+					.DecoratorStyleSet(&FEditorStyle::Get())
+					.Decorators(Decorators)
+					.Text(InContent.Text)
+					.WrapTextAt(WrapTextAt)
+					.Margin(4)
+					.LineHeightPercentage(1.1f);
+		}
+		break;
 	}
 
 	return SNullWidget::NullWidget;
@@ -306,4 +341,6 @@ void STutorialContent::HandleCacheWindowSize(const FVector2D& InWindowSize)
 EVisibility STutorialContent::GetVisibility() const
 {
 	return bIsVisible || Anchor.Type == ETutorialAnchorIdentifier::None ? EVisibility::SelfHitTestInvisible : EVisibility::Collapsed;
-}
\ No newline at end of file
+}
+
+#undef LOCTEXT_NAMESPACE
\ No newline at end of file
diff --git a/Engine/Source/Editor/IntroTutorials/Private/STutorialOverlay.cpp b/Engine/Source/Editor/IntroTutorials/Private/STutorialOverlay.cpp
--- a/Engine/Source/Editor/IntroTutorials/Private/STutorialOverlay.cpp
+++ b/Engine/Source/Editor/IntroTutorials/Private/STutorialOverlay.cpp
@@ -5,6 +5,7 @@
 #include "STutorialContent.h"
 #include "EditorTutorial.h"
 #include "IntroTutorials.h"
+#include "LevelEditor.h"
 
 void STutorialOverlay::Construct(const FArguments& InArgs, FTutorialStage* const InStage)
 {
@@ -53,29 +54,33 @@ void STutorialOverlay::Construct(const FArguments& InArgs, FTutorialStage* const
 			// now add canvas slots for widget-bound content
 			for(const FTutorialWidgetContent& WidgetContent : InStage->WidgetContent)
 			{
-				if(WidgetContent.Content.Type != ETutorialContent::None)
+				if(WidgetContent.Content.Type != ETutorialContent::None )
 				{
-					TSharedPtr<STutorialContent> ContentWidget = 
-						SNew(STutorialContent, WidgetContent.Content)
-						.HAlign(WidgetContent.HorizontalAlignment)
-						.VAlign(WidgetContent.VerticalAlignment)
-						.Offset(WidgetContent.Offset)
-						.IsStandalone(bIsStandalone)
-						.OnClosed(OnClosed)
-						.WrapTextAt(WidgetContent.ContentWidth)
-						.Anchor(WidgetContent.WidgetAnchor);
-					OpenBrowserForWidgetAnchor(WidgetContent);
-
-					OverlayCanvas->AddSlot()
-					.Position(TAttribute<FVector2D>::Create(TAttribute<FVector2D>::FGetter::CreateSP(ContentWidget.Get(), &STutorialContent::GetPosition)))
-					.Size(TAttribute<FVector2D>::Create(TAttribute<FVector2D>::FGetter::CreateSP(ContentWidget.Get(), &STutorialContent::GetSize)))
-					[
-						ContentWidget.ToSharedRef()
-					];
-
-					OnPaintNamedWidget.AddSP(ContentWidget.Get(), &STutorialContent::HandlePaintNamedWidget);
-					OnResetNamedWidget.AddSP(ContentWidget.Get(), &STutorialContent::HandleResetNamedWidget);
-					OnCacheWindowSize.AddSP(ContentWidget.Get(), &STutorialContent::HandleCacheWindowSize);
+					const bool bEmptyText = (WidgetContent.Content.Type == ETutorialContent::Text || WidgetContent.Content.Type == ETutorialContent::RichText) && WidgetContent.Content.Text.IsEmpty();
+					if(!bEmptyText)
+					{
+						TSharedPtr<STutorialContent> ContentWidget = 
+							SNew(STutorialContent, WidgetContent.Content)
+							.HAlign(WidgetContent.HorizontalAlignment)
+							.VAlign(WidgetContent.VerticalAlignment)
+							.Offset(WidgetContent.Offset)
+							.IsStandalone(bIsStandalone)
+							.OnClosed(OnClosed)
+							.WrapTextAt(WidgetContent.ContentWidth)
+							.Anchor(WidgetContent.WidgetAnchor);
+						OpenBrowserForWidgetAnchor(WidgetContent);
+
+						OverlayCanvas->AddSlot()
+						.Position(TAttribute<FVector2D>::Create(TAttribute<FVector2D>::FGetter::CreateSP(ContentWidget.Get(), &STutorialContent::GetPosition)))
+						.Size(TAttribute<FVector2D>::Create(TAttribute<FVector2D>::FGetter::CreateSP(ContentWidget.Get(), &STutorialContent::GetSize)))
+						[
+							ContentWidget.ToSharedRef()
+						];
+
+						OnPaintNamedWidget.AddSP(ContentWidget.Get(), &STutorialContent::HandlePaintNamedWidget);
+						OnResetNamedWidget.AddSP(ContentWidget.Get(), &STutorialContent::HandleResetNamedWidget);
+						OnCacheWindowSize.AddSP(ContentWidget.Get(), &STutorialContent::HandleCacheWindowSize);
+					}
 				}
 			}
 		}
diff --git a/Engine/Source/Editor/IntroTutorials/Private/TutorialStructCustomization.cpp b/Engine/Source/Editor/IntroTutorials/Private/TutorialStructCustomization.cpp
--- a/Engine/Source/Editor/IntroTutorials/Private/TutorialStructCustomization.cpp
+++ b/Engine/Source/Editor/IntroTutorials/Private/TutorialStructCustomization.cpp
@@ -3,6 +3,7 @@
 #include "IntroTutorialsPrivatePCH.h"
 #include "TutorialStructCustomization.h"
 #include "EditorTutorial.h"
+#include "STutorialEditableText.h"
 
 #define LOCTEXT_NAMESPACE "TutorialStructCustomization"
 
@@ -62,6 +63,42 @@ void FTutorialContentCustomization::CustomizeHeader( TSharedRef<class IPropertyH
 
 			return EVisibility::Collapsed;
 		}
+
+		static EVisibility GetRichTextVisibility(TSharedPtr<IPropertyHandle> InPropertyHandle)
+		{
+			check(InPropertyHandle.IsValid());
+
+			uint8 Value = 0;
+			if(InPropertyHandle->GetValue(Value) == FPropertyAccess::Success)
+			{
+				const ETutorialContent::Type EnumValue = (ETutorialContent::Type)Value;
+				return (EnumValue == ETutorialContent::RichText) ? EVisibility::Visible : EVisibility::Collapsed;
+			}
+
+			return EVisibility::Collapsed;
+		}
+
+		static FText GetValueAsText(TSharedPtr<IPropertyHandle> InPropertyHandle)
+		{
+			FText Text;
+
+			if( InPropertyHandle->GetValueAsFormattedText( Text ) == FPropertyAccess::MultipleValues )
+			{
+				Text = NSLOCTEXT("PropertyEditor", "MultipleValues", "Multiple Values");
+			}
+
+			return Text;
+		}
+
+		static void OnTextCommitted( const FText& NewText, ETextCommit::Type /*CommitInfo*/, TSharedPtr<IPropertyHandle> InPropertyHandle )
+		{
+			InPropertyHandle->SetValueFromFormattedString( NewText.ToString() );
+		}
+
+		static void OnTextChanged( const FText& NewText, TSharedPtr<IPropertyHandle> InPropertyHandle )
+		{
+			InPropertyHandle->SetValueFromFormattedString( NewText.ToString() );
+		}
 	};
 
 	HeaderRow
@@ -112,6 +149,20 @@ void FTutorialContentCustomization::CustomizeHeader( TSharedRef<class IPropertyH
 				TextProperty->CreatePropertyValueWidget()
 			]
 		]
+		+SVerticalBox::Slot()
+		.AutoHeight()
+		.Padding(0.0f, 2.0f)
+		[
+			SNew(SHorizontalBox)
+			.Visibility_Static(&Local::GetRichTextVisibility, TypeProperty)
+			+SHorizontalBox::Slot()
+			[
+				SNew(STutorialEditableText)
+				.Text_Static(&Local::GetValueAsText, TextProperty)
+				.OnTextCommitted(FOnTextCommitted::CreateStatic(&Local::OnTextCommitted, TextProperty))
+				.OnTextChanged(FOnTextChanged::CreateStatic(&Local::OnTextChanged, TextProperty))
+			]
+		]
 	];
 }
 
diff --git a/Engine/Source/Editor/IntroTutorials/Public/EditorTutorial.h b/Engine/Source/Editor/IntroTutorials/Public/EditorTutorial.h
--- a/Engine/Source/Editor/IntroTutorials/Public/EditorTutorial.h
+++ b/Engine/Source/Editor/IntroTutorials/Public/EditorTutorial.h
@@ -13,11 +13,14 @@ namespace ETutorialContent
 		/** Blank - displays no content */
 		None,
 
-		/** Text content */
+		/** Plain text content */
 		Text,
 
 		/** Content from a UDN excerpt */
 		UDNExcerpt,
+
+		/** Rich text content */
+		RichText,
 	};
 }
 
@@ -67,6 +70,11 @@ struct INTROTUTORIALS_API FTutorialContent
 {
 	GENERATED_USTRUCT_BODY()
 
+	FTutorialContent()
+	{
+		Type = ETutorialContent::Text;
+	}
+
 	/** The type of this content */
 	UPROPERTY(EditAnywhere, Category="Content")
 	TEnumAsByte<ETutorialContent::Type> Type;
@@ -93,6 +101,7 @@ struct INTROTUTORIALS_API FTutorialContentAnchor
 	FTutorialContentAnchor()
 	{
 		 Type = ETutorialAnchorIdentifier::None;
+		 bDrawHighlight = true;
 	}
 
 	UPROPERTY(EditAnywhere, Category="Anchor")
@@ -105,6 +114,10 @@ struct INTROTUTORIALS_API FTutorialContentAnchor
 	/** If reference is an asset, we use this to resolve it */
 	UPROPERTY(EditAnywhere, Category="Anchor")
 	FStringAssetReference Asset;
+
+	/** Whether to draw an animated highlight around the widget */
+	UPROPERTY(EditAnywhere, Category="Anchor")
+	bool bDrawHighlight;
 };
 
 /** Content that is displayed relative to a widget */
@@ -120,27 +133,27 @@ struct INTROTUTORIALS_API FTutorialWidgetContent
 		VerticalAlignment = VAlign_Bottom;
 	}
 
+	/** Content to associate with widget */
+	UPROPERTY(EditAnywhere, EditInline, Category="Widget")
+	FTutorialContent Content;
+
 	/** Anchor for content widget to highlight */
 	UPROPERTY(EditAnywhere, Category="Widget")
 	FTutorialContentAnchor WidgetAnchor;
 
-	UPROPERTY(EditAnywhere, Category="Widget")
+	UPROPERTY(EditAnywhere, AdvancedDisplay, Category="Widget")
 	TEnumAsByte<EHorizontalAlignment> HorizontalAlignment;
 
-	UPROPERTY(EditAnywhere, Category="Widget")
+	UPROPERTY(EditAnywhere, AdvancedDisplay, Category="Widget")
 	TEnumAsByte<EVerticalAlignment> VerticalAlignment;
 
 	/** Custom offset from widget */
-	UPROPERTY(EditAnywhere, Category="Widget")
+	UPROPERTY(EditAnywhere, AdvancedDisplay, Category="Widget")
 	FVector2D Offset;
 
 	/** Content width - text will be wrapped at this point */
-	UPROPERTY(EditAnywhere, Category="Widget", meta=(UIMin="10.0", UIMax="600.0"))
+	UPROPERTY(EditAnywhere, AdvancedDisplay, Category="Widget", meta=(UIMin="10.0", UIMax="600.0"))
 	float ContentWidth;
-
-	/** Content to associate with widget */
-	UPROPERTY(EditAnywhere, EditInline, Category="Widget")
-	FTutorialContent Content;
 };
 
 /** A single tutorial stage, containing the optional main content & a number of widgets with content attached */
