Added tutorials button in top bar of editor & sub-editors

Icon is only visible if content is available for the editor in question.
Split editor settings into two groups - one is persistent settings and one is progress/state.
Tutorials record their dismissed state, so users can permenantly disable the 'nag' for a particular tutorial.
Tutorial content now solidifies when the mouse is hovered over it, so it can be made easier to read.
Fixed crash on startup if an intro tutorial was displaying rich text.
Also fixed crash for TTP# 345094, where a zero-length tutorial was being accessed.

diff --git a/Engine/Source/Editor/EditorStyle/Private/SlateEditorStyle.cpp b/Engine/Source/Editor/EditorStyle/Private/SlateEditorStyle.cpp
--- a/Engine/Source/Editor/EditorStyle/Private/SlateEditorStyle.cpp
+++ b/Engine/Source/Editor/EditorStyle/Private/SlateEditorStyle.cpp
@@ -2395,7 +2395,7 @@ void FSlateEditorStyle::FStyle::SetupTutorialStyles()
 		const FTextBlockStyle TutorialHeaderText = FTextBlockStyle(DocumentationHeaderText)
 			.SetColorAndOpacity(FLinearColor::Black);
 
-		Set( "Tutorials.Border", new BOX_BRUSH( "Tutorials/OverlayFrame", FMargin(18.0f/64.0f), FLinearColor(1.0f, 1.0f, 1.0f, 0.8f) ) );
+		Set( "Tutorials.Border", new BOX_BRUSH( "Tutorials/OverlayFrame", FMargin(18.0f/64.0f), FLinearColor(1.0f, 1.0f, 1.0f, 1.0f) ) );
 		Set( "Tutorials.Browser.WelcomeHeader", FTextBlockStyle(TutorialHeaderText)
 			.SetFontSize(20));
 
@@ -2500,84 +2500,98 @@ void FSlateEditorStyle::FStyle::SetupTutorialStyles()
 			const FLinearColor DisabledColor(FColor(0xaaaaaa));
 			const FLinearColor TextColor(FColor(0xff2c3e50));
 
-			Set( "TutorialEditableText.RoundedBackground", new BOX_BRUSH( "Common/RoundedSelection_16x", 4.0f/16.0f, FLinearColor(FColor(0xffeff3f3)) ) );
+			Set("TutorialEditableText.RoundedBackground", new BOX_BRUSH("Common/RoundedSelection_16x", 4.0f / 16.0f, FLinearColor(FColor(0xffeff3f3))));
 
-			Set( "TutorialEditableText.Toolbar.HyperlinkImage", new IMAGE_BRUSH( "Tutorials/hyperlink", Icon16x16, TextColor ) );
+			Set("TutorialEditableText.Toolbar.HyperlinkImage", new IMAGE_BRUSH("Tutorials/hyperlink", Icon16x16, TextColor));
 
-			Set( "TutorialEditableText.Toolbar.TextColor", TextColor );
+			Set("TutorialEditableText.Toolbar.TextColor", TextColor);
 
-			Set( "TutorialEditableText.Toolbar.Text", FTextBlockStyle(NormalText)
-				.SetFont( TTF_FONT( "Fonts/Roboto-Regular", 10 ) )
-				.SetColorAndOpacity( TextColor )
+			Set("TutorialEditableText.Toolbar.Text", FTextBlockStyle(NormalText)
+				.SetFont(TTF_FONT("Fonts/Roboto-Regular", 10))
+				.SetColorAndOpacity(TextColor)
 				);
 
-			Set( "TutorialEditableText.Toolbar.BoldText", FTextBlockStyle(NormalText)
-				.SetFont( TTF_FONT( "Fonts/Roboto-Bold", 10 ) )
-				.SetColorAndOpacity( TextColor )
+			Set("TutorialEditableText.Toolbar.BoldText", FTextBlockStyle(NormalText)
+				.SetFont(TTF_FONT("Fonts/Roboto-Bold", 10))
+				.SetColorAndOpacity(TextColor)
 				);
 
-			Set( "TutorialEditableText.Toolbar.ItalicText", FTextBlockStyle(NormalText)
-				.SetFont( TTF_FONT( "Fonts/Roboto-Italic", 10 ) )
-				.SetColorAndOpacity( TextColor )
+			Set("TutorialEditableText.Toolbar.ItalicText", FTextBlockStyle(NormalText)
+				.SetFont(TTF_FONT("Fonts/Roboto-Italic", 10))
+				.SetColorAndOpacity(TextColor)
 				);
 
-			Set( "TutorialEditableText.Toolbar.Checkbox", FCheckBoxStyle()
+			Set("TutorialEditableText.Toolbar.Checkbox", FCheckBoxStyle()
 				.SetCheckBoxType(ESlateCheckBoxType::CheckBox)
-				.SetUncheckedImage( IMAGE_BRUSH( "Common/CheckBox", Icon16x16, FLinearColor::White ) )
-				.SetUncheckedHoveredImage( IMAGE_BRUSH( "Common/CheckBox", Icon16x16, HoverColor ) )
-				.SetUncheckedPressedImage( IMAGE_BRUSH( "Common/CheckBox_Hovered", Icon16x16, HoverColor ) )
-				.SetCheckedImage( IMAGE_BRUSH( "Common/CheckBox_Checked_Hovered", Icon16x16, FLinearColor::White ) )
-				.SetCheckedHoveredImage( IMAGE_BRUSH( "Common/CheckBox_Checked_Hovered", Icon16x16, HoverColor ) )
-				.SetCheckedPressedImage( IMAGE_BRUSH( "Common/CheckBox_Checked", Icon16x16, HoverColor ) )
-				.SetUndeterminedImage( IMAGE_BRUSH( "Common/CheckBox_Undetermined", Icon16x16, FLinearColor::White ) )
-				.SetUndeterminedHoveredImage( IMAGE_BRUSH( "Common/CheckBox_Undetermined_Hovered", Icon16x16, HoverColor ) )
-				.SetUndeterminedPressedImage( IMAGE_BRUSH( "Common/CheckBox_Undetermined_Hovered", Icon16x16, FLinearColor::White ) )
+				.SetUncheckedImage(IMAGE_BRUSH("Common/CheckBox", Icon16x16, FLinearColor::White))
+				.SetUncheckedHoveredImage(IMAGE_BRUSH("Common/CheckBox", Icon16x16, HoverColor))
+				.SetUncheckedPressedImage(IMAGE_BRUSH("Common/CheckBox_Hovered", Icon16x16, HoverColor))
+				.SetCheckedImage(IMAGE_BRUSH("Common/CheckBox_Checked_Hovered", Icon16x16, FLinearColor::White))
+				.SetCheckedHoveredImage(IMAGE_BRUSH("Common/CheckBox_Checked_Hovered", Icon16x16, HoverColor))
+				.SetCheckedPressedImage(IMAGE_BRUSH("Common/CheckBox_Checked", Icon16x16, HoverColor))
+				.SetUndeterminedImage(IMAGE_BRUSH("Common/CheckBox_Undetermined", Icon16x16, FLinearColor::White))
+				.SetUndeterminedHoveredImage(IMAGE_BRUSH("Common/CheckBox_Undetermined_Hovered", Icon16x16, HoverColor))
+				.SetUndeterminedPressedImage(IMAGE_BRUSH("Common/CheckBox_Undetermined_Hovered", Icon16x16, FLinearColor::White))
 				);
 
-			Set( "TutorialEditableText.Toolbar.ToggleButtonCheckbox", FCheckBoxStyle()
+			Set("TutorialEditableText.Toolbar.ToggleButtonCheckbox", FCheckBoxStyle()
 				.SetCheckBoxType(ESlateCheckBoxType::ToggleButton)
-				.SetUncheckedImage( BOX_BRUSH("Tutorials/FlatColorSquare", FVector2D(1.0f, 1.0f), FMargin(1), NormalColor ) )
-				.SetUncheckedHoveredImage( BOX_BRUSH("Tutorials/FlatColorSquare", FVector2D(1.0f, 1.0f), FMargin(1), HoverColor ) )
-				.SetUncheckedPressedImage( BOX_BRUSH("Tutorials/FlatColorSquare", FVector2D(1.0f, 1.0f), FMargin(1), HoverColor ) )
-				.SetCheckedImage( BOX_BRUSH("Tutorials/FlatColorSquare",  FVector2D(1.0f, 1.0f), FMargin(1), SelectedColor ) )
-				.SetCheckedHoveredImage( BOX_BRUSH("Tutorials/FlatColorSquare",  FVector2D(1.0f, 1.0f), FMargin(1), HoverColor ) )
-				.SetCheckedPressedImage( BOX_BRUSH("Tutorials/FlatColorSquare",  FVector2D(1.0f, 1.0f), FMargin(1), HoverColor ) )
+				.SetUncheckedImage(BOX_BRUSH("Tutorials/FlatColorSquare", FVector2D(1.0f, 1.0f), FMargin(1), NormalColor))
+				.SetUncheckedHoveredImage(BOX_BRUSH("Tutorials/FlatColorSquare", FVector2D(1.0f, 1.0f), FMargin(1), HoverColor))
+				.SetUncheckedPressedImage(BOX_BRUSH("Tutorials/FlatColorSquare", FVector2D(1.0f, 1.0f), FMargin(1), HoverColor))
+				.SetCheckedImage(BOX_BRUSH("Tutorials/FlatColorSquare", FVector2D(1.0f, 1.0f), FMargin(1), SelectedColor))
+				.SetCheckedHoveredImage(BOX_BRUSH("Tutorials/FlatColorSquare", FVector2D(1.0f, 1.0f), FMargin(1), HoverColor))
+				.SetCheckedPressedImage(BOX_BRUSH("Tutorials/FlatColorSquare", FVector2D(1.0f, 1.0f), FMargin(1), HoverColor))
 				);
 
 			const FButtonStyle Button = FButtonStyle()
-				.SetNormal( BOX_BRUSH("Tutorials/FlatColorSquare", FVector2D(1.0f, 1.0f), FMargin(1), NormalColor ) )
-				.SetHovered( BOX_BRUSH("Tutorials/FlatColorSquare", FVector2D(1.0f, 1.0f), FMargin(1), HoverColor ) )
-				.SetPressed( BOX_BRUSH("Tutorials/FlatColorSquare",  FVector2D(1.0f, 1.0f), FMargin(1), SelectedColor ) )
-				.SetDisabled( BOX_BRUSH("Tutorials/FlatColorSquare",  FVector2D(1.0f, 1.0f), FMargin(1), DisabledColor ) )
-				.SetNormalPadding( FMargin( 2,2,2,2 ) )
-				.SetPressedPadding( FMargin( 2,3,2,1 ) );
-			Set( "TutorialEditableText.Toolbar.Button", Button );
+				.SetNormal(BOX_BRUSH("Tutorials/FlatColorSquare", FVector2D(1.0f, 1.0f), FMargin(1), NormalColor))
+				.SetHovered(BOX_BRUSH("Tutorials/FlatColorSquare", FVector2D(1.0f, 1.0f), FMargin(1), HoverColor))
+				.SetPressed(BOX_BRUSH("Tutorials/FlatColorSquare", FVector2D(1.0f, 1.0f), FMargin(1), SelectedColor))
+				.SetDisabled(BOX_BRUSH("Tutorials/FlatColorSquare", FVector2D(1.0f, 1.0f), FMargin(1), DisabledColor))
+				.SetNormalPadding(FMargin(2, 2, 2, 2))
+				.SetPressedPadding(FMargin(2, 3, 2, 1));
+			Set("TutorialEditableText.Toolbar.Button", Button);
 
 			const FComboButtonStyle ComboButton = FComboButtonStyle()
 				.SetButtonStyle(Button)
 				.SetDownArrowImage(IMAGE_BRUSH("Common/ComboArrow", Icon8x8))
-				.SetMenuBorderBrush(BOX_BRUSH( "Tutorials/FlatColorSquare", FVector2D(1.0f, 1.0f), FMargin(1), NormalColor ))
+				.SetMenuBorderBrush(BOX_BRUSH("Tutorials/FlatColorSquare", FVector2D(1.0f, 1.0f), FMargin(1), NormalColor))
 				.SetMenuBorderPadding(FMargin(0.0f));
-			Set( "TutorialEditableText.Toolbar.ComboButton", ComboButton	);
+			Set("TutorialEditableText.Toolbar.ComboButton", ComboButton);
 
 			{
 				const FButtonStyle ComboBoxButton = FButtonStyle()
-					.SetNormal( BOX_BRUSH("Tutorials/FlatColorSquare", FVector2D(1.0f, 1.0f), FMargin(1), FLinearColor::White ) )
-					.SetHovered( BOX_BRUSH("Tutorials/FlatColorSquare", FVector2D(1.0f, 1.0f), FMargin(1), FLinearColor::White ) )
-					.SetPressed( BOX_BRUSH("Tutorials/FlatColorSquare",  FVector2D(1.0f, 1.0f), FMargin(1), FLinearColor::White ) )
-					.SetDisabled( BOX_BRUSH("Tutorials/FlatColorSquare",  FVector2D(1.0f, 1.0f), FMargin(1), DisabledColor ) )
-					.SetNormalPadding( FMargin( 2,2,2,2 ) )
-					.SetPressedPadding( FMargin( 2,3,2,1 ) );
+					.SetNormal(BOX_BRUSH("Tutorials/FlatColorSquare", FVector2D(1.0f, 1.0f), FMargin(1), FLinearColor::White))
+					.SetHovered(BOX_BRUSH("Tutorials/FlatColorSquare", FVector2D(1.0f, 1.0f), FMargin(1), FLinearColor::White))
+					.SetPressed(BOX_BRUSH("Tutorials/FlatColorSquare", FVector2D(1.0f, 1.0f), FMargin(1), FLinearColor::White))
+					.SetDisabled(BOX_BRUSH("Tutorials/FlatColorSquare", FVector2D(1.0f, 1.0f), FMargin(1), DisabledColor))
+					.SetNormalPadding(FMargin(2, 2, 2, 2))
+					.SetPressedPadding(FMargin(2, 3, 2, 1));
 
 				const FComboButtonStyle ComboBoxComboButton = FComboButtonStyle(ComboButton)
 					.SetButtonStyle(ComboBoxButton)
 					.SetMenuBorderPadding(FMargin(1.0));
 
-				Set( "TutorialEditableText.Toolbar.ComboBox", FComboBoxStyle()
+				Set("TutorialEditableText.Toolbar.ComboBox", FComboBoxStyle()
 					.SetComboButtonStyle(ComboBoxComboButton)
 					);
 			}
 		}
+
+		// In-editor tutorial launch button
+		{
+			Set("TutorialLaunch.Button", FButtonStyle()
+				.SetNormalPadding(0)
+				.SetPressedPadding(0)
+				.SetNormal(IMAGE_BRUSH("Tutorials/TutorialButton_Default_16x", Icon16x16))
+				.SetHovered(IMAGE_BRUSH("Tutorials/TutorialButton_Hovered_16x", Icon16x16))
+				.SetPressed(IMAGE_BRUSH("Tutorials/TutorialButton_Pressed_16x", Icon16x16))
+				);
+
+			Set("TutorialLaunch.Circle", new IMAGE_BRUSH("Tutorials/Circle_128x", Icon128x128, FLinearColor::White));
+			Set("TutorialLaunch.Circle.Color", FLinearColor::Green);
+		}
 	}
 }
 
diff --git a/Engine/Source/Editor/IntroTutorials/Private/EditorTutorialSettings.cpp b/Engine/Source/Editor/IntroTutorials/Private/EditorTutorialSettings.cpp
--- a/Engine/Source/Editor/IntroTutorials/Private/EditorTutorialSettings.cpp
+++ b/Engine/Source/Editor/IntroTutorials/Private/EditorTutorialSettings.cpp
@@ -9,67 +9,25 @@ UEditorTutorialSettings::UEditorTutorialSettings(const class FPostConstructIniti
 
 }
 
-void UEditorTutorialSettings::PostInitProperties()
+void UEditorTutorialSettings::FindTutorialsForContext(FName InContext, UEditorTutorial*& OutAttractTutorial, UEditorTutorial*& OutLaunchTutorial) const
 {
-	Super::PostInitProperties();
-
-	for(const auto& Progress : TutorialsProgress)
+	for (const auto& Context : TutorialContexts)
 	{
-		TSubclassOf<UEditorTutorial> TutorialClass = LoadClass<UEditorTutorial>(NULL, *Progress.Tutorial.AssetLongPathname, NULL, LOAD_None, NULL);
-		if(TutorialClass != nullptr)
+		if (Context.Context == InContext)
 		{
-			UEditorTutorial* Tutorial = TutorialClass->GetDefaultObject<UEditorTutorial>();
-			ProgressMap.Add(Tutorial, Progress);
+			TSubclassOf<UEditorTutorial> LaunchTutorialClass = LoadClass<UEditorTutorial>(nullptr, *Context.LaunchTutorial.AssetLongPathname, nullptr, LOAD_None, nullptr);
+			if (LaunchTutorialClass != nullptr)
+			{
+				OutLaunchTutorial = LaunchTutorialClass->GetDefaultObject<UEditorTutorial>();
+			}
+
+			TSubclassOf<UEditorTutorial> AttractTutorialClass = LoadClass<UEditorTutorial>(nullptr, *Context.AttractTutorial.AssetLongPathname, nullptr, LOAD_None, nullptr);
+			if (AttractTutorialClass != nullptr)
+			{
+				OutAttractTutorial = AttractTutorialClass->GetDefaultObject<UEditorTutorial>();
+			}
+
+			return;
 		}
 	}
-}
-
-int32 UEditorTutorialSettings::GetProgress(UEditorTutorial* InTutorial, bool& bOutHaveSeenTutorial) const
-{
-	const FTutorialProgress* FoundProgress = ProgressMap.Find(InTutorial);
-	if(FoundProgress != nullptr)
-	{
-		bOutHaveSeenTutorial = true;
-		return FoundProgress->CurrentStage;
-	}
-
-	bOutHaveSeenTutorial = false;
-	return 0;
-}
-
-bool UEditorTutorialSettings::HaveSeenTutorial(UEditorTutorial* InTutorial) const
-{
-	return ProgressMap.Find(InTutorial) != nullptr;
-}
-
-void UEditorTutorialSettings::RecordProgress(UEditorTutorial* InTutorial, int32 CurrentStage)
-{
-	if(InTutorial != nullptr)
-	{
-		FTutorialProgress* FoundProgress = ProgressMap.Find(InTutorial);
-		if(FoundProgress != nullptr)
-		{
-			FoundProgress->CurrentStage = FMath::Max(FoundProgress->CurrentStage, CurrentStage);
-		}
-		else
-		{
-			FTutorialProgress Progress;
-			Progress.Tutorial = FStringClassReference(InTutorial->GetClass());
-			Progress.CurrentStage = CurrentStage;
-
-			ProgressMap.Add(InTutorial, Progress);
-		}
-	}
-}
-
-void UEditorTutorialSettings::SaveProgress()
-{
-	TutorialsProgress.Empty();
-
-	for(const auto& ProgressMapEntry : ProgressMap)
-	{
-		TutorialsProgress.Add(ProgressMapEntry.Value);
-	}
-
-	SaveConfig();
 }
\ No newline at end of file
diff --git a/Engine/Source/Editor/IntroTutorials/Private/EditorTutorialSettings.h b/Engine/Source/Editor/IntroTutorials/Private/EditorTutorialSettings.h
--- a/Engine/Source/Editor/IntroTutorials/Private/EditorTutorialSettings.h
+++ b/Engine/Source/Editor/IntroTutorials/Private/EditorTutorialSettings.h
@@ -5,17 +5,23 @@
 #include "EditorTutorial.h"
 #include "EditorTutorialSettings.generated.h"
 
-/** Track the progress of an individual tutorial */
+/** Named context that corresponds to a particular tutorial */
 USTRUCT()
-struct FTutorialProgress
+struct FTutorialContext
 {
 	GENERATED_USTRUCT_BODY()
 
-	UPROPERTY()
-	FStringClassReference Tutorial;
+	/** The context that this tutorial is used in */
+	UPROPERTY(EditAnywhere, Category = "Tutorials")
+	FName Context;
 
-	UPROPERTY()
-	int32 CurrentStage;
+	/** The tutorial to use in this context to let the user know there is a tutorial available */
+	UPROPERTY(EditAnywhere, Category = "Tutorials", meta = (MetaClass = "EditorTutorial"))
+	FStringClassReference AttractTutorial;
+
+	/** The tutorial to use in this context when the user chooses to launch */
+	UPROPERTY(EditAnywhere, Category = "Tutorials", meta = (MetaClass = "EditorTutorial"))
+	FStringClassReference LaunchTutorial;
 };
 
 /** Editor-wide tutorial settings */
@@ -28,28 +34,14 @@ class UEditorTutorialSettings : public UObject
 	UPROPERTY(Config, EditAnywhere, Category="Tutorials")
 	TArray<FTutorialCategory> Categories;
 
-	UPROPERTY(Config)
-	TArray<FTutorialProgress> TutorialsProgress;
-
 	/** Tutorial to start on Editor startup */
 	UPROPERTY(Config, EditAnywhere, Category="Tutorials", meta=(MetaClass="EditorTutorial"))
 	FStringClassReference StartupTutorial;
 
-	/** UObject interface */
-	virtual void PostInitProperties() override;
-
-	/** Get the recorded progress of the pass-in tutorial */
-	int32 GetProgress(UEditorTutorial* InTutorial, bool& bOutHaveSeenTutorial) const;
-
-	/** Check if we have seen the passed-in tutorial before */
-	bool HaveSeenTutorial(UEditorTutorial* InTutorial) const;
-
-	/** Record the progress of the passed-in tutorial */
-	void RecordProgress(UEditorTutorial* InTutorial, int32 CurrentStage);
-
-	/** Save the progress of all our tutorials */
-	void SaveProgress();
+	/** Tutorials used in various contexts - e.g. the various asset editors */
+	UPROPERTY(Config, EditAnywhere, Category = "Tutorials")
+	TArray<FTutorialContext> TutorialContexts;
 
-	/** Recorded progress */
-	TMap<UEditorTutorial*, FTutorialProgress> ProgressMap;
+	/** Get the tutorial for the specified context */
+	void FindTutorialsForContext(FName InContext, UEditorTutorial*& OutAttractTutorial, UEditorTutorial*& OutLaunchTutorial) const;
 };
\ No newline at end of file
diff --git a/Engine/Source/Editor/IntroTutorials/Private/IntroTutorials.cpp b/Engine/Source/Editor/IntroTutorials/Private/IntroTutorials.cpp
--- a/Engine/Source/Editor/IntroTutorials/Private/IntroTutorials.cpp
+++ b/Engine/Source/Editor/IntroTutorials/Private/IntroTutorials.cpp
@@ -16,6 +16,7 @@
 #include "Runtime/Analytics/Analytics/Public/Interfaces/IAnalyticsProvider.h"
 #include "Particles/ParticleSystem.h"
 #include "EditorTutorialSettings.h"
+#include "TutorialStateSettings.h"
 #include "TutorialSettings.h"
 #include "Settings.h"
 #include "EditorTutorial.h"
@@ -26,6 +27,7 @@
 #include "TutorialStructCustomization.h"
 #include "EditorTutorialDetailsCustomization.h"
 #include "STutorialRoot.h"
+#include "STutorialButton.h"
 
 #define LOCTEXT_NAMESPACE "IntroTutorials"
 
@@ -555,7 +557,7 @@ bool FIntroTutorials::MaybeOpenWelcomeTutorial(const FString& TutorialPath, cons
 		if(EditorStartupTutorialClass != nullptr)
 		{
 			UEditorTutorial* Tutorial = EditorStartupTutorialClass->GetDefaultObject<UEditorTutorial>();
-			if(!GetDefault<UEditorTutorialSettings>()->HaveSeenTutorial(Tutorial))
+			if (!GetDefault<UTutorialStateSettings>()->HaveSeenTutorial(Tutorial))
 			{
 				LaunchTutorial(Tutorial);
 				return true;
@@ -567,7 +569,7 @@ bool FIntroTutorials::MaybeOpenWelcomeTutorial(const FString& TutorialPath, cons
 		if(ProjectStartupTutorialClass != nullptr)
 		{
 			UEditorTutorial* Tutorial = ProjectStartupTutorialClass->GetDefaultObject<UEditorTutorial>();
-			if(!GetDefault<UEditorTutorialSettings>()->HaveSeenTutorial(Tutorial))
+			if (!GetDefault<UTutorialStateSettings>()->HaveSeenTutorial(Tutorial))
 			{
 				LaunchTutorial(Tutorial);
 				return true;
@@ -753,11 +755,19 @@ void FIntroTutorials::SummonTutorialBrowser(TSharedRef<SWindow> InWindow, const
 	}
 }
 
-void FIntroTutorials::LaunchTutorial(UEditorTutorial* InTutorial, bool bInRestart, TWeakPtr<SWindow> InNavigationWindow)
+void FIntroTutorials::LaunchTutorial(UEditorTutorial* InTutorial, bool bInRestart, TWeakPtr<SWindow> InNavigationWindow, FSimpleDelegate OnTutorialClosed, FSimpleDelegate OnTutorialExited)
 {
 	if(TutorialRoot.IsValid())
 	{
-		TutorialRoot->LaunchTutorial(InTutorial, bInRestart, InNavigationWindow);
+		TutorialRoot->LaunchTutorial(InTutorial, bInRestart, InNavigationWindow, OnTutorialClosed, OnTutorialExited);
+	}
+}
+
+void FIntroTutorials::CloseAllTutorialContent()
+{
+	if (TutorialRoot.IsValid())
+	{
+		TutorialRoot->CloseAllTutorialContent();
 	}
 }
 
@@ -777,5 +787,11 @@ void FIntroTutorials::GoToNextStage(TWeakPtr<SWindow> InNavigationWindow)
 	}
 }
 
+TSharedRef<SWidget> FIntroTutorials::CreateTutorialsWidget(FName InContext, TWeakPtr<SWindow> InContextWindow) const
+{
+	return SNew(STutorialButton)
+		.Context(InContext)
+		.ContextWindow(InContextWindow);
+}
 
 #undef LOCTEXT_NAMESPACE
\ No newline at end of file
diff --git a/Engine/Source/Editor/IntroTutorials/Private/IntroTutorials.h b/Engine/Source/Editor/IntroTutorials/Private/IntroTutorials.h
--- a/Engine/Source/Editor/IntroTutorials/Private/IntroTutorials.h
+++ b/Engine/Source/Editor/IntroTutorials/Private/IntroTutorials.h
@@ -149,7 +149,9 @@ public:
 	// IIntroTutorials interface
 	virtual void RegisterTutorialForAssetEditor(UClass* AssetClass, const FString& TutorialDocPath, const FString& TutorialHasBeenSeenSettingName, const FString& SurveyGUIDString) override;
 	virtual void UnregisterTutorialForAssetEditor(UClass* AssetClass) override;
-	virtual void LaunchTutorial(UEditorTutorial* Tutorial, bool bRestart = true, TWeakPtr<SWindow> InNavigationWindow = nullptr) override;
+	virtual void LaunchTutorial(UEditorTutorial* Tutorial, bool bRestart = true, TWeakPtr<SWindow> InNavigationWindow = nullptr, FSimpleDelegate OnTutorialClosed = FSimpleDelegate(), FSimpleDelegate OnTutorialExited = FSimpleDelegate()) override;
+	virtual void CloseAllTutorialContent() override;
+	virtual TSharedRef<SWidget> CreateTutorialsWidget(FName InContext, TWeakPtr<SWindow> InContextWindow) const override;
 	// End of IIntroTutorials interface
 private:
 	/** The tab id of the tutorial tab */
diff --git a/Engine/Source/Editor/IntroTutorials/Private/SEditorTutorials.cpp b/Engine/Source/Editor/IntroTutorials/Private/SEditorTutorials.cpp
--- a/Engine/Source/Editor/IntroTutorials/Private/SEditorTutorials.cpp
+++ b/Engine/Source/Editor/IntroTutorials/Private/SEditorTutorials.cpp
@@ -79,9 +79,11 @@ void SEditorTutorials::Construct(const FArguments& InArgs)
 	];
 }
 
-void SEditorTutorials::LaunchTutorial(bool bInIsNavigationWindow)
+void SEditorTutorials::LaunchTutorial(bool bInIsNavigationWindow, FSimpleDelegate InOnTutorialClosed, FSimpleDelegate InOnTutorialExited)
 {
 	bShowNavigation = bInIsNavigationWindow;
+	OnTutorialClosed = InOnTutorialClosed;
+	OnTutorialExited = InOnTutorialExited;
 
 	RebuildCurrentContent();
 }
@@ -93,6 +95,13 @@ void SEditorTutorials::ShowBrowser(const FString& InFilter)
 	HandleHomeClicked();
 }
 
+void SEditorTutorials::HideContent()
+{
+	HandleHomeClicked();
+	bBrowserVisible = false;
+	bShowNavigation = false;
+}
+
 bool SEditorTutorials::IsNavigationVisible() const
 {
 	return bShowNavigation;
@@ -114,6 +123,7 @@ void SEditorTutorials::HandleCloseClicked()
 {
 	HandleHomeClicked();
 	bBrowserVisible = false;
+	OnTutorialClosed.ExecuteIfBound();
 }
 
 void SEditorTutorials::HandleBackClicked()
@@ -128,6 +138,7 @@ void SEditorTutorials::HandleHomeClicked()
 
 	ContentBox->ClearChildren();
 	bBrowserVisible = true;
+	OnTutorialExited.ExecuteIfBound();
 }
 
 void SEditorTutorials::HandleNextClicked()
diff --git a/Engine/Source/Editor/IntroTutorials/Private/SEditorTutorials.h b/Engine/Source/Editor/IntroTutorials/Private/SEditorTutorials.h
--- a/Engine/Source/Editor/IntroTutorials/Private/SEditorTutorials.h
+++ b/Engine/Source/Editor/IntroTutorials/Private/SEditorTutorials.h
@@ -40,11 +40,14 @@ class SEditorTutorials : public SCompoundWidget
 	void Construct(const FArguments& InArgs);
 
 	/** Launch a tutorial - will interrogate parent to get the tutorial data to display */
-	void LaunchTutorial(bool bInIsNavigationWindow);
+	void LaunchTutorial(bool bInIsNavigationWindow, FSimpleDelegate InOnTutorialClosed, FSimpleDelegate InOnTutorialExited);
 
 	/** Show the tutorials browser in this window */
 	void ShowBrowser(const FString& InFilter);
 
+	/** hide the tutorials browser in this window */
+	void HideContent();
+
 	/** Get parent window for this widget */
 	TSharedPtr<SWindow> GetParentWindow() const;
 
@@ -121,4 +124,8 @@ private:
 
 	/** Delegate fired to retrive the current tutorial stage */
 	FOnGetCurrentTutorialStage OnGetCurrentTutorialStage;
+
+	/** External delegates used to report user interaction */
+	FSimpleDelegate OnTutorialClosed;
+	FSimpleDelegate OnTutorialExited;
 };
\ No newline at end of file
diff --git a/Engine/Source/Editor/IntroTutorials/Private/STutorialButton.cpp b/Engine/Source/Editor/IntroTutorials/Private/STutorialButton.cpp
new file mode 100644
--- /dev/null
+++ b/Engine/Source/Editor/IntroTutorials/Private/STutorialButton.cpp
@@ -0,0 +1,183 @@
+// Copyright 1998-2014 Epic Games, Inc. All Rights Reserved.
+
+#include "IntroTutorialsPrivatePCH.h"
+#include "STutorialButton.h"
+#include "EditorTutorialSettings.h"
+#include "TutorialStateSettings.h"
+
+#define LOCTEXT_NAMESPACE "STutorialButton"
+
+namespace TutorialButtonConstants
+{
+	const float MaxPulseOffset = 32.0f;
+	const float PulseAnimationLength = 2.0f;
+}
+
+void STutorialButton::Construct(const FArguments& InArgs)
+{
+	Context = InArgs._Context;
+	ContextWindow = InArgs._ContextWindow;
+
+	bTutorialAvailable = false;
+	bTutorialCompleted = false;
+	bTutorialDismissed = false;
+	bDeferTutorialOpen = true;
+
+	PulseAnimation.AddCurve(0.0f, TutorialButtonConstants::PulseAnimationLength, ECurveEaseFunction::Linear);
+	PulseAnimation.Play();
+
+	ChildSlot
+	[
+		SNew(SButton)
+		.Tag(*FString::Printf(TEXT("%s.TutorialLaunchButton"), *Context.ToString()))
+		.ButtonStyle(FEditorStyle::Get(), "TutorialLaunch.Button")
+		.ToolTipText(LOCTEXT("TutorialLaunchToolTip", "Launch Tutorial (right-click for more options)"))
+		.Visibility(this, &STutorialButton::GetVisibility)
+		.OnClicked(this, &STutorialButton::HandleButtonClicked)
+		.ContentPadding(0.0f)
+		[
+			SNew(SBox)
+			.WidthOverride(16)
+			.HeightOverride(16)
+		]
+	];
+}
+
+void STutorialButton::Tick(const FGeometry& AllottedGeometry, const double InCurrentTime, const float InDeltaTime)
+{
+	if (bDeferTutorialOpen)
+	{
+		UEditorTutorial* AttractTutorial = nullptr;
+		UEditorTutorial* LaunchTutorial = nullptr;
+		GetDefault<UEditorTutorialSettings>()->FindTutorialsForContext(Context, AttractTutorial, LaunchTutorial);
+
+		bTutorialAvailable = (LaunchTutorial != nullptr);
+		bTutorialCompleted = (LaunchTutorial != nullptr) && GetDefault<UTutorialStateSettings>()->HaveCompletedTutorial(LaunchTutorial);
+		bTutorialDismissed = (AttractTutorial != nullptr) && GetDefault<UTutorialStateSettings>()->IsTutorialDismissed(AttractTutorial);
+		if (bTutorialAvailable && AttractTutorial != nullptr && !bTutorialDismissed && !bTutorialCompleted)
+		{
+			// kick off the attract tutorial if the user hasn't dismissed it and hasn't completed it
+			FIntroTutorials& IntroTutorials = FModuleManager::GetModuleChecked<FIntroTutorials>(TEXT("IntroTutorials"));
+			const bool bRestart = true;
+			IntroTutorials.LaunchTutorial(AttractTutorial, bRestart, ContextWindow);
+		}
+	}
+	bDeferTutorialOpen = false;
+}
+
+static void GetAnimationValues(float InAnimationProgress, float& OutAlphaFactor0, float& OutPulseFactor0, float& OutAlphaFactor1, float& OutPulseFactor1)
+{
+	InAnimationProgress = FMath::Fmod(InAnimationProgress * 2.0f, 1.0f);
+
+	OutAlphaFactor0 = FMath::Square(1.0f - InAnimationProgress);
+	OutPulseFactor0 = 1.0f - FMath::Square(1.0f - InAnimationProgress);
+
+	float OffsetProgress = FMath::Fmod(InAnimationProgress + 0.25f, 1.0f);
+	OutAlphaFactor1 = FMath::Square(1.0f - OffsetProgress);
+	OutPulseFactor1 = 1.0f - FMath::Square(1.0f - OffsetProgress);
+}
+
+int32 STutorialButton::OnPaint(const FPaintArgs& Args, const FGeometry& AllottedGeometry, const FSlateRect& MyClippingRect, FSlateWindowElementList& OutDrawElements, int32 LayerId, const FWidgetStyle& InWidgetStyle, bool bParentEnabled) const
+{
+	LayerId = SCompoundWidget::OnPaint(Args, AllottedGeometry, MyClippingRect, OutDrawElements, LayerId, InWidgetStyle, bParentEnabled) + 1000;
+
+	if (!(bTutorialCompleted || bTutorialDismissed))
+	{
+		float AlphaFactor0 = 0.0f;
+		float AlphaFactor1 = 0.0f;
+		float PulseFactor0 = 0.0f;
+		float PulseFactor1 = 0.0f;
+		GetAnimationValues(PulseAnimation.GetLerpLooping(), AlphaFactor0, PulseFactor0, AlphaFactor1, PulseFactor1);
+
+		const FSlateBrush* PulseBrush = FEditorStyle::Get().GetBrush(TEXT("TutorialLaunch.Circle"));
+		const FLinearColor PulseColor = FEditorStyle::Get().GetColor(TEXT("TutorialLaunch.Circle.Color"));
+
+		// We should be clipped by the window size, not our containing widget, as we want to draw outside the widget
+		const FVector2D WindowSize = OutDrawElements.GetWindow()->GetSizeInScreen();
+		FSlateRect WindowClippingRect(0.0f, 0.0f, WindowSize.X, WindowSize.Y);
+
+		{
+			FVector2D PulseOffset = FVector2D(PulseFactor0 * TutorialButtonConstants::MaxPulseOffset, PulseFactor0 * TutorialButtonConstants::MaxPulseOffset);
+
+			FVector2D BorderPosition = (AllottedGeometry.AbsolutePosition - ((FVector2D(PulseBrush->Margin.Left, PulseBrush->Margin.Top) * PulseBrush->ImageSize * AllottedGeometry.Scale) + PulseOffset));
+			FVector2D BorderSize = ((AllottedGeometry.Size * AllottedGeometry.Scale) + (PulseOffset * 2.0f) + (FVector2D(PulseBrush->Margin.Right * 2.0f, PulseBrush->Margin.Bottom * 2.0f) * PulseBrush->ImageSize * AllottedGeometry.Scale));
+
+			FPaintGeometry BorderGeometry(BorderPosition, BorderSize, AllottedGeometry.Scale);
+
+			// draw highlight border
+			FSlateDrawElement::MakeBox(OutDrawElements, LayerId++, BorderGeometry, PulseBrush, WindowClippingRect, ESlateDrawEffect::None, FLinearColor(PulseColor.R, PulseColor.G, PulseColor.B, AlphaFactor0));
+		}
+
+		{
+			FVector2D PulseOffset = FVector2D(PulseFactor1 * TutorialButtonConstants::MaxPulseOffset, PulseFactor1 * TutorialButtonConstants::MaxPulseOffset);
+
+			FVector2D BorderPosition = (AllottedGeometry.AbsolutePosition - ((FVector2D(PulseBrush->Margin.Left, PulseBrush->Margin.Top) * PulseBrush->ImageSize * AllottedGeometry.Scale) + PulseOffset));
+			FVector2D BorderSize = ((AllottedGeometry.Size * AllottedGeometry.Scale) + (PulseOffset * 2.0f) + (FVector2D(PulseBrush->Margin.Right * 2.0f, PulseBrush->Margin.Bottom * 2.0f) * PulseBrush->ImageSize * AllottedGeometry.Scale));
+
+			FPaintGeometry BorderGeometry(BorderPosition, BorderSize, AllottedGeometry.Scale);
+
+			// draw highlight border
+			FSlateDrawElement::MakeBox(OutDrawElements, LayerId++, BorderGeometry, PulseBrush, WindowClippingRect, ESlateDrawEffect::None, FLinearColor(PulseColor.R, PulseColor.G, PulseColor.B, AlphaFactor1));
+		}
+	}
+
+	return LayerId;
+}
+
+FReply STutorialButton::HandleButtonClicked()
+{
+	UEditorTutorial* AttractTutorial = nullptr;
+	UEditorTutorial* LaunchTutorial = nullptr;
+	GetDefault<UEditorTutorialSettings>()->FindTutorialsForContext(Context, AttractTutorial, LaunchTutorial);
+	if (LaunchTutorial != nullptr && ContextWindow.IsValid())
+	{
+		FIntroTutorials& IntroTutorials = FModuleManager::GetModuleChecked<FIntroTutorials>(TEXT("IntroTutorials"));
+		const bool bRestart = true;
+		IntroTutorials.LaunchTutorial(LaunchTutorial, bRestart, ContextWindow);
+	}
+
+	return FReply::Handled();
+}
+
+EVisibility STutorialButton::GetVisibility() const
+{
+	return bTutorialAvailable ? EVisibility::Visible : EVisibility::Collapsed;
+}
+
+FReply STutorialButton::OnMouseButtonDown(const FGeometry& MyGeometry, const FPointerEvent& MouseEvent)
+{
+	if (MouseEvent.GetEffectingButton() == EKeys::RightMouseButton)
+	{
+		const bool bInShouldCloseWindowAfterMenuSelection = true;
+		FMenuBuilder MenuBuilder(bInShouldCloseWindowAfterMenuSelection, nullptr);
+
+		MenuBuilder.AddMenuEntry(
+			LOCTEXT("DismissReminder", "Dismiss Alert"),
+			LOCTEXT("DismissReminderTooltip", "Don't show me this alert again"),
+			FSlateIcon(),
+			FUIAction(FExecuteAction::CreateSP(this, &STutorialButton::DismissAlert))
+			);
+
+		FSlateApplication::Get().PushMenu(SharedThis(this), MenuBuilder.MakeWidget(), FSlateApplication::Get().GetCursorPos(), FPopupTransitionEffect::ContextMenu);
+	}
+	return FReply::Handled();
+}
+
+void STutorialButton::DismissAlert()
+{
+	UEditorTutorial* AttractTutorial = nullptr;
+	UEditorTutorial* LaunchTutorial = nullptr;
+	GetDefault<UEditorTutorialSettings>()->FindTutorialsForContext(Context, AttractTutorial, LaunchTutorial);
+	if (AttractTutorial != nullptr)
+	{
+		const bool bDismissAcrossSessions = true;
+		GetMutableDefault<UTutorialStateSettings>()->DismissTutorial(AttractTutorial, bDismissAcrossSessions);
+		GetMutableDefault<UTutorialStateSettings>()->SaveProgress();
+		bTutorialDismissed = true;
+
+		FIntroTutorials& IntroTutorials = FModuleManager::GetModuleChecked<FIntroTutorials>(TEXT("IntroTutorials"));
+		IntroTutorials.CloseAllTutorialContent();
+	}
+}
+
+#undef LOCTEXT_NAMESPACE
\ No newline at end of file
diff --git a/Engine/Source/Editor/IntroTutorials/Private/STutorialButton.h b/Engine/Source/Editor/IntroTutorials/Private/STutorialButton.h
new file mode 100644
--- /dev/null
+++ b/Engine/Source/Editor/IntroTutorials/Private/STutorialButton.h
@@ -0,0 +1,52 @@
+// Copyright 1998-2014 Epic Games, Inc. All Rights Reserved.
+
+#pragma once
+
+class STutorialButton : public SCompoundWidget
+{
+	SLATE_BEGIN_ARGS(STutorialButton) {}
+	
+	SLATE_ARGUMENT(FName, Context)
+
+	SLATE_ARGUMENT(TWeakPtr<SWindow>, ContextWindow)
+
+	SLATE_END_ARGS()
+
+	/** Widget constructor */
+	void Construct(const FArguments& Args);
+
+	virtual int32 OnPaint(const FPaintArgs& Args, const FGeometry& AllottedGeometry, const FSlateRect& MyClippingRect, FSlateWindowElementList& OutDrawElements, int32 LayerId, const FWidgetStyle& InWidgetStyle, bool bParentEnabled) const override;
+	virtual FReply OnMouseButtonDown(const FGeometry& MyGeometry, const FPointerEvent& MouseEvent) override;
+	virtual void Tick(const FGeometry& AllottedGeometry, const double InCurrentTime, const float InDeltaTime) override;
+
+private:
+	void HandleMainFrameLoad(TSharedPtr<SWindow> InRootWindow, bool bIsNewProjectWindow);
+
+	FReply HandleButtonClicked();
+
+	EVisibility GetVisibility() const;
+
+	void DismissAlert();
+
+private:
+	/** Flag to defer tutorial open until the first Tick() */
+	bool bDeferTutorialOpen;
+
+	/** Whether we have a tutorial for this context */
+	bool bTutorialAvailable;
+
+	/** Whether we have completed the tutorial for this content */
+	bool bTutorialCompleted;
+
+	/** Whether we have dismissed the tutorial for this content */
+	bool bTutorialDismissed;
+
+	/** Context that this widget was created for (i.e. what part of the editor) */
+	FName Context;
+
+	/** Window that the tutorial should be launched in */
+	TWeakPtr<SWindow> ContextWindow;
+
+	/** Animation curve for displaying pulse */
+	FCurveSequence PulseAnimation;
+};
\ No newline at end of file
diff --git a/Engine/Source/Editor/IntroTutorials/Private/STutorialContent.cpp b/Engine/Source/Editor/IntroTutorials/Private/STutorialContent.cpp
--- a/Engine/Source/Editor/IntroTutorials/Private/STutorialContent.cpp
+++ b/Engine/Source/Editor/IntroTutorials/Private/STutorialContent.cpp
@@ -48,9 +48,10 @@ void STutorialContent::Construct(const FArguments& InArgs, const FTutorialConten
 		SAssignNew(ContentWidget, SBorder)
 
 		// Add more padding if the content is to be displayed centrally (i.e. not on a widget)
-		.Padding(Anchor.Type != ETutorialAnchorIdentifier::None ? 18.0f : 24.0f)
+		.Padding(24.0f)
 		.BorderImage(FEditorStyle::GetBrush("Tutorials.Border"))
 		.Visibility(this, &STutorialContent::GetVisibility)
+		.BorderBackgroundColor(this, &STutorialContent::GetBackgroundColor)
 		.ForegroundColor(FCoreStyle::Get().GetSlateColor("InvertedForeground"))
 		[
 			SNew(SHorizontalBox)
@@ -59,7 +60,7 @@ void STutorialContent::Construct(const FArguments& InArgs, const FTutorialConten
 			.MaxWidth(600.0f)
 			.VAlign(VAlign_Center)
 			[
-				GenerateContentWidget(InContent, InArgs._WrapTextAt, Anchor.Type != ETutorialAnchorIdentifier::None, DocumentationPage)
+				GenerateContentWidget(InContent, InArgs._WrapTextAt, DocumentationPage)
 			]
 			+SHorizontalBox::Slot()
 			.AutoWidth()
@@ -101,6 +102,9 @@ static void GetAnimationValues(bool bIsIntro, float InAnimationProgress, float&
 
 int32 STutorialContent::OnPaint( const FPaintArgs& Args, const FGeometry& AllottedGeometry, const FSlateRect& MyClippingRect, FSlateWindowElementList& OutDrawElements, int32 LayerId, const FWidgetStyle& InWidgetStyle, bool bParentEnabled ) const
 {
+	CachedContentGeometry = AllottedGeometry;
+	CachedContentGeometry.AbsolutePosition += OutDrawElements.GetWindow()->GetPositionInScreen();
+
 	if(bIsVisible && Anchor.Type != ETutorialAnchorIdentifier::None && Anchor.bDrawHighlight)
 	{
 		float AlphaFactor;
@@ -164,7 +168,7 @@ static TSharedRef<SWidget> GetStageTitle(const FExcerpt& InExcerpt, int32 InCurr
 	return SNullWidget::NullWidget;
 }
 
-TSharedRef<SWidget> STutorialContent::GenerateContentWidget(const FTutorialContent& InContent, float WrapTextAt, bool bForWidget, TSharedPtr<IDocumentationPage>& OutDocumentationPage)
+TSharedRef<SWidget> STutorialContent::GenerateContentWidget(const FTutorialContent& InContent, float WrapTextAt, TSharedPtr<IDocumentationPage>& OutDocumentationPage)
 {
 	// Style for the documentation
 	static FDocumentationStyle DocumentationStyle;
@@ -188,7 +192,7 @@ TSharedRef<SWidget> STutorialContent::GenerateContentWidget(const FTutorialConte
 				.Visibility(EVisibility::SelfHitTestInvisible)
 				.WrapTextAt(WrapTextAt)
 				.Text(InContent.Text)
-				.TextStyle(FEditorStyle::Get(), bForWidget ? "Tutorials.WidgetContent" : "Tutorials.Content");
+				.TextStyle(FEditorStyle::Get(), "Tutorials.Content");
 		}
 
 	case ETutorialContent::UDNExcerpt:
@@ -234,7 +238,7 @@ TSharedRef<SWidget> STutorialContent::GenerateContentWidget(const FTutorialConte
 			Decorators.Add(FHyperlinkDecorator::Create(TEXT("tutorial"), FSlateHyperlinkRun::FOnClick::CreateStatic(&TutorialTextHelpers::OnTutorialLinkClicked)));
 			Decorators.Add(FHyperlinkDecorator::Create(TEXT("code"), FSlateHyperlinkRun::FOnClick::CreateStatic(&TutorialTextHelpers::OnCodeLinkClicked)));
 			Decorators.Add(FHyperlinkDecorator::Create(TEXT("asset"), FSlateHyperlinkRun::FOnClick::CreateStatic(&TutorialTextHelpers::OnAssetLinkClicked)));
-			Decorators.Add(FTextStyleDecorator::Create(&TextStyles));
+			Decorators.Add(FTextStyleDecorator::Create(TextStyles));
 
 			return SNew(SRichTextBlock)
 					.Visibility(EVisibility::SelfHitTestInvisible)
@@ -320,10 +324,14 @@ void STutorialContent::HandlePaintNamedWidget(TSharedRef<SWidget> InWidget, cons
 	switch(Anchor.Type)
 	{
 	case ETutorialAnchorIdentifier::NamedWidget:
-		if(Anchor.WrapperIdentifier == InWidget->GetTag())
 		{
-			bIsVisible = true;
-			CachedGeometry = InGeometry;
+			TSharedPtr<FTagMetaData> MetaData = InWidget->GetMetaData<FTagMetaData>();
+			if( Anchor.WrapperIdentifier == InWidget->GetTag() ||
+				(MetaData.IsValid() && MetaData->Tag == Anchor.WrapperIdentifier))
+			{
+				bIsVisible = true;
+				CachedGeometry = InGeometry;
+			}
 		}
 		break;
 	}
@@ -344,4 +352,11 @@ EVisibility STutorialContent::GetVisibility() const
 	return bIsVisible || Anchor.Type == ETutorialAnchorIdentifier::None ? EVisibility::SelfHitTestInvisible : EVisibility::Collapsed;
 }
 
+FSlateColor STutorialContent::GetBackgroundColor() const
+{
+	// note cant use IsHovered() here because our widget is SelfHitTestInvisible
+	const FVector2D CursorPos = FSlateApplication::Get().GetCursorPos();
+	return CachedContentGeometry.IsUnderLocation(CursorPos) ? FLinearColor(1.0f, 1.0f, 1.0f, 1.0f) : FLinearColor(1.0f, 1.0f, 1.0f, 0.8f);
+}
+
 #undef LOCTEXT_NAMESPACE
\ No newline at end of file
diff --git a/Engine/Source/Editor/IntroTutorials/Private/STutorialContent.h b/Engine/Source/Editor/IntroTutorials/Private/STutorialContent.h
--- a/Engine/Source/Editor/IntroTutorials/Private/STutorialContent.h
+++ b/Engine/Source/Editor/IntroTutorials/Private/STutorialContent.h
@@ -46,7 +46,7 @@ class STutorialContent : public SCompoundWidget
 	virtual int32 OnPaint( const FPaintArgs& Args, const FGeometry& AllottedGeometry, const FSlateRect& MyClippingRect, FSlateWindowElementList& OutDrawElements, int32 LayerId, const FWidgetStyle& InWidgetStyle, bool bParentEnabled ) const override;
 
 	/** Helper function to generate widgets from an FTutorialContent struct */
-	static TSharedRef<SWidget> GenerateContentWidget(const FTutorialContent& InContent, float WrapTextAt, bool bForWidget, TSharedPtr<IDocumentationPage>& OutDocumentationPage);
+	static TSharedRef<SWidget> GenerateContentWidget(const FTutorialContent& InContent, float WrapTextAt, TSharedPtr<IDocumentationPage>& OutDocumentationPage);
 
 	/** Handle repositioning the widget */
 	FVector2D GetPosition() const;
@@ -74,12 +74,20 @@ private:
 	/** Get close button visibility - varies depending on whether we are standalone or not */
 	EVisibility GetCloseButtonVisibility() const;
 
+	/** Alter the background color depending on hover state */
+	FSlateColor GetBackgroundColor() const;
+
+private:
+
 	/** Copy of the window size we were last draw at */
 	FVector2D CachedWindowSize;
 
 	/** Copy of the geometry our widget was last drawn with */
 	FGeometry CachedGeometry;
 
+	/** Copy of the geometry our content was last drawn with */
+	mutable FGeometry CachedContentGeometry;
+
 	/** Container for widget content */
 	TSharedPtr<SWidget> ContentWidget;
 
diff --git a/Engine/Source/Editor/IntroTutorials/Private/STutorialEditableText.cpp b/Engine/Source/Editor/IntroTutorials/Private/STutorialEditableText.cpp
--- a/Engine/Source/Editor/IntroTutorials/Private/STutorialEditableText.cpp
+++ b/Engine/Source/Editor/IntroTutorials/Private/STutorialEditableText.cpp
@@ -186,7 +186,7 @@ void STutorialEditableText::Construct(const FArguments& InArgs)
 	RichTextMarshaller->AppendInlineDecorator(FHyperlinkDecorator::Create(TEXT("tutorial"), OnTutorialLinkClicked));
 	RichTextMarshaller->AppendInlineDecorator(FHyperlinkDecorator::Create(TEXT("code"), OnCodeLinkClicked));
 	RichTextMarshaller->AppendInlineDecorator(FHyperlinkDecorator::Create(TEXT("asset"), OnAssetLinkClicked));
-	RichTextMarshaller->AppendInlineDecorator(FTextStyleDecorator::Create(&TextStyles));
+	RichTextMarshaller->AppendInlineDecorator(FTextStyleDecorator::Create(TextStyles));
 
 	this->ChildSlot
 	[
diff --git a/Engine/Source/Editor/IntroTutorials/Private/STutorialEditableText.h b/Engine/Source/Editor/IntroTutorials/Private/STutorialEditableText.h
--- a/Engine/Source/Editor/IntroTutorials/Private/STutorialEditableText.h
+++ b/Engine/Source/Editor/IntroTutorials/Private/STutorialEditableText.h
@@ -203,7 +203,7 @@ class FTextStyleDecorator : public ITextDecorator
 {
 public:
 
-	static TSharedRef<FTextStyleDecorator> Create(FTextStyles* const InTextStyles)
+	static TSharedRef<FTextStyleDecorator> Create(const FTextStyles& InTextStyles)
 	{
 		return MakeShareable(new FTextStyleDecorator(InTextStyles));
 	}
@@ -230,17 +230,17 @@ public:
 		*InOutModelText += OriginalText.Mid(RunParseResult.ContentRange.BeginIndex, RunParseResult.ContentRange.EndIndex - RunParseResult.ContentRange.BeginIndex);
 		ModelRange.EndIndex = InOutModelText->Len();
 
-		return FSlateTextRun::Create(RunInfo, InOutModelText, TextStyles->CreateTextBlockStyle(RunInfo), ModelRange);
+		return FSlateTextRun::Create(RunInfo, InOutModelText, TextStyles.CreateTextBlockStyle(RunInfo), ModelRange);
 	}
 
 private:
 
-	FTextStyleDecorator(FTextStyles* const InTextStyles)
+	FTextStyleDecorator(const FTextStyles& InTextStyles)
 		: TextStyles(InTextStyles)
 	{
 	}
 
-	FTextStyles* TextStyles;
+	FTextStyles TextStyles;
 };
 
 
diff --git a/Engine/Source/Editor/IntroTutorials/Private/STutorialRoot.cpp b/Engine/Source/Editor/IntroTutorials/Private/STutorialRoot.cpp
--- a/Engine/Source/Editor/IntroTutorials/Private/STutorialRoot.cpp
+++ b/Engine/Source/Editor/IntroTutorials/Private/STutorialRoot.cpp
@@ -4,6 +4,7 @@
 #include "STutorialRoot.h"
 #include "SEditorTutorials.h"
 #include "EditorTutorialSettings.h"
+#include "TutorialStateSettings.h"
 
 #define LOCTEXT_NAMESPACE "STutorialRoot"
 
@@ -76,14 +77,14 @@ void STutorialRoot::SummonTutorialBrowser(TSharedRef<SWindow> InWindow, const FS
 	}
 }
 
-void STutorialRoot::LaunchTutorial(UEditorTutorial* InTutorial, bool bInRestart, TWeakPtr<SWindow> InNavigationWindow)
+void STutorialRoot::LaunchTutorial(UEditorTutorial* InTutorial, bool bInRestart, TWeakPtr<SWindow> InNavigationWindow, FSimpleDelegate InOnTutorialClosed, FSimpleDelegate InOnTutorialExited)
 {
 	if(InTutorial != nullptr)
 	{
 		CurrentTutorial = InTutorial;
 
 		bool bHaveSeenTutorial = false;
-		CurrentTutorialStage = bInRestart ? 0 : GetDefault<UEditorTutorialSettings>()->GetProgress(CurrentTutorial, bHaveSeenTutorial);
+		CurrentTutorialStage = bInRestart ? 0 : GetDefault<UTutorialStateSettings>()->GetProgress(CurrentTutorial, bHaveSeenTutorial);
 
 		// launch tutorial for all windows we wrap - any tutorial can display over any window
 		for(auto& TutorialWidget : TutorialWidgets)
@@ -99,12 +100,23 @@ void STutorialRoot::LaunchTutorial(UEditorTutorial* InTutorial, bool bInRestart,
 				{
 					bIsNavigationWindow = (TutorialWidget.Value.Pin()->GetParentWindow() == InNavigationWindow.Pin());
 				}
-				TutorialWidget.Value.Pin()->LaunchTutorial(bIsNavigationWindow);
+				TutorialWidget.Value.Pin()->LaunchTutorial(bIsNavigationWindow, InOnTutorialClosed, InOnTutorialExited);
 			}
 		}
 	}
 }
 
+void STutorialRoot::CloseAllTutorialContent()
+{
+	for (auto& TutorialWidget : TutorialWidgets)
+	{
+		if (TutorialWidget.Value.IsValid())
+		{
+			TutorialWidget.Value.Pin()->HideContent();
+		}
+	}
+}
+
 void STutorialRoot::HandleNextClicked(TWeakPtr<SWindow> InNavigationWindow)
 {
 	GoToNextStage(InNavigationWindow);
@@ -137,8 +149,8 @@ void STutorialRoot::HandleHomeClicked()
 {
 	if(CurrentTutorial != nullptr)
 	{
-		GetMutableDefault<UEditorTutorialSettings>()->RecordProgress(CurrentTutorial, CurrentTutorialStage);
-		GetMutableDefault<UEditorTutorialSettings>()->SaveProgress();
+		GetMutableDefault<UTutorialStateSettings>()->RecordProgress(CurrentTutorial, CurrentTutorialStage);
+		GetMutableDefault<UTutorialStateSettings>()->SaveProgress();
 	}
 
 	CurrentTutorial = nullptr;
@@ -204,7 +216,7 @@ void STutorialRoot::GoToNextStage(TWeakPtr<SWindow> InNavigationWindow)
 			TSubclassOf<UEditorTutorial> NextTutorialClass = LoadClass<UEditorTutorial>(NULL, *CurrentTutorial->NextTutorial.AssetLongPathname, NULL, LOAD_None, NULL);
 			if(NextTutorialClass != nullptr)
 			{
-				LaunchTutorial(NextTutorialClass->GetDefaultObject<UEditorTutorial>(), true, InNavigationWindow);
+				LaunchTutorial(NextTutorialClass->GetDefaultObject<UEditorTutorial>(), true, InNavigationWindow, FSimpleDelegate(), FSimpleDelegate());
 			}
 			else
 			{
@@ -214,10 +226,10 @@ void STutorialRoot::GoToNextStage(TWeakPtr<SWindow> InNavigationWindow)
 		else
 		{
 			CurrentTutorialStage = FMath::Min(CurrentTutorialStage + 1, CurrentTutorial->Stages.Num() - 1);
-			GetMutableDefault<UEditorTutorialSettings>()->RecordProgress(CurrentTutorial, CurrentTutorialStage);
+			GetMutableDefault<UTutorialStateSettings>()->RecordProgress(CurrentTutorial, CurrentTutorialStage);
 		}
 
-		if (CurrentTutorial != nullptr && (CurrentTutorial != PreviousTutorial || CurrentTutorialStage != PreviousTutorialStage))
+		if (CurrentTutorial != nullptr && CurrentTutorialStage < CurrentTutorial->Stages.Num() && (CurrentTutorial != PreviousTutorial || CurrentTutorialStage != PreviousTutorialStage))
 		{
 			CurrentTutorial->HandleTutorialStageStarted(CurrentTutorial->Stages[CurrentTutorialStage].Name);
 		}
diff --git a/Engine/Source/Editor/IntroTutorials/Private/STutorialRoot.h b/Engine/Source/Editor/IntroTutorials/Private/STutorialRoot.h
--- a/Engine/Source/Editor/IntroTutorials/Private/STutorialRoot.h
+++ b/Engine/Source/Editor/IntroTutorials/Private/STutorialRoot.h
@@ -28,7 +28,10 @@ public:
 	virtual void AddReferencedObjects( FReferenceCollector& Collector ) override;
 
 	/** Launch the specified tutorial from the specified window */
-	void LaunchTutorial(UEditorTutorial* InTutorial, bool bInRestart, TWeakPtr<SWindow> InNavigationWindow);
+	void LaunchTutorial(UEditorTutorial* InTutorial, bool bInRestart, TWeakPtr<SWindow> InNavigationWindow, FSimpleDelegate InOnTutorialClosed, FSimpleDelegate InOnTutorialExited);
+
+	/** Close all tutorial content */
+	void CloseAllTutorialContent();
 
 	/** Summon the browser widget for the specified window */
 	void SummonTutorialBrowser(TSharedRef<SWindow> InWindow, const FString& InFilter = TEXT(""));
diff --git a/Engine/Source/Editor/IntroTutorials/Private/STutorialsBrowser.cpp b/Engine/Source/Editor/IntroTutorials/Private/STutorialsBrowser.cpp
--- a/Engine/Source/Editor/IntroTutorials/Private/STutorialsBrowser.cpp
+++ b/Engine/Source/Editor/IntroTutorials/Private/STutorialsBrowser.cpp
@@ -6,6 +6,7 @@
 #include "STutorialContent.h"
 #include "TutorialSettings.h"
 #include "EditorTutorialSettings.h"
+#include "TutorialStateSettings.h"
 #include "AssetRegistryModule.h"
 
 #define LOCTEXT_NAMESPACE "TutorialsBrowser"
@@ -253,7 +254,7 @@ public:
 						+SVerticalBox::Slot()
 						.FillHeight(1.0f)
 						[
-							STutorialContent::GenerateContentWidget(Tutorial->SummaryContent, 500.0f, false, DocumentationPage)
+							STutorialContent::GenerateContentWidget(Tutorial->SummaryContent, 500.0f, DocumentationPage)
 						]
 					]
 				]
@@ -279,7 +280,7 @@ public:
 	TOptional<float> GetProgress() const
 	{
 		bool bHaveSeenTutorial = false;
-		const int32 CurrentStage = GetDefault<UEditorTutorialSettings>()->GetProgress(Tutorial, bHaveSeenTutorial);
+		const int32 CurrentStage = GetDefault<UTutorialStateSettings>()->GetProgress(Tutorial, bHaveSeenTutorial);
 		return (Tutorial->Stages.Num() > 0) ? (float)(CurrentStage + 1) / (float)Tutorial->Stages.Num() : 0.0f;
 	}
 
@@ -288,7 +289,7 @@ public:
 		if(LaunchButton->IsHovered())
 		{
 			bool bHaveSeenTutorial = false;
-			const int32 CurrentStage = GetDefault<UEditorTutorialSettings>()->GetProgress(Tutorial, bHaveSeenTutorial);
+			const int32 CurrentStage = GetDefault<UTutorialStateSettings>()->GetProgress(Tutorial, bHaveSeenTutorial);
 			return LaunchButton->IsHovered() && bHaveSeenTutorial ? EVisibility::Visible : EVisibility::Hidden;
 		}
 
@@ -629,7 +630,7 @@ FSlateColor STutorialsBrowser::GetBackButtonColor() const
 
 void STutorialsBrowser::OnTutorialSelected(UEditorTutorial* InTutorial, bool bRestart)
 {
-	OnLaunchTutorial.ExecuteIfBound(InTutorial, bRestart, ParentWindow);
+	OnLaunchTutorial.ExecuteIfBound(InTutorial, bRestart, ParentWindow, FSimpleDelegate(), FSimpleDelegate());
 }
 
 void STutorialsBrowser::OnCategorySelected(const FString& InCategory)
diff --git a/Engine/Source/Editor/IntroTutorials/Private/STutorialsBrowser.h b/Engine/Source/Editor/IntroTutorials/Private/STutorialsBrowser.h
--- a/Engine/Source/Editor/IntroTutorials/Private/STutorialsBrowser.h
+++ b/Engine/Source/Editor/IntroTutorials/Private/STutorialsBrowser.h
@@ -4,7 +4,7 @@
 
 class UEditorTutorial;
 
-DECLARE_DELEGATE_ThreeParams(FOnLaunchTutorial, UEditorTutorial* /** InTutorialToLaunch */, bool /* bRestart */, TWeakPtr<SWindow> /* InFromWindow */)
+DECLARE_DELEGATE_FiveParams(FOnLaunchTutorial, UEditorTutorial* /** InTutorialToLaunch */, bool /* bRestart */, TWeakPtr<SWindow> /* InFromWindow */, FSimpleDelegate /* InOnTutorialClosed */, FSimpleDelegate /* InOnTutorialExited */);
 
 /** Abstract base class for list entries in the tutorial menu */
 struct ITutorialListEntry
diff --git a/Engine/Source/Editor/IntroTutorials/Private/TutorialStateSettings.cpp b/Engine/Source/Editor/IntroTutorials/Private/TutorialStateSettings.cpp
new file mode 100644
--- /dev/null
+++ b/Engine/Source/Editor/IntroTutorials/Private/TutorialStateSettings.cpp
@@ -0,0 +1,120 @@
+// Copyright 1998-2014 Epic Games, Inc. All Rights Reserved.
+
+#include "IntroTutorialsPrivatePCH.h"
+#include "TutorialStateSettings.h"
+
+UTutorialStateSettings::UTutorialStateSettings(const class FPostConstructInitializeProperties& PCIP)
+	: Super(PCIP)
+{
+
+}
+
+void UTutorialStateSettings::PostInitProperties()
+{
+	Super::PostInitProperties();
+
+	for(const auto& Progress : TutorialsProgress)
+	{
+		TSubclassOf<UEditorTutorial> TutorialClass = LoadClass<UEditorTutorial>(NULL, *Progress.Tutorial.AssetLongPathname, NULL, LOAD_None, NULL);
+		if(TutorialClass != nullptr)
+		{
+			UEditorTutorial* Tutorial = TutorialClass->GetDefaultObject<UEditorTutorial>();
+			ProgressMap.Add(Tutorial, Progress);
+		}
+	}
+}
+
+int32 UTutorialStateSettings::GetProgress(UEditorTutorial* InTutorial, bool& bOutHaveSeenTutorial) const
+{
+	const FTutorialProgress* FoundProgress = ProgressMap.Find(InTutorial);
+	if(FoundProgress != nullptr)
+	{
+		bOutHaveSeenTutorial = true;
+		return FoundProgress->CurrentStage;
+	}
+
+	bOutHaveSeenTutorial = false;
+	return 0;
+}
+
+bool UTutorialStateSettings::HaveSeenTutorial(UEditorTutorial* InTutorial) const
+{
+	return ProgressMap.Find(InTutorial) != nullptr;
+}
+
+bool UTutorialStateSettings::HaveCompletedTutorial(UEditorTutorial* InTutorial) const
+{
+	const FTutorialProgress* FoundProgress = ProgressMap.Find(InTutorial);
+	if (FoundProgress != nullptr)
+	{
+		return FoundProgress->CurrentStage >= InTutorial->Stages.Num() - 1;
+	}
+	return false;
+}
+
+void UTutorialStateSettings::RecordProgress(UEditorTutorial* InTutorial, int32 CurrentStage)
+{
+	if(InTutorial != nullptr)
+	{
+		FTutorialProgress* FoundProgress = ProgressMap.Find(InTutorial);
+		if(FoundProgress != nullptr)
+		{
+			FoundProgress->CurrentStage = FMath::Max(FoundProgress->CurrentStage, CurrentStage);
+		}
+		else
+		{
+			FTutorialProgress Progress;
+			Progress.Tutorial = FStringClassReference(InTutorial->GetClass());
+			Progress.CurrentStage = CurrentStage;
+			Progress.bUserDismissed = false;
+			Progress.bUserDismissedThisSession = false;
+
+			ProgressMap.Add(InTutorial, Progress);
+		}
+	}
+}
+
+void UTutorialStateSettings::DismissTutorial(UEditorTutorial* InTutorial, bool bDismissAcrossSessions)
+{
+	if (InTutorial != nullptr)
+	{
+		FTutorialProgress* FoundProgress = ProgressMap.Find(InTutorial);
+		if (FoundProgress != nullptr)
+		{
+			FoundProgress->bUserDismissed = bDismissAcrossSessions;
+			FoundProgress->bUserDismissedThisSession = true;
+		}
+		else
+		{
+			FTutorialProgress Progress;
+			Progress.Tutorial = FStringClassReference(InTutorial->GetClass());
+			Progress.CurrentStage = 0;
+			Progress.bUserDismissed = bDismissAcrossSessions;
+			Progress.bUserDismissedThisSession = true;
+
+			ProgressMap.Add(InTutorial, Progress);
+		}
+	}
+}
+
+bool UTutorialStateSettings::IsTutorialDismissed(UEditorTutorial* InTutorial) const
+{
+	const FTutorialProgress* FoundProgress = ProgressMap.Find(InTutorial);
+	if (FoundProgress != nullptr)
+	{
+		return FoundProgress->bUserDismissed || FoundProgress->bUserDismissedThisSession;
+	}
+	return false;
+}
+
+void UTutorialStateSettings::SaveProgress()
+{
+	TutorialsProgress.Empty();
+
+	for(const auto& ProgressMapEntry : ProgressMap)
+	{
+		TutorialsProgress.Add(ProgressMapEntry.Value);
+	}
+
+	SaveConfig();
+}
diff --git a/Engine/Source/Editor/IntroTutorials/Private/TutorialStateSettings.h b/Engine/Source/Editor/IntroTutorials/Private/TutorialStateSettings.h
new file mode 100644
--- /dev/null
+++ b/Engine/Source/Editor/IntroTutorials/Private/TutorialStateSettings.h
@@ -0,0 +1,68 @@
+// Copyright 1998-2014 Epic Games, Inc. All Rights Reserved.
+
+#pragma once
+
+#include "EditorTutorial.h"
+#include "TutorialStateSettings.generated.h"
+
+/** Track the progress of an individual tutorial */
+USTRUCT()
+struct FTutorialProgress
+{
+	GENERATED_USTRUCT_BODY()
+
+	FTutorialProgress()
+	{
+		bUserDismissedThisSession = false;
+	}
+
+	UPROPERTY()
+	FStringClassReference Tutorial;
+
+	UPROPERTY()
+	int32 CurrentStage;
+
+	UPROPERTY()
+	bool bUserDismissed;
+
+	/** Non-persistent flag indicating the user dismissed this tutorial */
+	bool bUserDismissedThisSession;
+};
+
+/** Tutorial settings used to track completion state */
+UCLASS(config=EditorGameAgnostic)
+class UTutorialStateSettings : public UObject
+{
+	GENERATED_UCLASS_BODY()
+
+	UPROPERTY(Config)
+	TArray<FTutorialProgress> TutorialsProgress;
+
+	/** UObject interface */
+	virtual void PostInitProperties() override;
+
+	/** Get the recorded progress of the pass-in tutorial */
+	int32 GetProgress(UEditorTutorial* InTutorial, bool& bOutHaveSeenTutorial) const;
+
+	/** Check if we have seen the passed-in tutorial before */
+	bool HaveSeenTutorial(UEditorTutorial* InTutorial) const;
+
+	/** Check if completed the passed in tutorial (i.e. seen all of its stages) */
+	bool HaveCompletedTutorial(UEditorTutorial* InTutorial) const;
+
+	/** Flag a tutorial as dismissed */
+	void DismissTutorial(UEditorTutorial* InTutorial, bool bDismissAcrossSessions);
+
+	/** Check if a tutorial has been dismissed */
+	bool IsTutorialDismissed(UEditorTutorial* InTutorial) const;
+
+	/** Record the progress of the passed-in tutorial */
+	void RecordProgress(UEditorTutorial* InTutorial, int32 CurrentStage);
+
+	/** Save the progress of all our tutorials */
+	void SaveProgress();
+
+private:
+	/** Recorded progress */
+	TMap<UEditorTutorial*, FTutorialProgress> ProgressMap;
+};
\ No newline at end of file
diff --git a/Engine/Source/Editor/IntroTutorials/Public/IIntroTutorials.h b/Engine/Source/Editor/IntroTutorials/Public/IIntroTutorials.h
--- a/Engine/Source/Editor/IntroTutorials/Public/IIntroTutorials.h
+++ b/Engine/Source/Editor/IntroTutorials/Public/IIntroTutorials.h
@@ -58,6 +58,19 @@ public:
 	 * @param	bRestart	Whether to restart the tutorial or resume from where we left off last time.
 	 * @param	InNavigationWindow	Optional window to launch the tutorial from - this is where navigation will be displayed.
 	 */
-	virtual void LaunchTutorial(UEditorTutorial* Tutorial, bool bRestart = true, TWeakPtr<SWindow> InNavigationWindow = nullptr) = 0;
+	virtual void LaunchTutorial(UEditorTutorial* Tutorial, bool bRestart = true, TWeakPtr<SWindow> InNavigationWindow = nullptr, FSimpleDelegate OnTutorialClosed = FSimpleDelegate(), FSimpleDelegate OnTutorialExited = FSimpleDelegate()) = 0;
+
+	/**
+ 	 *  Close all tutorial content, including the browser.
+	 */
+	virtual void CloseAllTutorialContent() = 0;
+
+	/**
+	 * Create a widget that allows access to the tutorial for the current context.
+	 * @param	InContext			The name of the context this widget is attached to (e.g. "LevelEditor")
+	 * @param	InContextWindow		The window that the context is attached to (e.g. the main window, or an asset editor tab)
+	 * @return a widget used to access context-sensitive tutorials
+	 */
+	virtual TSharedRef<SWidget> CreateTutorialsWidget(FName InContext, TWeakPtr<SWindow> InContextWindow) const = 0;
 };
 
diff --git a/Engine/Source/Editor/LevelEditor/LevelEditor.Build.cs b/Engine/Source/Editor/LevelEditor/LevelEditor.Build.cs
--- a/Engine/Source/Editor/LevelEditor/LevelEditor.Build.cs
+++ b/Engine/Source/Editor/LevelEditor/LevelEditor.Build.cs
@@ -13,6 +13,7 @@ public class LevelEditor : ModuleRules
 				"MainFrame",
                 "PlacementMode",
 				"ReferenceViewer",
+                "IntroTutorials"
 			}
 		);
 
@@ -74,6 +75,7 @@ public class LevelEditor : ModuleRules
                 "PlacementMode",
 				"UserFeedback",
 				"ReferenceViewer",
+                "IntroTutorials"
 			}
 		);
 	}
diff --git a/Engine/Source/Editor/LevelEditor/Private/LevelEditor.cpp b/Engine/Source/Editor/LevelEditor/Private/LevelEditor.cpp
--- a/Engine/Source/Editor/LevelEditor/Private/LevelEditor.cpp
+++ b/Engine/Source/Editor/LevelEditor/Private/LevelEditor.cpp
@@ -19,6 +19,7 @@
 #include "SDockTab.h"
 #include "ToolkitManager.h"
 #include "TargetPlatform.h"
+#include "IIntroTutorials.h"
 
 // @todo Editor: remove this circular dependency
 #include "Editor/MainFrame/Public/Interfaces/IMainFrameModule.h"
@@ -64,6 +65,9 @@ TSharedRef<SDockTab> FLevelEditorModule::SpawnLevelEditor( const FSpawnTabArgs&
 	IUserFeedbackModule& UserFeedback = FModuleManager::LoadModuleChecked<IUserFeedbackModule>(TEXT("UserFeedback"));
 	TSharedRef<SWidget> UserFeedbackWidget = UserFeedback.CreateFeedbackWidget(NSLOCTEXT("UserFeedback", "LevelEditing", "Level Editing"));
 
+	IIntroTutorials& IntroTutorials = FModuleManager::LoadModuleChecked<IIntroTutorials>(TEXT("IntroTutorials"));
+	TSharedRef<SWidget> TutorialWidget = IntroTutorials.CreateTutorialsWidget(TEXT("LevelEditor"), OwnerWindow);
+
 	TSharedPtr< SWidget > RightContent;
 	{
 		FString OptionalBranchPrefix;
@@ -105,7 +109,7 @@ TSharedRef<SDockTab> FLevelEditorModule::SpawnLevelEditor( const FSpawnTabArgs&
 #if PLATFORM_MAC
 				+SHorizontalBox::Slot()
 				.AutoWidth()
-				.Padding(16.0f, 0.0f, 0.0f, 0.0f)
+				.Padding(8.0f, 0.0f, 0.0f, 0.0f)
 				.VAlign(VAlign_Center)
 				[
 					LevelEditorTab->GetRightContent()
@@ -114,11 +118,18 @@ TSharedRef<SDockTab> FLevelEditorModule::SpawnLevelEditor( const FSpawnTabArgs&
 		
 				+SHorizontalBox::Slot()
 				.AutoWidth()
-				.Padding(16.0f, 0.0f, 0.0f, 0.0f)
+				.Padding(8.0f, 0.0f, 0.0f, 0.0f)
 				.VAlign(VAlign_Center)
 				[
 					UserFeedbackWidget
 				]
+				+ SHorizontalBox::Slot()
+				.AutoWidth()
+				.Padding(8.0f, 0.0f, 0.0f, 0.0f)
+				.VAlign(VAlign_Center)
+				[
+					TutorialWidget
+				]
 		;
 	}
 	LevelEditorTab->SetRightContent( RightContent.ToSharedRef() );
diff --git a/Engine/Source/Editor/UnrealEd/Private/Toolkits/AssetEditorToolkit.cpp b/Engine/Source/Editor/UnrealEd/Private/Toolkits/AssetEditorToolkit.cpp
--- a/Engine/Source/Editor/UnrealEd/Private/Toolkits/AssetEditorToolkit.cpp
+++ b/Engine/Source/Editor/UnrealEd/Private/Toolkits/AssetEditorToolkit.cpp
@@ -12,6 +12,7 @@
 #include "IUserFeedbackModule.h"
 #include "IDocumentation.h"
 #include "ReferenceViewer.h"
+#include "IIntroTutorials.h"
 
 #define LOCTEXT_NAMESPACE "AssetEditorToolkit"
 
@@ -102,7 +103,26 @@ void FAssetEditorToolkit::InitAssetEditor( const EToolkitMode::Type Mode, const
 		IUserFeedbackModule& UserFeedback = FModuleManager::LoadModuleChecked<IUserFeedbackModule>(TEXT("UserFeedback"));
 		TSharedRef<SWidget> UserFeedbackWidget = UserFeedback.CreateFeedbackWidget(GetBaseToolkitName());
 
-		NewMajorTab->SetRightContent(UserFeedbackWidget);
+		IIntroTutorials& IntroTutorials = FModuleManager::LoadModuleChecked<IIntroTutorials>(TEXT("IntroTutorials"));
+		TSharedRef<SWidget> TutorialWidget = IntroTutorials.CreateTutorialsWidget(GetToolkitFName(), NewMajorTab->GetParentWindow());
+
+		NewMajorTab->SetRightContent(
+				SNew(SHorizontalBox)
+				+SHorizontalBox::Slot()
+				.AutoWidth()
+				.Padding(8.0f, 0.0f, 0.0f, 0.0f)
+				.VAlign(VAlign_Center)
+				[
+					UserFeedbackWidget
+				]
+				+SHorizontalBox::Slot()
+				.AutoWidth()
+				.Padding(8.0f, 0.0f, 0.0f, 0.0f)
+				.VAlign(VAlign_Center)
+				[
+					TutorialWidget
+				]	
+			);
 
 		const TSharedRef<FTabManager> NewTabManager = FGlobalTabmanager::Get()->NewTabManager( NewMajorTab.ToSharedRef() );		
 		NewTabManager->SetOnPersistLayout(FTabManager::FOnPersistLayout::CreateRaw(this, &FAssetEditorToolkit::HandleTabManagerPersistLayout));
diff --git a/Engine/Source/Editor/UnrealEd/UnrealEd.Build.cs b/Engine/Source/Editor/UnrealEd/UnrealEd.Build.cs
--- a/Engine/Source/Editor/UnrealEd/UnrealEd.Build.cs
+++ b/Engine/Source/Editor/UnrealEd/UnrealEd.Build.cs
@@ -47,6 +47,7 @@ public class UnrealEd : ModuleRules
 				"ViewportSnapping",
 				"SourceCodeAccess",
 				"ReferenceViewer",
+                "IntroTutorials",
 			}
 		);
 
diff --git a/Engine/Source/Editor/UserFeedback/Private/UserFeedback.cpp b/Engine/Source/Editor/UserFeedback/Private/UserFeedback.cpp
--- a/Engine/Source/Editor/UserFeedback/Private/UserFeedback.cpp
+++ b/Engine/Source/Editor/UserFeedback/Private/UserFeedback.cpp
@@ -630,7 +630,7 @@ public:
 			return SNew(SHorizontalBox)
 
 			+ SHorizontalBox::Slot()
-				.Padding(0.0f, 0.0f, 10.0f, 0.0f)
+				.Padding(0.0f, 0.0f, 0.0f, 0.0f)
 				.VAlign(VAlign_Center)
 				[
 					SNew(SUserFeedbackButtonWidget).Context(Context)
