Editor - Changing the marketplace icon to match the one on the forums.

diff --git a/Engine/Source/Editor/EditorStyle/Private/SlateEditorStyle.cpp b/Engine/Source/Editor/EditorStyle/Private/SlateEditorStyle.cpp
--- a/Engine/Source/Editor/EditorStyle/Private/SlateEditorStyle.cpp
+++ b/Engine/Source/Editor/EditorStyle/Private/SlateEditorStyle.cpp
@@ -2649,7 +2649,8 @@ void FSlateEditorStyle::FStyle::SetupPropertyEditorStyles()
 			.SetShadowColorAndOpacity( FLinearColor::Black )
 		);
 
-		Set( "PropertyEditor.AssetThumbnailShadow", new BOX_BRUSH( "ContentBrowser/ThumbnailShadow", FMargin( 5.0f / 64.0f ) ) );
+		Set( "PropertyEditor.AssetThumbnailShadow", new BOX_BRUSH( "ContentBrowser/ThumbnailShadow", FMargin( 5.0f / 60.0f ) ) );
+		Set( "PropertyEditor.AssetThumbnailLight", new BOX_BRUSH( "ContentBrowser/ThumbnailLight", FMargin( 5.0f / 60.0f ), SelectionColor ) );
 		Set( "PropertyEditor.AssetClass", FTextBlockStyle(NormalText)
 			.SetFont( TTF_CORE_FONT( "Fonts/Roboto-Regular", 10 ) )
 			.SetColorAndOpacity( FLinearColor::White )
@@ -3750,8 +3751,8 @@ void FSlateEditorStyle::FStyle::SetupLevelEditorStyle()
 		Set( "LevelEditor.WorldProperties.Tab", new IMAGE_BRUSH( "Icons/icon_worldscript_40x", Icon16x16 ) );
 		Set( "LevelEditor.OpenContentBrowser", new IMAGE_BRUSH( "Icons/icon_ContentBrowser_40x", Icon40x40 ) );
 		Set( "LevelEditor.OpenContentBrowser.Small", new IMAGE_BRUSH( "Icons/icon_ContentBrowser_40x", Icon20x20 ) );
-		Set( "LevelEditor.OpenMarketplace", new IMAGE_BRUSH( "Icons/icon_Samples_Download_40x", Icon40x40 ) );
-		Set( "LevelEditor.OpenMarketplace.Small", new IMAGE_BRUSH( "Icons/icon_Samples_Download_20x", Icon20x20 ) );
+		Set( "LevelEditor.OpenMarketplace", new IMAGE_BRUSH( "Icons/icon_Marketplace_40x", Icon40x40 ) );
+		Set( "LevelEditor.OpenMarketplace.Small", new IMAGE_BRUSH( "Icons/icon_Marketplace_20x", Icon20x20 ) );
 		Set( "LevelEditor.OpenLevelBlueprint", new IMAGE_BRUSH( "Icons/icon_kismet2_40x", Icon40x40 ) );
 		Set( "LevelEditor.OpenLevelBlueprint.Small", new IMAGE_BRUSH( "Icons/icon_kismet2_40x", Icon20x20 ) );
 		Set( "LevelEditor.CreateClassBlueprint", new IMAGE_BRUSH("Icons/icon_class_Blueprint_New_16x", Icon16x16));
