Removed gradient background from content browser thumbnail icons

diff --git a/Engine/Source/Editor/EditorStyle/Private/SlateEditorStyle.cpp b/Engine/Source/Editor/EditorStyle/Private/SlateEditorStyle.cpp
--- a/Engine/Source/Editor/EditorStyle/Private/SlateEditorStyle.cpp
+++ b/Engine/Source/Editor/EditorStyle/Private/SlateEditorStyle.cpp
@@ -1226,7 +1226,7 @@ void FSlateEditorStyle::FStyle::SetupGeneralStyles()
 		Set( "AssetThumbnail.HintShadowOffset", FVector2D(1,1) );
 		Set( "AssetThumbnail.HintShadowColorAndOpacity", FLinearColor(0, 0, 0, 1) );
 		Set( "AssetThumbnail.HintBackground", new BOX_BRUSH( "Common/TableViewHeader", FMargin(8.0f/32.0f) ) );
-		Set( "AssetThumbnail.Border", new IMAGE_BRUSH( "Common/DiagonalGradient", FVector2D(32.f, 32.f), FLinearColor(1, 1, 1, 1) ) );
+		Set( "AssetThumbnail.Border", new FSlateColorBrush( FColor::White ) );
 	}
 
 	// Open any asset dialog
diff --git a/Engine/Source/Editor/UnrealEd/Private/AssetThumbnail.cpp b/Engine/Source/Editor/UnrealEd/Private/AssetThumbnail.cpp
--- a/Engine/Source/Editor/UnrealEd/Private/AssetThumbnail.cpp
+++ b/Engine/Source/Editor/UnrealEd/Private/AssetThumbnail.cpp
@@ -149,7 +149,7 @@ public:
 			OverlayWidget->AddSlot()
 			[
 				SNew(SBorder)
-				.Padding(2)
+				.Padding(1)
 				.BorderImage(FEditorStyle::GetBrush(Style, ".Border"))
 				//.BorderImage(FEditorStyle::GetBrush("NoBrush"))
 				.BorderBackgroundColor(this, &SAssetThumbnail::GetViewportBorderColorAndOpacity)
