Editor - The ToggleButtonCheckbox style now has a faint alpha white background to visually call out the clickable surface.

diff --git a/Engine/Source/Editor/EditorStyle/Private/SlateEditorStyle.cpp b/Engine/Source/Editor/EditorStyle/Private/SlateEditorStyle.cpp
--- a/Engine/Source/Editor/EditorStyle/Private/SlateEditorStyle.cpp
+++ b/Engine/Source/Editor/EditorStyle/Private/SlateEditorStyle.cpp
@@ -298,7 +298,7 @@ void FSlateEditorStyle::FStyle::SetupGeneralStyles()
 		/* Create a checkbox style for "ToggleButton" ... */
 		const FCheckBoxStyle ToggleButtonStyle = FCheckBoxStyle()
 			.SetCheckBoxType(ESlateCheckBoxType::ToggleButton)
-			.SetUncheckedImage( FSlateNoResource() )
+			.SetUncheckedImage( BOX_BRUSH( "Common/RoundedSelection_16x", 4.0f/16.0f, FLinearColor( 1, 1, 1, 0.1f ) ) )
 			.SetUncheckedHoveredImage( BOX_BRUSH( "Common/RoundedSelection_16x", 4.0f/16.0f, SelectionColor ) )
 			.SetUncheckedPressedImage( BOX_BRUSH( "Common/RoundedSelection_16x", 4.0f/16.0f, SelectionColor_Pressed ) )
 			.SetCheckedImage( BOX_BRUSH( "Common/RoundedSelection_16x",  4.0f/16.0f, SelectionColor_Pressed ) )
