Asset labels in the content browser can now fill beyond the width of the thumbnails.

The font has also been reduced in size to reduce the amount of wrapping on asset labels at small sizes.

Also polished the asset borders and shadows.

This addresses TTP 345457 and 345452.

diff --git a/Engine/Plugins/Editor/EpicSurvey/Source/EpicSurvey/Private/SSurvey.cpp b/Engine/Plugins/Editor/EpicSurvey/Source/EpicSurvey/Private/SSurvey.cpp
--- a/Engine/Plugins/Editor/EpicSurvey/Source/EpicSurvey/Private/SSurvey.cpp
+++ b/Engine/Plugins/Editor/EpicSurvey/Source/EpicSurvey/Private/SSurvey.cpp
@@ -35,7 +35,7 @@ void SSurvey::Construct( const FArguments& Args, const TSharedRef< FEpicSurvey >
 			[
 				// Drop shadow border
 				SNew(SBorder)
-				.Padding( 5.0f )
+				.Padding( 4.0f )
 				.BorderImage( FEditorStyle::GetBrush("ContentBrowser.ThumbnailShadow") )
 				[
 					SAssignNew( ContentsContainer, SBorder )
diff --git a/Engine/Source/Developer/ExternalImagePicker/Private/SExternalImagePicker.cpp b/Engine/Source/Developer/ExternalImagePicker/Private/SExternalImagePicker.cpp
--- a/Engine/Source/Developer/ExternalImagePicker/Private/SExternalImagePicker.cpp
+++ b/Engine/Source/Developer/ExternalImagePicker/Private/SExternalImagePicker.cpp
@@ -35,7 +35,7 @@ void SExternalImagePicker::Construct(const FArguments& InArgs)
 			[
 				SNew(SBorder)
 				.BorderImage(FEditorStyle::Get().GetBrush("ExternalImagePicker.ThumbnailShadow"))
-				.Padding(5.0f)
+				.Padding(4.0f)
 				.Content()
 				[
 					SNew(SBorder)
diff --git a/Engine/Source/Editor/BspMode/Private/BspModeStyle.cpp b/Engine/Source/Editor/BspMode/Private/BspModeStyle.cpp
--- a/Engine/Source/Editor/BspMode/Private/BspModeStyle.cpp
+++ b/Engine/Source/Editor/BspMode/Private/BspModeStyle.cpp
@@ -114,7 +114,7 @@ void FBspModeStyle::Initialize()
 		.SetSelectedTextColor( InvertedForeground )
 		);
 
-	StyleSet->Set( "BspMode.ThumbnailShadow", new BOX_BRUSH( "ContentBrowser/ThumbnailShadow" , FMargin( 5.0f / 64.0f ) ) );
+	StyleSet->Set( "BspMode.ThumbnailShadow", new BOX_BRUSH( "ContentBrowser/ThumbnailShadow" , FMargin( 4.0f / 64.0f ) ) );
 	StyleSet->Set( "BspMode.ThumbnailBackground", new IMAGE_BRUSH( "Common/ClassBackground_64x", FVector2D(64.f, 64.f), FLinearColor(0.75f, 0.75f, 0.75f, 1.0f) ) );
 	StyleSet->Set( "BspMode.ThumbnailText", NormalText );
 
diff --git a/Engine/Source/Editor/BspMode/Private/SBspPalette.cpp b/Engine/Source/Editor/BspMode/Private/SBspPalette.cpp
--- a/Engine/Source/Editor/BspMode/Private/SBspPalette.cpp
+++ b/Engine/Source/Editor/BspMode/Private/SBspPalette.cpp
@@ -155,7 +155,7 @@ TSharedRef<ITableRow> SBspPalette::MakeListViewWidget(TSharedPtr<FBspBuilderType
 			.AutoWidth()
 			[
 				SNew(SBorder)
-				.Padding(5.0f)
+				.Padding(4.0f)
 				.BorderImage(FBspModeStyle::Get().GetBrush("BspMode.ThumbnailShadow"))
 				[
 					SNew(SBox)
diff --git a/Engine/Source/Editor/ContentBrowser/Private/AssetViewWidgets.cpp b/Engine/Source/Editor/ContentBrowser/Private/AssetViewWidgets.cpp
--- a/Engine/Source/Editor/ContentBrowser/Private/AssetViewWidgets.cpp
+++ b/Engine/Source/Editor/ContentBrowser/Private/AssetViewWidgets.cpp
@@ -973,12 +973,13 @@ void SAssetListItem::Construct( const FArguments& InArgs )
 			.VAlign(VAlign_Center)
 			[
 				SNew( SBox )
+				.Padding(ThumbnailPadding - 4.f)
 				.WidthOverride( this, &SAssetListItem::GetThumbnailBoxSize )
 				.HeightOverride( this, &SAssetListItem::GetThumbnailBoxSize )
 				[
 					// Drop shadow border
 					SNew(SBorder)
-					.Padding(ThumbnailPadding)
+					.Padding(4.f)
 					.BorderImage( IsFolder() ? FEditorStyle::GetBrush("NoBorder") : FEditorStyle::GetBrush("ContentBrowser.ThumbnailShadow") )
 					[
 						SNew(SOverlay)
@@ -1145,7 +1146,7 @@ void SAssetTileItem::Construct( const FArguments& InArgs )
 		);
 
 	AssetThumbnail = InArgs._AssetThumbnail;
-	ItemWidth = InArgs._ItemWidth;
+	ThumbnailSize = InArgs._ThumbnailSize;
 	ThumbnailPadding = IsFolder() ? InArgs._ThumbnailPadding + 5.0f : InArgs._ThumbnailPadding;
 
 	TSharedPtr<SWidget> Thumbnail;
@@ -1182,16 +1183,17 @@ void SAssetTileItem::Construct( const FArguments& InArgs )
 			// Thumbnail
 			+SVerticalBox::Slot()
 			.AutoHeight()
+			.HAlign(HAlign_Center)
 			[
-				// Use ItemWidth for both the width and height to make it square.
 				// The remainder of the space is reserved for the name.
 				SNew(SBox)
+				.Padding(ThumbnailPadding - 4.f)
 				.WidthOverride( this, &SAssetTileItem::GetThumbnailBoxSize )
 				.HeightOverride( this, &SAssetTileItem::GetThumbnailBoxSize )
 				[
 					// Drop shadow border
 					SNew(SBorder)
-					.Padding(ThumbnailPadding)
+					.Padding(4.f)
 					.BorderImage(IsFolder() ? FEditorStyle::GetBrush("NoBorder") : FEditorStyle::GetBrush("ContentBrowser.ThumbnailShadow"))
 					[
 						SNew(SOverlay)
@@ -1260,7 +1262,7 @@ void SAssetTileItem::Construct( const FArguments& InArgs )
 			]
 
 			+SVerticalBox::Slot()
-			.Padding(FMargin(ThumbnailPadding, 0, ThumbnailPadding, ThumbnailPadding))
+			.Padding(FMargin(2.f, 0, 2.f, ThumbnailPadding))
 			.HAlign(HAlign_Center)
 			.VAlign(VAlign_Center)
 			.FillHeight(1.f)
@@ -1302,14 +1304,14 @@ void SAssetTileItem::OnAssetDataChanged()
 	SetToolTip( CreateToolTipWidget() );
 }
 
-FOptionalSize SAssetTileItem::GetThumbnailBoxSize() const
+FOptionalSize SAssetTileItem::GetSCCImageSize() const
 {
-	return FOptionalSize( ItemWidth.Get() );
+	return GetThumbnailBoxSize().Get() * 0.2;
 }
 
-FOptionalSize SAssetTileItem::GetSCCImageSize() const
+FOptionalSize SAssetTileItem::GetThumbnailBoxSize() const
 {
-	return GetThumbnailBoxSize().Get() * 0.2;
+	return ThumbnailSize.Get(16);
 }
 
 FSlateFontInfo SAssetTileItem::GetThumbnailFont() const
@@ -1331,7 +1333,8 @@ FSlateFontInfo SAssetTileItem::GetThumbnailFont() const
 
 float SAssetTileItem::GetNameTextWrapWidth() const
 {
-	return ItemWidth.Get() - ThumbnailPadding * 2.f;
+	// Wrap to the entire size of the tile, minus some padding
+	return LastGeometry.Size.X - 4.f;
 }
 
 ///////////////////////////////
diff --git a/Engine/Source/Editor/ContentBrowser/Private/AssetViewWidgets.h b/Engine/Source/Editor/ContentBrowser/Private/AssetViewWidgets.h
--- a/Engine/Source/Editor/ContentBrowser/Private/AssetViewWidgets.h
+++ b/Engine/Source/Editor/ContentBrowser/Private/AssetViewWidgets.h
@@ -379,7 +379,7 @@ public:
 		, _ThumbnailLabel( EThumbnailLabel::ClassName )
 		, _ThumbnailHintColorAndOpacity( FLinearColor( 0.0f, 0.0f, 0.0f, 0.0f ) )
 		, _AllowThumbnailHintLabel(true)
-		, _ItemWidth(16)
+		, _ThumbnailSize(16)
 		, _ShouldAllowToolTip(true)
 		, _ThumbnailEditMode(false)
 	{}
@@ -402,8 +402,8 @@ public:
 		/** Whether the thumbnail should ever show it's hint label */
 		SLATE_ARGUMENT( bool, AllowThumbnailHintLabel )
 
-		/** The width of the item */
-		SLATE_ATTRIBUTE( float, ItemWidth )
+		/** The width of the thumbnail */
+		SLATE_ATTRIBUTE( float, ThumbnailSize )
 
 		/** Delegate for when an asset name has entered a rename state */
 		SLATE_EVENT( FOnRenameBegin, OnRenameBegin )
@@ -473,7 +473,7 @@ private:
 	TSharedPtr<FAssetThumbnail> AssetThumbnail;
 
 	/** The width of the item. Used to enforce a square thumbnail. */
-	TAttribute<float> ItemWidth;
+	TAttribute<float> ThumbnailSize;
 
 	/** The padding allotted for the thumbnail */
 	float ThumbnailPadding;
diff --git a/Engine/Source/Editor/ContentBrowser/Private/SAssetView.cpp b/Engine/Source/Editor/ContentBrowser/Private/SAssetView.cpp
--- a/Engine/Source/Editor/ContentBrowser/Private/SAssetView.cpp
+++ b/Engine/Source/Editor/ContentBrowser/Private/SAssetView.cpp
@@ -110,7 +110,7 @@ void SAssetView::Construct( const FArguments& InArgs )
 	TileViewThumbnailResolution = 256;
 	TileViewThumbnailSize = 128;
 	TileViewThumbnailPadding = 5;
-	TileViewNameHeight = 36;
+	TileViewNameHeight = 32;
 	ThumbnailScaleSliderValue = InArgs._ThumbnailScale; 
 
 	if ( !ThumbnailScaleSliderValue.IsBound() )
@@ -118,7 +118,7 @@ void SAssetView::Construct( const FArguments& InArgs )
 		ThumbnailScaleSliderValue = FMath::Clamp<float>(ThumbnailScaleSliderValue.Get(), 0.0f, 1.0f);
 	}
 
-	MinThumbnailScale = 0.4f * ThumbnailScaleRangeScalar;
+	MinThumbnailScale = 0.3f * ThumbnailScaleRangeScalar;
 	MaxThumbnailScale = 2.0f * ThumbnailScaleRangeScalar;
 
 	bCanShowClasses = InArgs._CanShowClasses;
@@ -786,6 +786,8 @@ void SAssetView::Tick(const FGeometry& AllottedGeometry, const double InCurrentT
 
 	CalculateFillScale( AllottedGeometry );
 
+	TileViewNameHeight = int(FMath::Min(GetTileViewThumbnailSize(), 32.f));
+
 	CurrentTime = InCurrentTime;
 
 	// If there were any assets that were recently added via the asset registry, process them now
@@ -2695,7 +2697,7 @@ TSharedRef<ITableRow> SAssetView::MakeTileViewWidget(TSharedPtr<FAssetViewItem>
 		TSharedRef<SAssetTileItem> Item =
 			SNew(SAssetTileItem)
 			.AssetItem(AssetItem)
-			.ItemWidth(this, &SAssetView::GetTileViewItemWidth)
+			.ThumbnailSize(this, &SAssetView::GetTileViewThumbnailSize)
 			.OnRenameBegin(this, &SAssetView::AssetRenameBegin)
 			.OnRenameCommit(this, &SAssetView::AssetRenameCommit)
 			.OnVerifyRenameCommit(this, &SAssetView::AssetVerifyRenameCommit)
@@ -2740,7 +2742,7 @@ TSharedRef<ITableRow> SAssetView::MakeTileViewWidget(TSharedPtr<FAssetViewItem>
 			.AssetThumbnail(AssetThumbnail)
 			.AssetItem(AssetItem)
 			.ThumbnailPadding(TileViewThumbnailPadding)
-			.ItemWidth(this, &SAssetView::GetTileViewItemWidth)
+			.ThumbnailSize(this, &SAssetView::GetTileViewThumbnailSize)
 			.OnRenameBegin(this, &SAssetView::AssetRenameBegin)
 			.OnRenameCommit(this, &SAssetView::AssetRenameCommit)
 			.OnVerifyRenameCommit(this, &SAssetView::AssetVerifyRenameCommit)
@@ -3507,12 +3509,12 @@ float SAssetView::GetListViewItemHeight() const
 
 float SAssetView::GetTileViewItemHeight() const
 {
-	return TileViewNameHeight + GetTileViewItemBaseHeight() * FillScale;
+	return TileViewNameHeight + GetTileViewThumbnailSize();
 }
 
 float SAssetView::GetTileViewItemBaseHeight() const
 {
-	return (TileViewThumbnailSize + TileViewThumbnailPadding * 2) * FMath::Lerp(MinThumbnailScale, MaxThumbnailScale, GetThumbnailScale());
+	return TileViewThumbnailSize * FMath::Lerp(MinThumbnailScale, MaxThumbnailScale, GetThumbnailScale()) + TileViewThumbnailPadding * 2;
 }
 
 float SAssetView::GetTileViewItemWidth() const
@@ -3522,7 +3524,12 @@ float SAssetView::GetTileViewItemWidth() const
 
 float SAssetView::GetTileViewItemBaseWidth() const
 {
-	return ( TileViewThumbnailSize + TileViewThumbnailPadding * 2 ) * FMath::Lerp( MinThumbnailScale, MaxThumbnailScale, GetThumbnailScale() );
+	return FMath::Max(75.f, GetTileViewThumbnailSize() + TileViewThumbnailPadding * 2);
+}
+
+float SAssetView::GetTileViewThumbnailSize() const
+{
+	return TileViewThumbnailSize * FMath::Lerp( MinThumbnailScale, MaxThumbnailScale, GetThumbnailScale() );
 }
 
 EColumnSortMode::Type SAssetView::GetColumnSortMode(const FName ColumnId) const
diff --git a/Engine/Source/Editor/ContentBrowser/Private/SAssetView.h b/Engine/Source/Editor/ContentBrowser/Private/SAssetView.h
--- a/Engine/Source/Editor/ContentBrowser/Private/SAssetView.h
+++ b/Engine/Source/Editor/ContentBrowser/Private/SAssetView.h
@@ -515,6 +515,7 @@ private:
 
 	/** Gets the scaled item width for the tile view before the filler scale is applied */
 	float GetTileViewItemBaseWidth() const;
+	float GetTileViewThumbnailSize() const;
 
 	/** Gets the sort mode for the supplied ColumnId */
 	EColumnSortMode::Type GetColumnSortMode(const FName ColumnId) const;
diff --git a/Engine/Source/Editor/EditorStyle/Private/SlateEditorStyle.cpp b/Engine/Source/Editor/EditorStyle/Private/SlateEditorStyle.cpp
--- a/Engine/Source/Editor/EditorStyle/Private/SlateEditorStyle.cpp
+++ b/Engine/Source/Editor/EditorStyle/Private/SlateEditorStyle.cpp
@@ -1217,16 +1217,16 @@ void FSlateEditorStyle::FStyle::SetupGeneralStyles()
 		Set( "AssetThumbnail.AssetBackground", new IMAGE_BRUSH( "Common/AssetBackground", FVector2D(64.f, 64.f), FLinearColor(1.f, 1.f, 1.f, 1.0f) ) );
 		Set( "AssetThumbnail.ClassBackground", new IMAGE_BRUSH( "Common/ClassBackground_64x", FVector2D(64.f, 64.f), FLinearColor(0.75f, 0.75f, 0.75f, 1.0f) ) );
 		Set( "AssetThumbnail.DataOnlyBPAssetBackground", new IMAGE_BRUSH( "Common/DataOnlyBPAssetBackground_64x", FVector2D(64.f, 64.f), FLinearColor(1, 1, 1, 1) ) );
-		Set( "AssetThumbnail.Font", TTF_CORE_FONT( "Fonts/Roboto-Regular", 12 ) );
-		Set( "AssetThumbnail.FontSmall", TTF_CORE_FONT( "Fonts/Roboto-Regular", 8 ) );
+		Set( "AssetThumbnail.Font", TTF_CORE_FONT( "Fonts/Roboto-Regular", 10 ) );
+		Set( "AssetThumbnail.FontSmall", TTF_CORE_FONT( "Fonts/Roboto-Regular", 7 ) );
 		Set( "AssetThumbnail.ColorAndOpacity", FLinearColor(0.75f, 0.75f, 0.75f, 1) );
 		Set( "AssetThumbnail.ShadowOffset", FVector2D(1,1) );
-		Set( "AssetThumbnail.ShadowColorAndOpacity", FLinearColor(0, 0, 0, 1) );
+		Set( "AssetThumbnail.ShadowColorAndOpacity", FLinearColor(0, 0, 0, 0.5) );
 		Set( "AssetThumbnail.HintFont", TTF_CORE_FONT( "Fonts/Roboto-Regular", 8 ) );
 		Set( "AssetThumbnail.HintFontSmall", TTF_CORE_FONT( "Fonts/Roboto-Regular", 6 ) );
 		Set( "AssetThumbnail.HintColorAndOpacity", FLinearColor(0.75f, 0.75f, 0.75f, 1) );
 		Set( "AssetThumbnail.HintShadowOffset", FVector2D(1,1) );
-		Set( "AssetThumbnail.HintShadowColorAndOpacity", FLinearColor(0, 0, 0, 1) );
+		Set( "AssetThumbnail.HintShadowColorAndOpacity", FLinearColor(0, 0, 0, 0.5) );
 		Set( "AssetThumbnail.HintBackground", new BOX_BRUSH( "Common/TableViewHeader", FMargin(8.0f/32.0f) ) );
 		Set( "AssetThumbnail.Border", new FSlateColorBrush( FColor::White ) );
 	}
@@ -2674,8 +2674,8 @@ void FSlateEditorStyle::FStyle::SetupPropertyEditorStyles()
 			.SetShadowColorAndOpacity( FLinearColor::Black )
 		);
 
-		Set( "PropertyEditor.AssetThumbnailShadow", new BOX_BRUSH( "ContentBrowser/ThumbnailShadow", FMargin( 5.0f / 60.0f ) ) );
-		Set( "PropertyEditor.AssetThumbnailLight", new BOX_BRUSH( "ContentBrowser/ThumbnailLight", FMargin( 5.0f / 60.0f ), SelectionColor ) );
+		Set( "PropertyEditor.AssetThumbnailLight", new BOX_BRUSH( "ContentBrowser/ThumbnailLight", FMargin( 5.0f / 64.0f ), SelectionColor ) );
+		Set( "PropertyEditor.AssetThumbnailShadow", new BOX_BRUSH( "ContentBrowser/ThumbnailShadow", FMargin( 4.0f / 64.0f ) ) );
 		Set( "PropertyEditor.AssetClass", FTextBlockStyle(NormalText)
 			.SetFont( TTF_CORE_FONT( "Fonts/Roboto-Regular", 10 ) )
 			.SetColorAndOpacity( FLinearColor::White )
@@ -4745,8 +4745,8 @@ void FSlateEditorStyle::FStyle::SetupContentBrowserStyle()
 		);
 
 		// Tile view
-		Set( "ContentBrowser.AssetTileViewNameFont", TTF_CORE_FONT( "Fonts/Roboto-Regular", 10 ) );
-		Set( "ContentBrowser.AssetTileViewNameFontSmall", TTF_CORE_FONT( "Fonts/Roboto-Regular", 8 ) );
+		Set( "ContentBrowser.AssetTileViewNameFont", TTF_CORE_FONT( "Fonts/Roboto-Regular", 8 ) );
+		Set( "ContentBrowser.AssetTileViewNameFontSmall", TTF_CORE_FONT( "Fonts/Roboto-Regular", 7 ) );
 		Set( "ContentBrowser.AssetTileViewNameFontDirty", TTF_CORE_FONT( "Fonts/Roboto-Bold", 10 ) );
 		Set( "ContentBrowser.AssetListView.TableRow", FTableRowStyle()
 			.SetEvenRowBackgroundBrush( FSlateNoResource() )
@@ -4827,7 +4827,7 @@ void FSlateEditorStyle::FStyle::SetupContentBrowserStyle()
 		Set( "ContentBrowser.ImportPackage", new IMAGE_BRUSH( "Icons/icon_Import_40x", Icon25x25 ) );
 
 		// Misc
-		Set( "ContentBrowser.ThumbnailShadow", new BOX_BRUSH( "ContentBrowser/ThumbnailShadow" , FMargin( 5.0f / 64.0f ) ) );
+		Set( "ContentBrowser.ThumbnailShadow", new BOX_BRUSH( "ContentBrowser/ThumbnailShadow" , FMargin( 4.0f / 64.0f ) ) );
 		Set( "AssetDiscoveryIndicator.DiscovertingAssetsFont", TTF_CORE_FONT( "Fonts/Roboto-Regular", 12 ) );
 		Set( "ContentBrowser.ColumnViewAssetIcon", new IMAGE_BRUSH( "Icons/doc_16x", Icon16x16 ) );
 		Set( "ContentBrowser.ColumnViewFolderIcon", new IMAGE_BRUSH( "Icons/FolderClosed", FVector2D(18, 16) ) );
@@ -5689,7 +5689,7 @@ void FSlateEditorStyle::FStyle::SetupAutomationStyles()
 	// External image picker
 	{
 		Set("ExternalImagePicker.BlankImage", new IMAGE_BRUSH( "Icons/BlankIcon", Icon16x16 ) );
-		Set("ExternalImagePicker.ThumbnailShadow", new BOX_BRUSH( "ContentBrowser/ThumbnailShadow" , FMargin( 5.0f / 64.0f ) ) );
+		Set("ExternalImagePicker.ThumbnailShadow", new BOX_BRUSH( "ContentBrowser/ThumbnailShadow" , FMargin( 4.0f / 64.0f ) ) );
 		Set("ExternalImagePicker.PickImageButton", new IMAGE_BRUSH( "Icons/ellipsis_12x", Icon12x12 ) );
 	}
 
diff --git a/Engine/Source/Editor/GameProjectGeneration/Private/SNewProjectWizard.cpp b/Engine/Source/Editor/GameProjectGeneration/Private/SNewProjectWizard.cpp
--- a/Engine/Source/Editor/GameProjectGeneration/Private/SNewProjectWizard.cpp
+++ b/Engine/Source/Editor/GameProjectGeneration/Private/SNewProjectWizard.cpp
@@ -417,7 +417,7 @@ void SNewProjectWizard::Construct( const FArguments& InArgs )
 													+ SOverlay::Slot()
 													[
 														SNew(SBorder)
-														.Padding(FMargin(0, 0, 0, 5.f))
+														.Padding(FMargin(0, 0, 0, 4.f))
 														.BorderImage(FEditorStyle::GetBrush("ContentBrowser.ThumbnailShadow"))
 														[
 															SNew(SImage)
diff --git a/Engine/Source/Editor/GameProjectGeneration/Private/SProjectBrowser.cpp b/Engine/Source/Editor/GameProjectGeneration/Private/SProjectBrowser.cpp
--- a/Engine/Source/Editor/GameProjectGeneration/Private/SProjectBrowser.cpp
+++ b/Engine/Source/Editor/GameProjectGeneration/Private/SProjectBrowser.cpp
@@ -98,7 +98,7 @@ BEGIN_SLATE_FUNCTION_BUILD_OPTIMIZATION
 void SProjectBrowser::Construct( const FArguments& InArgs )
 {
 	bPreventSelectionChangeEvent = false;
-	ThumbnailBorderPadding = 5;
+	ThumbnailBorderPadding = 4;
 	ThumbnailSize = 128.0f;
 
 	// Prepare the categories box
diff --git a/Engine/Source/Editor/PlacementMode/Private/SPlacementModeTools.cpp b/Engine/Source/Editor/PlacementMode/Private/SPlacementModeTools.cpp
--- a/Engine/Source/Editor/PlacementMode/Private/SPlacementModeTools.cpp
+++ b/Engine/Source/Editor/PlacementMode/Private/SPlacementModeTools.cpp
@@ -171,7 +171,7 @@ void SPlacementAssetEntry::Construct(const FArguments& InArgs, UActorFactory* In
 			[
 				// Drop shadow border
 				SNew( SBorder )
-				.Padding( 5 )
+				.Padding( 4 )
 				.BorderImage( FEditorStyle::GetBrush( "ContentBrowser.ThumbnailShadow" ) )
 				[
 					SNew( SBox )
diff --git a/Engine/Source/Editor/PropertyEditor/Private/UserInterface/PropertyEditor/SPropertyEditorAsset.cpp b/Engine/Source/Editor/PropertyEditor/Private/UserInterface/PropertyEditor/SPropertyEditorAsset.cpp
--- a/Engine/Source/Editor/PropertyEditor/Private/UserInterface/PropertyEditor/SPropertyEditorAsset.cpp
+++ b/Engine/Source/Editor/PropertyEditor/Private/UserInterface/PropertyEditor/SPropertyEditorAsset.cpp
@@ -117,7 +117,7 @@ void SPropertyEditorAsset::Construct( const FArguments& InArgs, const TSharedPtr
 				.AutoWidth()
 				[
 					SAssignNew( ThumbnailBorder, SBorder )
-					.Padding( 3.0f )
+					.Padding( 5.0f )
 					.BorderImage( this, &SPropertyEditorAsset::GetThumbnailBorder )
 					.OnMouseDoubleClick( this, &SPropertyEditorAsset::OnAssetThumbnailDoubleClick )
 					[
diff --git a/Engine/Source/Editor/UnrealEd/Private/AssetThumbnail.cpp b/Engine/Source/Editor/UnrealEd/Private/AssetThumbnail.cpp
--- a/Engine/Source/Editor/UnrealEd/Private/AssetThumbnail.cpp
+++ b/Engine/Source/Editor/UnrealEd/Private/AssetThumbnail.cpp
@@ -52,7 +52,6 @@ public:
 		AssetThumbnail = InArgs._AssetThumbnail;
 		bHasRenderedThumbnail = false;
 		WidthLastFrame = 0;
-		GenericThumbnailBorderPadding = 2.f;
 
 		AssetThumbnail->OnAssetDataChanged().AddSP(this, &SAssetThumbnail::OnAssetDataChanged);
 
@@ -93,7 +92,7 @@ public:
 			SNew(SBorder)
 			.BorderImage( this, &SAssetThumbnail::GetBackgroundBrush )
 			.BorderBackgroundColor( this, &SAssetThumbnail::GetAssetColor )
-			.Padding( GenericThumbnailBorderPadding )
+			.Padding( 0 )
 			.VAlign(VAlign_Center)
 			.HAlign(HAlign_Center)
 			.Visibility(this, &SAssetThumbnail::GetClassThumbnailVisibility)
@@ -109,7 +108,7 @@ public:
 			SNew(SBorder)
 			.BorderImage( this, &SAssetThumbnail::GetBackgroundBrush )
 			.BorderBackgroundColor( this, &SAssetThumbnail::GetAssetColor )
-			.Padding( GenericThumbnailBorderPadding )
+			.Padding( 0 )
 			.VAlign(VAlign_Center) .HAlign(HAlign_Center)
 			.Visibility(this, &SAssetThumbnail::GetGenericThumbnailVisibility)
 			[
@@ -149,10 +148,8 @@ public:
 			OverlayWidget->AddSlot()
 			[
 				SNew(SBorder)
-				.Padding(1)
-				.BorderImage(FEditorStyle::GetBrush(Style, ".Border"))
-				//.BorderImage(FEditorStyle::GetBrush("NoBrush"))
-				.BorderBackgroundColor(this, &SAssetThumbnail::GetViewportBorderColorAndOpacity)
+				.Padding(0)
+				.BorderImage(FEditorStyle::GetBrush("NoBrush"))
 				.ColorAndOpacity(this, &SAssetThumbnail::GetViewportColorAndOpacity)
 				.Visibility(this, &SAssetThumbnail::GetViewportVisibility)
 				[
@@ -184,6 +181,18 @@ public:
 			]
 		];
 
+		// The asset color strip, only visible when the rendered viewport is
+		OverlayWidget->AddSlot()
+		.HAlign(HAlign_Fill)
+		.VAlign(VAlign_Bottom)
+		[
+			SNew(SBorder)
+			.BorderImage(FEditorStyle::GetBrush("WhiteBrush"))
+			.BorderBackgroundColor(AssetColor)
+			.Visibility(this, &SAssetThumbnail::GetViewportVisibility)
+			.Padding(this, &SAssetThumbnail::GetAssetColorStripPadding)
+		];
+
 		ChildSlot
 		[
 			OverlayWidget
@@ -292,7 +301,8 @@ private:
 
 	float GetTextWrapWidth() const
 	{
-		return WidthLastFrame - GenericThumbnailBorderPadding * 2.f;
+		/** 2px padding either side of the text */
+		return WidthLastFrame;// - 4.f;
 	}
 
 	const FSlateBrush* GetBackgroundBrush() const
@@ -332,6 +342,13 @@ private:
 		return bHasRenderedThumbnail ? EVisibility::Visible : EVisibility::Collapsed;
 	}
 
+	FMargin GetAssetColorStripPadding() const
+	{
+		// The strip is 2.5% the height of the thumbnail
+		const float Height = FMath::CeilToFloat(WidthLastFrame*0.025f);
+		return FMargin(0,Height,0,0);
+	}
+
 	const FSlateBrush* GetClassThumbnailBrush() const
 	{
 		if (ClassThumbnailBrushOverride.IsNone())
@@ -534,7 +551,6 @@ private:
 
 	FLinearColor AssetColor;
 	float WidthLastFrame;
-	float GenericThumbnailBorderPadding;
 	bool bHasRenderedThumbnail;
 	FName Style;
 	TAttribute< FText > HighlightedText;
