Missing style addition

diff --git a/Engine/Source/Editor/EditorStyle/Private/SlateEditorStyle.cpp b/Engine/Source/Editor/EditorStyle/Private/SlateEditorStyle.cpp
--- a/Engine/Source/Editor/EditorStyle/Private/SlateEditorStyle.cpp
+++ b/Engine/Source/Editor/EditorStyle/Private/SlateEditorStyle.cpp
@@ -3626,6 +3626,8 @@ void FSlateEditorStyle::FStyle::SetupGraphEditorStyles()
 		}
 
 		// SCSEditor
+		Set("SCSEditor.NewComponentIcon", new IMAGE_BRUSH("Icons/icon_Blueprint_SpawnActor_16x", Icon16x16));
+
 		Set("SCSEditor.ToggleComponentEditing" , new IMAGE_BRUSH( "Icons/icon_translate_40x", Icon40x40) );
 		Set("SCSEditor.ToggleComponentEditing.Small" , new IMAGE_BRUSH( "Icons/icon_translate_40x", Icon20x20) );
 
