#TTP 345364 - EDITOR ART REQUEST: Actor classes missing icons

renamed new icons so they marry up with the class names

diff --git a/Engine/Source/Editor/EditorStyle/Private/SlateEditorStyle.cpp b/Engine/Source/Editor/EditorStyle/Private/SlateEditorStyle.cpp
--- a/Engine/Source/Editor/EditorStyle/Private/SlateEditorStyle.cpp
+++ b/Engine/Source/Editor/EditorStyle/Private/SlateEditorStyle.cpp
@@ -4631,20 +4631,24 @@ void FSlateEditorStyle::FStyle::SetupClassIconsAndThumbnails()
 			TEXT("DirectionalLightMovable"),
 			TEXT("DirectionalLightStatic"),
 			TEXT("DirectionalLightStationary"),
+			TEXT("DocumentationActor"),
 			TEXT("ExponentialHeightFog"),
 			TEXT("Font"),
 			TEXT("ForceFeedbackEffect"),
 			TEXT("InterpData"),
 			TEXT("KillZVolume"),
 			TEXT("Landscape"),
+			TEXT("LevelBounds"),
 			TEXT("LevelStreamingVolume"),
 			TEXT("LightmassCharacterIndirectDetailVolume"),
 			TEXT("LightmassImportanceVolume"),
 			TEXT("MassiveLODOverrideVolume"),
 			TEXT("Material"),
 			TEXT("MaterialFunction"),
+			TEXT("MaterialInstanceActor"),
 			TEXT("MaterialInstanceConstant"),
 			TEXT("MaterialParameterCollection"),
+			TEXT("MatineeActor"),
 			TEXT("MediaPlayer"),
 			TEXT("MediaSoundWave"),
 			TEXT("MediaTexture"),
@@ -4655,8 +4659,11 @@ void FSlateEditorStyle::FStyle::SetupClassIconsAndThumbnails()
 			TEXT("Note"),
 			TEXT("ObjectLibrary"),
 			TEXT("PainCausingVolume"),
+			TEXT("PaperCharacter"),
 			TEXT("ParticleSystem"),
 			TEXT("PhysicalMaterial"),
+			TEXT("PhysicsConstraintActor"),
+			TEXT("PhysicsThruster"),
 			TEXT("PhysicsVolume"),
 			TEXT("PlaneReflectionCapture"),
 			TEXT("PlayerStart"),
