UMG "Toolbox" -> "Palette" name change to brushes

diff --git a/Engine/Source/Editor/EditorStyle/Private/SlateEditorStyle.cpp b/Engine/Source/Editor/EditorStyle/Private/SlateEditorStyle.cpp
--- a/Engine/Source/Editor/EditorStyle/Private/SlateEditorStyle.cpp
+++ b/Engine/Source/Editor/EditorStyle/Private/SlateEditorStyle.cpp
@@ -5810,9 +5810,9 @@ void FSlateEditorStyle::FStyle::SetupUMGEditorStyles()
 	Set("UMGEditor.FillRight", new IMAGE_BRUSH("Icons/UMG/umg_anchors_fill_right", Icon64x64));
 	Set("UMGEditor.FillFill", new IMAGE_BRUSH("Icons/UMG/umg_anchors_fill_fill", Icon64x64));
 	
-	Set("UMGEditor.ToolboxPreviewWidget", new IMAGE_BRUSH("Icons/UMG/umg_anchors_fill_fill", Icon16x16));
-	
-	Set("UMGEditor.ToolboxHeader", FTableRowStyle()
+	Set("UMGEditor.PalettePreviewWidget", new IMAGE_BRUSH("Icons/UMG/umg_anchors_fill_fill", Icon16x16));
+
+	Set("UMGEditor.PaletteHeader", FTableRowStyle()
 		.SetEvenRowBackgroundBrush( BOX_BRUSH( "PropertyView/DetailCategoryTop", FMargin( 4/16.0f, 8.0f/16.0f, 4/16.0f, 4/16.0f ) ) )
 		.SetEvenRowBackgroundHoveredBrush( BOX_BRUSH( "PropertyView/DetailCategoryTop", FMargin( 4/16.0f, 8.0f/16.0f, 4/16.0f, 4/16.0f ), FLinearColor(0.5f,0.5f,0.5f,1.0f) ) )
 		.SetOddRowBackgroundBrush( BOX_BRUSH( "PropertyView/DetailCategoryTop", FMargin( 4/16.0f, 8.0f/16.0f, 4/16.0f, 4/16.0f ) ) )
