Fixed missing slate brushes warning on startup

diff --git a/Engine/Source/Editor/EditorStyle/Private/SlateEditorStyle.cpp b/Engine/Source/Editor/EditorStyle/Private/SlateEditorStyle.cpp
--- a/Engine/Source/Editor/EditorStyle/Private/SlateEditorStyle.cpp
+++ b/Engine/Source/Editor/EditorStyle/Private/SlateEditorStyle.cpp
@@ -2501,10 +2501,10 @@ void FSlateEditorStyle::FStyle::SetupTutorialStyles()
 		Set( "Tutorials.Browser.RestartButton", new IMAGE_BRUSH( "Tutorials/BrowserRestart", FVector2D(16.0f, 16.0f), FLinearColor(0.0f, 0.0f, 0.0f, 1.0f) ) );
 
 		Set( "Tutorials.Navigation.Button", FButtonStyle()
-			.SetNormal( BOX_BRUSH( "Tutorials/ButtonHoverHint", FMargin(4/16.0f), FLinearColor(0,0,0,0) ) )
-			.SetHovered( BOX_BRUSH( "Tutorials/ButtonHoverHint", FMargin(4/16.0f), FLinearColor(0,0,0,0) ) )
-			.SetPressed( BOX_BRUSH( "Tutorials/ButtonHoverHint", FMargin(4/16.0f), FLinearColor(0,0,0,0) ) )
-			.SetDisabled( BOX_BRUSH( "Tutorials/ButtonHoverHint", FMargin(4/16.0f), FLinearColor(0,0,0,0) ) )
+			.SetNormal( BOX_BRUSH( "Common/ButtonHoverHint", FMargin(4/16.0f), FLinearColor(0,0,0,0) ) )
+			.SetHovered( BOX_BRUSH( "Common/ButtonHoverHint", FMargin(4/16.0f), FLinearColor(0,0,0,0) ) )
+			.SetPressed( BOX_BRUSH( "Common/ButtonHoverHint", FMargin(4/16.0f), FLinearColor(0,0,0,0) ) )
+			.SetDisabled( BOX_BRUSH( "Common/ButtonHoverHint", FMargin(4/16.0f), FLinearColor(0,0,0,0) ) )
 			.SetNormalPadding( FMargin(0,0,0,1))
 			.SetPressedPadding( FMargin(0,1,0,0)));
 
