Fixed play in viewport using simulate icon.

diff --git a/Engine/Source/Editor/EditorStyle/Private/SlateEditorStyle.cpp b/Engine/Source/Editor/EditorStyle/Private/SlateEditorStyle.cpp
--- a/Engine/Source/Editor/EditorStyle/Private/SlateEditorStyle.cpp
+++ b/Engine/Source/Editor/EditorStyle/Private/SlateEditorStyle.cpp
@@ -4496,8 +4496,8 @@ void FSlateEditorStyle::FStyle::SetupPersonaStyle()
 
 		Set( "PlayWorld.RepeatLastPlay", new IMAGE_BRUSH( "Icons/icon_simulate_40x", Icon40x40 ) );
 		Set( "PlayWorld.RepeatLastPlay.Small", new IMAGE_BRUSH( "Icons/icon_simulate_40x", Icon20x20 ) );
-		Set( "PlayWorld.PlayInViewport", new IMAGE_BRUSH( "Icons/icon_simulate_40x", Icon40x40 ) );
-		Set( "PlayWorld.PlayInViewport.Small", new IMAGE_BRUSH( "Icons/icon_simulate_40x", Icon20x20 ) );
+		Set( "PlayWorld.PlayInViewport", new IMAGE_BRUSH("Icons/icon_playInSelectedViewport_40x", Icon40x40 ) );
+		Set( "PlayWorld.PlayInViewport.Small", new IMAGE_BRUSH("Icons/icon_playInSelectedViewport_40x", Icon20x20 ) );
 		Set( "PlayWorld.PlayInEditorFloating", new IMAGE_BRUSH( "Icons/icon_playInWindow_40x", Icon40x40 ) );
 		Set( "PlayWorld.PlayInEditorFloating.Small", new IMAGE_BRUSH( "Icons/icon_playInWindow_40x", Icon20x20 ) );
 		Set( "PlayWorld.PlayInMobilePreview", new IMAGE_BRUSH( "Icons/icon_PlayMobilePreview_40x", Icon40x40 ) );
