Adding external link indicator to tutorial hyperlinks

Missed a file from my last submit somehow.

diff --git a/Engine/Source/Editor/EditorStyle/Private/SlateEditorStyle.cpp b/Engine/Source/Editor/EditorStyle/Private/SlateEditorStyle.cpp
--- a/Engine/Source/Editor/EditorStyle/Private/SlateEditorStyle.cpp
+++ b/Engine/Source/Editor/EditorStyle/Private/SlateEditorStyle.cpp
@@ -2643,6 +2643,8 @@ void FSlateEditorStyle::FStyle::SetupTutorialStyles()
 					.SetPadding(FMargin(0.0f));
 				Set( "Tutorials.Content.Hyperlink", RichTextHyperlink );
 
+				Set( "Tutorials.Content.ExternalLink", new IMAGE_BRUSH("Tutorials/ExternalLink", Icon16x16, FLinearColor::Blue));
+
 				// legacy style
 				Set( "TutorialEditableText.Editor.Hyperlink", RichTextHyperlink );
 			}
