Window menu refactor & polish

- Window menu is now sectioned and labeled based on the current editor. There's now a local workspace root member in FTabManager and a workspace category in FAssetEditorToolkit (both are FWorkspaceItem objects). Individual editors attach their local category to the tab manager's local root. Workflow app modes have their own category members that are swapped out when the mode changes.

- Finally, the AssetEditorCategory of FWorkspaceMenuStructure has been removed entirely.

- Replaced the AddMenuSeparator() call in FTabManager::PopulateSpawnerMenu_Helper() with a section of the same title as the workspace category.

- Tab spawner menu entries for the local editor now properly show the icon of the associated tab. To accomplish this it was necessary to change FWorkflowTabFactory::TabIcon to be an FSlateIcon instead of an FSlateBrush*. All factory instances have been updated accordingly.

- Added & updated lots of icons! (those missing will be TTP'd)

- The nomad tab spawner section (named "General" in the menu) has been largely compressed into the Developer Tools submenu, which has also been organized into sections for readability.

- Unreal frontend options were also moved into a context menu within the General section

- Moved all experimental tools to their own section of the Window menu. When they're no longer experimental they should register as nomads in the appropriate category

- Undo history now under Edit menu

diff --git a/Engine/Plugins/2D/Paper2D/Source/Paper2DEditor/Private/FlipbookEditor/FlipbookEditor.cpp b/Engine/Plugins/2D/Paper2D/Source/Paper2DEditor/Private/FlipbookEditor/FlipbookEditor.cpp
--- a/Engine/Plugins/2D/Paper2D/Source/Paper2DEditor/Private/FlipbookEditor/FlipbookEditor.cpp
+++ b/Engine/Plugins/2D/Paper2D/Source/Paper2DEditor/Private/FlipbookEditor/FlipbookEditor.cpp
@@ -299,18 +299,20 @@ TSharedRef<SDockTab> FFlipbookEditor::SpawnTab_Details(const FSpawnTabArgs& Args
 
 void FFlipbookEditor::RegisterTabSpawners(const TSharedRef<class FTabManager>& TabManager)
 {
-	FAssetEditorToolkit::RegisterTabSpawners(TabManager);
+	WorkspaceMenuCategory = TabManager->GetLocalWorkspaceMenuRoot()->AddGroup(LOCTEXT("WorkspaceMenu_FlipbookEditor", "Flipbook Editor"));
+	auto WorkspaceMenuCategoryRef = WorkspaceMenuCategory.ToSharedRef();
 
-	const IWorkspaceMenuStructure& MenuStructure = WorkspaceMenu::GetMenuStructure();
+	FAssetEditorToolkit::RegisterTabSpawners(TabManager);
 
 	TabManager->RegisterTabSpawner(FFlipbookEditorTabs::ViewportID, FOnSpawnTab::CreateSP(this, &FFlipbookEditor::SpawnTab_Viewport))
 		.SetDisplayName( LOCTEXT("ViewportTab", "Viewport") )
-		.SetGroup( MenuStructure.GetAssetEditorCategory() );
-
+		.SetGroup(WorkspaceMenuCategoryRef)
+		.SetIcon(FSlateIcon(FEditorStyle::GetStyleSetName(), "LevelEditor.Tabs.Viewports"));
 
 	TabManager->RegisterTabSpawner(FFlipbookEditorTabs::DetailsID, FOnSpawnTab::CreateSP(this, &FFlipbookEditor::SpawnTab_Details))
 		.SetDisplayName( LOCTEXT("DetailsTabLabel", "Details") )
-		.SetGroup( MenuStructure.GetAssetEditorCategory() );
+		.SetGroup(WorkspaceMenuCategoryRef)
+		.SetIcon(FSlateIcon(FEditorStyle::GetStyleSetName(), "LevelEditor.Tabs.Details"));
 }
 
 void FFlipbookEditor::UnregisterTabSpawners(const TSharedRef<class FTabManager>& TabManager)
diff --git a/Engine/Plugins/2D/Paper2D/Source/Paper2DEditor/Private/SpriteEditor/SpriteEditor.cpp b/Engine/Plugins/2D/Paper2D/Source/Paper2DEditor/Private/SpriteEditor/SpriteEditor.cpp
--- a/Engine/Plugins/2D/Paper2D/Source/Paper2DEditor/Private/SpriteEditor/SpriteEditor.cpp
+++ b/Engine/Plugins/2D/Paper2D/Source/Paper2DEditor/Private/SpriteEditor/SpriteEditor.cpp
@@ -308,7 +308,6 @@ TSharedRef<SDockTab> FSpriteEditor::SpawnTab_Details(const FSpawnTabArgs& Args)
 
 	// Spawn the tab
 	return SNew(SDockTab)
-		.Icon(FEditorStyle::GetBrush("LevelEditor.Tabs.Details"))
 		.Label(LOCTEXT("DetailsTab_Title", "Details"))
 		[
 			SNew(SSpritePropertiesTabBody, SpriteEditorPtr)
@@ -321,7 +320,6 @@ TSharedRef<SDockTab> FSpriteEditor::SpawnTab_SpriteList(const FSpawnTabArgs& Arg
 
 	// Spawn the tab
 	return SNew(SDockTab)
-		.Icon(FEditorStyle::GetBrush("LevelEditor.Tabs.ContentBrowser"))
 		.Label(LOCTEXT("SpriteListTab_Title", "Sprite List"))
 		[
 			SNew(SSpriteList, SpriteEditorPtr)
@@ -330,21 +328,25 @@ TSharedRef<SDockTab> FSpriteEditor::SpawnTab_SpriteList(const FSpawnTabArgs& Arg
 
 void FSpriteEditor::RegisterTabSpawners(const TSharedRef<class FTabManager>& TabManager)
 {
-	FAssetEditorToolkit::RegisterTabSpawners(TabManager);
+	WorkspaceMenuCategory = TabManager->GetLocalWorkspaceMenuRoot()->AddGroup(LOCTEXT("WorkspaceMenu_SpriteEditor", "Sprite Editor"));
+	auto WorkspaceMenuCategoryRef = WorkspaceMenuCategory.ToSharedRef();
 
-	const IWorkspaceMenuStructure& MenuStructure = WorkspaceMenu::GetMenuStructure();
+	FAssetEditorToolkit::RegisterTabSpawners(TabManager);
 
 	TabManager->RegisterTabSpawner(FSpriteEditorTabs::ViewportID, FOnSpawnTab::CreateSP(this, &FSpriteEditor::SpawnTab_Viewport))
 		.SetDisplayName( LOCTEXT("ViewportTab", "Viewport") )
-		.SetGroup( MenuStructure.GetAssetEditorCategory() );
+		.SetGroup(WorkspaceMenuCategoryRef)
+		.SetIcon(FSlateIcon(FEditorStyle::GetStyleSetName(), "LevelEditor.Tabs.Viewports"));
 
 	TabManager->RegisterTabSpawner(FSpriteEditorTabs::DetailsID, FOnSpawnTab::CreateSP(this, &FSpriteEditor::SpawnTab_Details))
 		.SetDisplayName( LOCTEXT("DetailsTabLabel", "Details") )
-		.SetGroup( MenuStructure.GetAssetEditorCategory() );
+		.SetGroup(WorkspaceMenuCategoryRef)
+		.SetIcon(FSlateIcon(FEditorStyle::GetStyleSetName(), "LevelEditor.Tabs.Details"));
 
 	TabManager->RegisterTabSpawner(FSpriteEditorTabs::SpriteListID, FOnSpawnTab::CreateSP(this, &FSpriteEditor::SpawnTab_SpriteList))
 		.SetDisplayName( LOCTEXT("SpriteListTabLabel", "Sprite List") )
-		.SetGroup( MenuStructure.GetAssetEditorCategory() );
+		.SetGroup(WorkspaceMenuCategoryRef)
+		.SetIcon(FSlateIcon(FEditorStyle::GetStyleSetName(), "LevelEditor.Tabs.ContentBrowser"));
 }
 
 void FSpriteEditor::UnregisterTabSpawners(const TSharedRef<class FTabManager>& TabManager)
diff --git a/Engine/Plugins/2D/Paper2D/Source/Paper2DEditor/Private/TileMapEditing/TileMapEditor.cpp b/Engine/Plugins/2D/Paper2D/Source/Paper2DEditor/Private/TileMapEditing/TileMapEditor.cpp
--- a/Engine/Plugins/2D/Paper2D/Source/Paper2DEditor/Private/TileMapEditing/TileMapEditor.cpp
+++ b/Engine/Plugins/2D/Paper2D/Source/Paper2DEditor/Private/TileMapEditing/TileMapEditor.cpp
@@ -231,17 +231,20 @@ TSharedRef<SDockTab> FTileMapEditor::SpawnTab_Details(const FSpawnTabArgs& Args)
 
 void FTileMapEditor::RegisterTabSpawners(const TSharedRef<class FTabManager>& TabManager)
 {
-	FAssetEditorToolkit::RegisterTabSpawners(TabManager);
+	WorkspaceMenuCategory = TabManager->GetLocalWorkspaceMenuRoot()->AddGroup(LOCTEXT("WorkspaceMenu_TileMapEditor", "Tile Map Editor"));
+	auto WorkspaceMenuCategoryRef = WorkspaceMenuCategory.ToSharedRef();
 
-	const IWorkspaceMenuStructure& MenuStructure = WorkspaceMenu::GetMenuStructure();
+	FAssetEditorToolkit::RegisterTabSpawners(TabManager);
 
 	TabManager->RegisterTabSpawner(FTileMapEditorTabs::ViewportID, FOnSpawnTab::CreateSP(this, &FTileMapEditor::SpawnTab_Viewport))
 		.SetDisplayName(LOCTEXT("ViewportTab", "Viewport"))
-		.SetGroup(MenuStructure.GetAssetEditorCategory());
+		.SetGroup(WorkspaceMenuCategoryRef)
+		.SetIcon(FSlateIcon(FEditorStyle::GetStyleSetName(), "LevelEditor.Tabs.Viewports"));
 
 	TabManager->RegisterTabSpawner(FTileMapEditorTabs::DetailsID, FOnSpawnTab::CreateSP(this, &FTileMapEditor::SpawnTab_Details))
 		.SetDisplayName(LOCTEXT("DetailsTabLabel", "Details"))
-		.SetGroup(MenuStructure.GetAssetEditorCategory());
+		.SetGroup(WorkspaceMenuCategoryRef)
+		.SetIcon(FSlateIcon(FEditorStyle::GetStyleSetName(), "LevelEditor.Tabs.Details"));
 }
 
 void FTileMapEditor::UnregisterTabSpawners(const TSharedRef<class FTabManager>& TabManager)
diff --git a/Engine/Plugins/2D/Paper2D/Source/PaperJsonImporter/PaperJsonImporter.Build.cs b/Engine/Plugins/2D/Paper2D/Source/PaperJsonImporter/PaperJsonImporter.Build.cs
--- a/Engine/Plugins/2D/Paper2D/Source/PaperJsonImporter/PaperJsonImporter.Build.cs
+++ b/Engine/Plugins/2D/Paper2D/Source/PaperJsonImporter/PaperJsonImporter.Build.cs
@@ -18,7 +18,8 @@ public class PaperJsonImporter : ModuleRules
 				"UnrealEd",
 				"Paper2DEditor",
 				"AssetTools",
-				"ContentBrowser"
+				"ContentBrowser",
+                "EditorStyle"
 			});
 
 		PrivateIncludePathModuleNames.AddRange(
diff --git a/Engine/Plugins/Experimental/LiveEditor/Source/LiveEditor/Private/LiveEditor.cpp b/Engine/Plugins/Experimental/LiveEditor/Source/LiveEditor/Private/LiveEditor.cpp
--- a/Engine/Plugins/Experimental/LiveEditor/Source/LiveEditor/Private/LiveEditor.cpp
+++ b/Engine/Plugins/Experimental/LiveEditor/Source/LiveEditor/Private/LiveEditor.cpp
@@ -36,7 +36,7 @@ class FComponentsEditorModeOverride : public FBlueprintEditorApplicationMode
 {
 public:
 	FComponentsEditorModeOverride(TSharedPtr<FBlueprintEditor> InBlueprintEditor, FName InModeName)
-		: FBlueprintEditorApplicationMode(InBlueprintEditor, InModeName)
+		: FBlueprintEditorApplicationMode(InBlueprintEditor, InModeName, FBlueprintEditorApplicationModes::GetLocalizedMode)
 	{
 	}
 
@@ -209,7 +209,7 @@ void FLiveEditor::InstallHooks()
 
 	FGlobalTabmanager::Get()->RegisterNomadTabSpawner(LiveEditorModule::LiveEditorApp, FOnSpawnTab::CreateStatic(&SpawnLiveEditorTab))
 		.SetDisplayName(NSLOCTEXT("LiveEditorPlugin", "TabTitle", "Live Editor"))
-		.SetGroup(WorkspaceMenu::GetMenuStructure().GetDeveloperToolsCategory())
+		.SetGroup(WorkspaceMenu::GetMenuStructure().GetDeveloperToolsMiscCategory())
 		.SetIcon(FSlateIcon(FEditorStyle::GetStyleSetName(), "LiveEditor.TabIcon"));
 }
 
diff --git a/Engine/Plugins/Media/MediaPlayerEditor/Source/MediaPlayerEditor/Private/MediaPlayerEditorToolkit.cpp b/Engine/Plugins/Media/MediaPlayerEditor/Source/MediaPlayerEditor/Private/MediaPlayerEditorToolkit.cpp
--- a/Engine/Plugins/Media/MediaPlayerEditor/Source/MediaPlayerEditor/Private/MediaPlayerEditorToolkit.cpp
+++ b/Engine/Plugins/Media/MediaPlayerEditor/Source/MediaPlayerEditor/Private/MediaPlayerEditorToolkit.cpp
@@ -107,17 +107,20 @@ FString FMediaPlayerEditorToolkit::GetDocumentationLink() const
 
 void FMediaPlayerEditorToolkit::RegisterTabSpawners( const TSharedRef<class FTabManager>& TabManager )
 {
+	WorkspaceMenuCategory = TabManager->GetLocalWorkspaceMenuRoot()->AddGroup(LOCTEXT("WorkspaceMenu_MediaPlayerEditor", "Media Player Editor"));
+	auto WorkspaceMenuCategoryRef = WorkspaceMenuCategory.ToSharedRef();
+
 	FAssetEditorToolkit::RegisterTabSpawners(TabManager);
 
-	const IWorkspaceMenuStructure& MenuStructure = WorkspaceMenu::GetMenuStructure();
+	TabManager->RegisterTabSpawner( ViewerTabId, FOnSpawnTab::CreateSP( this, &FMediaPlayerEditorToolkit::HandleTabManagerSpawnTab, ViewerTabId ) )
+		.SetDisplayName( LOCTEXT( "PlayerTabName", "Player" ) )
+		.SetGroup( WorkspaceMenuCategoryRef )
+		.SetIcon( FSlateIcon( FEditorStyle::GetStyleSetName(), "LevelEditor.Tabs.Viewports" ) );
 
-	TabManager->RegisterTabSpawner(ViewerTabId, FOnSpawnTab::CreateSP(this, &FMediaPlayerEditorToolkit::HandleTabManagerSpawnTab, ViewerTabId))
-		.SetDisplayName(LOCTEXT("PlayerTabName", "Player"))
-		.SetGroup(MenuStructure.GetAssetEditorCategory());
-	
-	TabManager->RegisterTabSpawner(DetailsTabId, FOnSpawnTab::CreateSP(this, &FMediaPlayerEditorToolkit::HandleTabManagerSpawnTab, DetailsTabId))
-		.SetDisplayName(LOCTEXT("DetailsTabName", "Details") )
-		.SetGroup(MenuStructure.GetAssetEditorCategory());
+	TabManager->RegisterTabSpawner( DetailsTabId, FOnSpawnTab::CreateSP( this, &FMediaPlayerEditorToolkit::HandleTabManagerSpawnTab, DetailsTabId ) )
+		.SetDisplayName( LOCTEXT( "DetailsTabName", "Details" ) )
+		.SetGroup( WorkspaceMenuCategoryRef )
+		.SetIcon( FSlateIcon( FEditorStyle::GetStyleSetName(), "LevelEditor.Tabs.Details" ) );
 }
 
 
diff --git a/Engine/Plugins/Messaging/MessagingDebugger/Source/MessagingDebugger/MessagingDebugger.Build.cs b/Engine/Plugins/Messaging/MessagingDebugger/Source/MessagingDebugger/MessagingDebugger.Build.cs
--- a/Engine/Plugins/Messaging/MessagingDebugger/Source/MessagingDebugger/MessagingDebugger.Build.cs
+++ b/Engine/Plugins/Messaging/MessagingDebugger/Source/MessagingDebugger/MessagingDebugger.Build.cs
@@ -15,6 +15,7 @@ public class MessagingDebugger : ModuleRules
 				"Serialization",
 				"Slate",
 				"SlateCore",
+                "WorkspaceMenuStructure",
 			}
 		);
 
diff --git a/Engine/Plugins/Messaging/MessagingDebugger/Source/MessagingDebugger/Private/MessagingDebuggerModule.cpp b/Engine/Plugins/Messaging/MessagingDebugger/Source/MessagingDebugger/Private/MessagingDebuggerModule.cpp
--- a/Engine/Plugins/Messaging/MessagingDebugger/Source/MessagingDebugger/Private/MessagingDebuggerModule.cpp
+++ b/Engine/Plugins/Messaging/MessagingDebugger/Source/MessagingDebugger/Private/MessagingDebuggerModule.cpp
@@ -2,6 +2,7 @@
 
 #include "MessagingDebuggerPrivatePCH.h"
 #include "Runtime/Core/Public/Features/IModularFeatures.h"
+#include "WorkspaceMenuStructureModule.h"
 
 
 #define LOCTEXT_NAMESPACE "FMessagingDebuggerModule"
@@ -28,16 +29,16 @@ public:
 		FMessagingDebuggerCommands::Register();
 
 		IModularFeatures::Get().RegisterModularFeature("MessagingDebugger", this);
-
-		FGlobalTabmanager::Get()->RegisterTabSpawner(MessagingDebuggerTabName, FOnSpawnTab::CreateRaw(this, &FMessagingDebuggerModule::SpawnMessagingDebuggerTab))
-			.SetDisplayName(NSLOCTEXT("FMessagingDebuggerModule", "DebuggerTabTitle", "Messaging Debugger"))
-			.SetTooltipText(NSLOCTEXT("FMessagingDebuggerModule", "DebuggerTooltipText", "Open the Messaging Debugger tab."))
-			.SetIcon(FSlateIcon(Style->GetStyleSetName(), "MessagingDebuggerTabIcon"));
+		
+		// This is still experimental in the editor, so it'll be invoked specifically in FMainMenu if the experimental settings flag is set.
+		// When no longer experimental, switch to the nomad spawner registration below
+		FGlobalTabmanager::Get()->RegisterTabSpawner(MessagingDebuggerTabName, FOnSpawnTab::CreateRaw(this, &FMessagingDebuggerModule::SpawnMessagingDebuggerTab));
 	}
 
 	virtual void ShutdownModule( ) override
 	{
 		FGlobalTabmanager::Get()->UnregisterTabSpawner(MessagingDebuggerTabName);
+
 		IModularFeatures::Get().UnregisterModularFeature("MessagingDebugger", this);
 		FMessagingDebuggerCommands::Unregister();
 	}
diff --git a/Engine/Plugins/Messaging/MessagingDebugger/Source/MessagingDebugger/Private/Widgets/SMessagingDebugger.cpp b/Engine/Plugins/Messaging/MessagingDebugger/Source/MessagingDebugger/Private/Widgets/SMessagingDebugger.cpp
--- a/Engine/Plugins/Messaging/MessagingDebugger/Source/MessagingDebugger/Private/Widgets/SMessagingDebugger.cpp
+++ b/Engine/Plugins/Messaging/MessagingDebugger/Source/MessagingDebugger/Private/Widgets/SMessagingDebugger.cpp
@@ -21,7 +21,7 @@ static const FName MessageTypesTabId("MessageTypes");
 static const FName ToolbarTabId("Toolbar");
 
 
-/* SMessagingDebugger structors
+/* SMessagingDebugger constructors
  *****************************************************************************/
 
 SMessagingDebugger::SMessagingDebugger( )
@@ -53,9 +53,7 @@ void SMessagingDebugger::Construct( const FArguments& InArgs, const TSharedRef<S
 
 	// create & initialize tab manager
 	TabManager = FGlobalTabmanager::Get()->NewTabManager(ConstructUnderMajorTab);
-
-	TSharedRef<FWorkspaceItem> RootMenuGroup = FWorkspaceItem::NewGroup(LOCTEXT("RootMenuGroupName", "Root"));
-	TSharedRef<FWorkspaceItem> AppMenuGroup = RootMenuGroup->AddGroup(LOCTEXT("AppMenuGroupName", "Messaging Debugger"));
+	TSharedRef<FWorkspaceItem> AppMenuGroup = TabManager->GetLocalWorkspaceMenuRoot()->AddGroup(LOCTEXT("MessagingDebuggerGroupName", "Messaging Debugger"));
 
 	TabManager->RegisterTabSpawner(BreakpointsTabId, FOnSpawnTab::CreateRaw(this, &SMessagingDebugger::HandleTabManagerSpawnTab, BreakpointsTabId))
 		.SetDisplayName(LOCTEXT("BreakpointsTabTitle", "Breakpoints"))
@@ -195,7 +193,7 @@ void SMessagingDebugger::Construct( const FArguments& InArgs, const TSharedRef<S
 	MenuBarBuilder.AddPullDownMenu(
 		LOCTEXT("WindowMenuLabel", "Window"),
 		FText::GetEmpty(),
-		FNewMenuDelegate::CreateStatic(&SMessagingDebugger::FillWindowMenu, RootMenuGroup, AppMenuGroup, TabManager),
+		FNewMenuDelegate::CreateStatic(&SMessagingDebugger::FillWindowMenu, TabManager),
 		"Window"
 	);
 
@@ -225,8 +223,6 @@ void SMessagingDebugger::Construct( const FArguments& InArgs, const TSharedRef<S
 
 void SMessagingDebugger::Tick( const FGeometry& AllottedGeometry, const double InCurrentTime, const float InDeltaTime )
 {
-	SCompoundWidget::Tick(AllottedGeometry, InCurrentTime, InDeltaTime);
-
 	MessageTracer->Tick(InDeltaTime);
 }
 
@@ -234,33 +230,25 @@ void SMessagingDebugger::Tick( const FGeometry& AllottedGeometry, const double I
 /* SMessagingDebugger implementation
  *****************************************************************************/
 
-void SMessagingDebugger::FillWindowMenu( FMenuBuilder& MenuBuilder, TSharedRef<FWorkspaceItem> RootMenuGroup, TSharedRef<FWorkspaceItem> AppMenuGroup, const TSharedPtr<FTabManager> TabManager )
+void SMessagingDebugger::FillWindowMenu(FMenuBuilder& MenuBuilder, const TSharedPtr<FTabManager> TabManager)
 {
 	if (!TabManager.IsValid())
 	{
 		return;
 	}
 
-	MenuBuilder.BeginSection("WindowLocalTabSpawners", LOCTEXT("MessagingDebuggerMenuGroup", "Messaging Debugger"));
-	{
-		TabManager->PopulateTabSpawnerMenu(MenuBuilder, AppMenuGroup);
-	}
-	MenuBuilder.EndSection();
-
 #if !WITH_EDITOR
-	MenuBuilder.BeginSection("WindowGlobalTabSpawners", LOCTEXT("UfeMenuGroup", "Unreal Frontend"));
-	{
-		FGlobalTabmanager::Get()->PopulateTabSpawnerMenu(MenuBuilder, RootMenuGroup);
-	}
-	MenuBuilder.EndSection();
+	FGlobalTabmanager::Get()->PopulateTabSpawnerMenu(MenuBuilder, WorkspaceMenu::GetMenuStructure().GetStructureRoot());
 #endif //!WITH_EDITOR
+
+	TabManager->PopulateLocalTabSpawnerMenu(MenuBuilder);
 }
 
 
 /* SMessagingDebugger callbacks
  *****************************************************************************/
 
-bool SMessagingDebugger::HandleBreakDebuggerCommandCanExecute( ) const
+bool SMessagingDebugger::HandleBreakDebuggerCommandCanExecute() const
 {
 	return MessageTracer->IsRunning() && !MessageTracer->IsBreaking();
 }
diff --git a/Engine/Plugins/Messaging/MessagingDebugger/Source/MessagingDebugger/Private/Widgets/SMessagingDebugger.h b/Engine/Plugins/Messaging/MessagingDebugger/Source/MessagingDebugger/Private/Widgets/SMessagingDebugger.h
--- a/Engine/Plugins/Messaging/MessagingDebugger/Source/MessagingDebugger/Private/Widgets/SMessagingDebugger.h
+++ b/Engine/Plugins/Messaging/MessagingDebugger/Source/MessagingDebugger/Private/Widgets/SMessagingDebugger.h
@@ -37,20 +37,17 @@ public:
 public:
 
 	// SCompoundWidget overrides
-
 	virtual void Tick( const FGeometry& AllottedGeometry, const double InCurrentTime, const float InDeltaTime ) override;
 
 protected:
 
 	/**
-	 * Fills the Window menu with menu items.
-	 *
-	 * @param MenuBuilder The multi-box builder that should be filled with content for this pull-down menu.
-	 * @param RootMenuGroup The root menu group.
-	 * @param AppMenuGroup The application menu group.
-	 * @param TabManager A Tab Manager from which to populate tab spawner menu items.
-	 */
-	static void FillWindowMenu( FMenuBuilder& MenuBuilder, TSharedRef<FWorkspaceItem> RootMenuGroup, TSharedRef<FWorkspaceItem> AppMenuGroup, const TSharedPtr<FTabManager> TabManager );
+	* Fills the Window menu with menu items.
+	*
+	* @param MenuBuilder The multi-box builder that should be filled with content for this pull-down menu.
+	* @param TabManager A Tab Manager from which to populate tab spawner menu items.
+	*/
+	static void FillWindowMenu(FMenuBuilder& MenuBuilder, const TSharedPtr<FTabManager> TabManager);
 
 private:
 
diff --git a/Engine/Source/Developer/AssetTools/Private/AssetTypeActions/AssetTypeActions_DestructibleMesh.h b/Engine/Source/Developer/AssetTools/Private/AssetTypeActions/AssetTypeActions_DestructibleMesh.h
--- a/Engine/Source/Developer/AssetTools/Private/AssetTypeActions/AssetTypeActions_DestructibleMesh.h
+++ b/Engine/Source/Developer/AssetTools/Private/AssetTypeActions/AssetTypeActions_DestructibleMesh.h
@@ -10,12 +10,7 @@ public:
 	virtual FColor GetTypeColor() const override { return FColor(200,128,128); }
 	virtual UClass* GetSupportedClass() const override;
 	virtual uint32 GetCategories() override { return EAssetTypeCategories::Physics; }
+	virtual bool HasActions( const TArray<UObject*>& InObjects ) const override { return false; }
 	virtual void OpenAssetEditor( const TArray<UObject*>& InObjects, TSharedPtr<class IToolkitHost> EditWithinLevelEditor = TSharedPtr<IToolkitHost>() ) override;
-	virtual void GetActions( const TArray<UObject*>& InObjects, FMenuBuilder& MenuBuilder ) override;
-
-	/** Handler for when Edit is selected */
-	void ExecuteEdit(TArray<TWeakObjectPtr<class UDestructibleMesh>> Objects);
-
-	/** Handler for when Reimport is selected */
-	void ExecuteReimport(TArray<TWeakObjectPtr<class UDestructibleMesh>> Objects);
+	virtual bool IsImportedAsset() const override { return false; }
 };
diff --git a/Engine/Source/Developer/CollisionAnalyzer/Private/CollisionAnalyzerModule.cpp b/Engine/Source/Developer/CollisionAnalyzer/Private/CollisionAnalyzerModule.cpp
--- a/Engine/Source/Developer/CollisionAnalyzer/Private/CollisionAnalyzerModule.cpp
+++ b/Engine/Source/Developer/CollisionAnalyzer/Private/CollisionAnalyzerModule.cpp
@@ -19,7 +19,7 @@ void FCollisionAnalyzerModule::StartupModule()
 	FGlobalTabmanager::Get()->RegisterNomadTabSpawner(CollisionAnalyzerModule::CollisionAnalyzerApp, FOnSpawnTab::CreateRaw(this, &FCollisionAnalyzerModule::SpawnCollisionAnalyzerTab))
 		.SetDisplayName(NSLOCTEXT("CollisionAnalyzerModule", "TabTitle", "Collision Analyzer"))
 		.SetTooltipText(NSLOCTEXT("CollisionAnalyzerModule", "TooltipText", "Open the Collision Analyzer tab."))
-		.SetGroup(WorkspaceMenu::GetMenuStructure().GetDeveloperToolsCategory())
+		.SetGroup(WorkspaceMenu::GetMenuStructure().GetDeveloperToolsDebugCategory())
 		.SetIcon(FSlateIcon(FEditorStyle::GetStyleSetName(), "CollisionAnalyzer.TabIcon"));
 }
 
diff --git a/Engine/Source/Developer/DeviceManager/DeviceManager.Build.cs b/Engine/Source/Developer/DeviceManager/DeviceManager.Build.cs
--- a/Engine/Source/Developer/DeviceManager/DeviceManager.Build.cs
+++ b/Engine/Source/Developer/DeviceManager/DeviceManager.Build.cs
@@ -21,6 +21,7 @@ public class DeviceManager : ModuleRules
 				"SlateCore",
 				"TargetPlatform",
 				"DesktopPlatform",
+                "WorkspaceMenuStructure",
 			}
 		);
 
diff --git a/Engine/Source/Developer/DeviceManager/Private/DeviceManagerModule.cpp b/Engine/Source/Developer/DeviceManager/Private/DeviceManagerModule.cpp
--- a/Engine/Source/Developer/DeviceManager/Private/DeviceManagerModule.cpp
+++ b/Engine/Source/Developer/DeviceManager/Private/DeviceManagerModule.cpp
@@ -1,7 +1,7 @@
 // Copyright 1998-2014 Epic Games, Inc. All Rights Reserved.
 
 #include "DeviceManagerPrivatePCH.h"
-
+#include "WorkspaceMenuStructureModule.h"
 
 static const FName DeviceManagerTabName("DeviceManager");
 
@@ -23,15 +23,21 @@ public:
 
 		TargetDeviceServiceManager = TargetDeviceServicesModule.GetDeviceServiceManager();
 
-		FGlobalTabmanager::Get()->RegisterTabSpawner(DeviceManagerTabName, FOnSpawnTab::CreateRaw(this, &FDeviceManagerModule::SpawnDeviceManagerTab))
+		auto& TabSpawnerEntry = FGlobalTabmanager::Get()->RegisterNomadTabSpawner(DeviceManagerTabName, FOnSpawnTab::CreateRaw(this, &FDeviceManagerModule::SpawnDeviceManagerTab))
 			.SetDisplayName(NSLOCTEXT("FDeviceManagerModule", "DeviceManagerTabTitle", "Device Manager"))
-			.SetTooltipText(NSLOCTEXT("FDeviceManagerModule", "DeviceManagerTooltipText", "Open the Device Manager tab."))
+			.SetTooltipText(NSLOCTEXT("FDeviceManagerModule", "DeviceManagerTooltipText", "View and manage connected devices."))
 			.SetIcon(FSlateIcon(FEditorStyle::GetStyleSetName(), "DeviceDetails.TabIcon"));
+
+#if WITH_EDITOR
+		TabSpawnerEntry.SetGroup(WorkspaceMenu::GetMenuStructure().GetDeveloperToolsMiscCategory());
+#else
+		TabSpawnerEntry.SetGroup(WorkspaceMenu::GetMenuStructure().GetToolsCategory());
+#endif
 	}
 
 	virtual void ShutdownModule( ) override
 	{
-		FGlobalTabmanager::Get()->UnregisterTabSpawner(DeviceManagerTabName);
+		FGlobalTabmanager::Get()->UnregisterNomadTabSpawner(DeviceManagerTabName);
 	}
 
 public:
diff --git a/Engine/Source/Developer/DeviceManager/Private/Widgets/SDeviceManager.cpp b/Engine/Source/Developer/DeviceManager/Private/Widgets/SDeviceManager.cpp
--- a/Engine/Source/Developer/DeviceManager/Private/Widgets/SDeviceManager.cpp
+++ b/Engine/Source/Developer/DeviceManager/Private/Widgets/SDeviceManager.cpp
@@ -16,7 +16,7 @@ static const FName DeviceProcessesTabId("DeviceProcesses");
 static const FName DeviceToolbarTabId("DeviceToolbar");
 
 
-/* SDeviceManager structors
+/* SDeviceManager constructors
  *****************************************************************************/
 
 SDeviceManager::SDeviceManager( )
@@ -38,9 +38,7 @@ void SDeviceManager::Construct( const FArguments& InArgs, const ITargetDeviceSer
 
 	// create & initialize tab manager
 	TabManager = FGlobalTabmanager::Get()->NewTabManager(ConstructUnderMajorTab);
-
-	TSharedRef<FWorkspaceItem> RootMenuGroup = FWorkspaceItem::NewGroup(LOCTEXT("RootMenuGroup", "Root"));
-	TSharedRef<FWorkspaceItem> AppMenuGroup = RootMenuGroup->AddGroup(LOCTEXT("DeviceManagerMenuGroupName", "Device Manager Tabs"));
+	TSharedRef<FWorkspaceItem> AppMenuGroup = TabManager->GetLocalWorkspaceMenuRoot()->AddGroup(LOCTEXT("DeviceManagerMenuGroupName", "Device Manager"));
 
 	TabManager->RegisterTabSpawner(DeviceBrowserTabId, FOnSpawnTab::CreateRaw(this, &SDeviceManager::HandleTabManagerSpawnTab, DeviceBrowserTabId))
 		.SetDisplayName(LOCTEXT("DeviceBrowserTabTitle", "Device Browser"))
@@ -113,7 +111,7 @@ void SDeviceManager::Construct( const FArguments& InArgs, const ITargetDeviceSer
 	MenuBarBuilder.AddPullDownMenu(
 		LOCTEXT("WindowMenuLabel", "Window"),
 		FText::GetEmpty(),
-		FNewMenuDelegate::CreateStatic(&SDeviceManager::FillWindowMenu, RootMenuGroup, AppMenuGroup, TabManager),
+		FNewMenuDelegate::CreateStatic(&SDeviceManager::FillWindowMenu, TabManager),
 		"Window"
 	);
 
@@ -199,7 +197,7 @@ void SDeviceManager::BindCommands( )
 }
 
 
-void SDeviceManager::FillWindowMenu( FMenuBuilder& MenuBuilder, TSharedRef<FWorkspaceItem> RootMenuGroup, TSharedRef<FWorkspaceItem> AppMenuGroup, const TSharedPtr<FTabManager> TabManager )
+void SDeviceManager::FillWindowMenu( FMenuBuilder& MenuBuilder, const TSharedPtr<FTabManager> TabManager )
 {
 	if (!TabManager.IsValid())
 	{
@@ -207,18 +205,10 @@ void SDeviceManager::FillWindowMenu( FMenuBuilder& MenuBuilder, TSharedRef<FWork
 	}
 
 #if !WITH_EDITOR
-	MenuBuilder.BeginSection("WindowGlobalTabSpawners", LOCTEXT("UfeMenuGroup", "Unreal Frontend"));
-	{
-		FGlobalTabmanager::Get()->PopulateTabSpawnerMenu(MenuBuilder, RootMenuGroup);
-	}
-	MenuBuilder.EndSection();
+		FGlobalTabmanager::Get()->PopulateTabSpawnerMenu(MenuBuilder, WorkspaceMenu::GetMenuStructure().GetStructureRoot());
 #endif //!WITH_EDITOR
 
-	MenuBuilder.BeginSection("WindowLocalTabSpawners", LOCTEXT("DeviceManagerMenuGroup", "Device Manager"));
-	{
-		TabManager->PopulateTabSpawnerMenu(MenuBuilder, AppMenuGroup);
-	}
-	MenuBuilder.EndSection();
+	TabManager->PopulateLocalTabSpawnerMenu(MenuBuilder);
 }
 
 
diff --git a/Engine/Source/Developer/DeviceManager/Private/Widgets/SDeviceManager.h b/Engine/Source/Developer/DeviceManager/Private/Widgets/SDeviceManager.h
--- a/Engine/Source/Developer/DeviceManager/Private/Widgets/SDeviceManager.h
+++ b/Engine/Source/Developer/DeviceManager/Private/Widgets/SDeviceManager.h
@@ -41,14 +41,12 @@ protected:
 	void BindCommands( );
 
 	/**
-	 * Fills the Window menu with menu items.
-	 *
-	 * @param MenuBuilder The multi-box builder that should be filled with content for this pull-down menu.
-	 * @param RootMenuGroup The root menu group.
-	 * @param AppMenuGroup The application menu group.
-	 * @param TabManager A Tab Manager from which to populate tab spawner menu items.
-	 */
-	static void FillWindowMenu( FMenuBuilder& MenuBuilder, TSharedRef<FWorkspaceItem> RootMenuGroup, TSharedRef<FWorkspaceItem> AppMenuGroup, const TSharedPtr<FTabManager> TabManager );
+	* Fills the Window menu with menu items.
+	*
+	* @param MenuBuilder The multi-box builder that should be filled with content for this pull-down menu.
+	* @param TabManager A Tab Manager from which to populate tab spawner menu items.
+	*/
+	static void FillWindowMenu(FMenuBuilder& MenuBuilder, const TSharedPtr<FTabManager> TabManager);
 
 	/**
 	 * Validates actions on the specified device.
diff --git a/Engine/Source/Developer/MessageLog/Private/MessageLogModule.cpp b/Engine/Source/Developer/MessageLog/Private/MessageLogModule.cpp
--- a/Engine/Source/Developer/MessageLog/Private/MessageLogModule.cpp
+++ b/Engine/Source/Developer/MessageLog/Private/MessageLogModule.cpp
@@ -43,11 +43,11 @@ void FMessageLogModule::StartupModule()
 	FGlobalTabmanager::Get()->RegisterNomadTabSpawner("MessageLog", FOnSpawnTab::CreateStatic( &SpawnMessageLog, MessageLogViewModel.ToSharedRef() ))
 		.SetDisplayName(NSLOCTEXT("UnrealEditor", "MessageLogTab", "Message Log"))
 		.SetTooltipText(NSLOCTEXT("UnrealEditor", "MessageLogTooltipText", "Open the Message Log tab."))
-		.SetGroup( WorkspaceMenu::GetMenuStructure().GetToolsCategory() )
+		.SetGroup( WorkspaceMenu::GetMenuStructure().GetDeveloperToolsLogCategory() )
 		.SetIcon(FSlateIcon(FEditorStyle::GetStyleSetName(), "MessageLog.TabIcon"));
 #endif
 
-	// Bind us so message log output is routed via this moddule
+	// Bind us so message log output is routed via this module
 	FMessageLog::OnGetLog().BindStatic(&GetLog);
 }
 
diff --git a/Engine/Source/Developer/OutputLog/Private/OutputLogModule.cpp b/Engine/Source/Developer/OutputLog/Private/OutputLogModule.cpp
--- a/Engine/Source/Developer/OutputLog/Private/OutputLogModule.cpp
+++ b/Engine/Source/Developer/OutputLog/Private/OutputLogModule.cpp
@@ -71,7 +71,7 @@ void FOutputLogModule::StartupModule()
 	FGlobalTabmanager::Get()->RegisterNomadTabSpawner(OutputLogModule::OutputLogTabName, FOnSpawnTab::CreateStatic( &SpawnOutputLog ) )
 		.SetDisplayName(NSLOCTEXT("UnrealEditor", "OutputLogTab", "Output Log"))
 		.SetTooltipText(NSLOCTEXT("UnrealEditor", "OutputLogTooltipText", "Open the Output Log tab."))
-		.SetGroup( WorkspaceMenu::GetMenuStructure().GetToolsCategory() )
+		.SetGroup( WorkspaceMenu::GetMenuStructure().GetDeveloperToolsLogCategory() )
 		.SetIcon( FSlateIcon(FEditorStyle::GetStyleSetName(), "Log.TabIcon") );
 	
 	OutputLogHistory = MakeShareable(new FOutputLogHistory);
diff --git a/Engine/Source/Developer/ProjectLauncher/Private/ProjectLauncherModule.cpp b/Engine/Source/Developer/ProjectLauncher/Private/ProjectLauncherModule.cpp
--- a/Engine/Source/Developer/ProjectLauncher/Private/ProjectLauncherModule.cpp
+++ b/Engine/Source/Developer/ProjectLauncher/Private/ProjectLauncherModule.cpp
@@ -1,11 +1,10 @@
 // Copyright 1998-2014 Epic Games, Inc. All Rights Reserved.
 
 #include "ProjectLauncherPrivatePCH.h"
-
+#include "WorkspaceMenuStructureModule.h"
 
 static const FName ProjectLauncherTabName("ProjectLauncher");
 
-
 /**
  * Implements the SessionSProjectLauncher module.
  */
@@ -31,15 +30,26 @@ public:
 	
 	virtual void StartupModule( ) override
 	{
-		FGlobalTabmanager::Get()->RegisterTabSpawner(ProjectLauncherTabName, FOnSpawnTab::CreateRaw(this, &FProjectLauncherModule::SpawnProjectLauncherTab))
+#if WITH_EDITOR
+		FGlobalTabmanager::Get()->RegisterTabSpawner(ProjectLauncherTabName, FOnSpawnTab::CreateRaw(this, &FProjectLauncherModule::SpawnProjectLauncherTab));
+#else
+		// This is still experimental in the editor, so it'll be invoked specifically in FMainMenu if the experimental settings flag is set.
+		//@todo Enable this in the editor when no longer experimental
+		FGlobalTabmanager::Get()->RegisterNomadTabSpawner(ProjectLauncherTabName, FOnSpawnTab::CreateRaw(this, &FProjectLauncherModule::SpawnProjectLauncherTab))
 			.SetDisplayName(NSLOCTEXT("FProjectLauncherModule", "ProjectLauncherTabTitle", "Project Launcher"))
 			.SetTooltipText(NSLOCTEXT("FProjectLauncherModule", "ProjectLauncherTooltipText", "Open the Project Launcher tab."))
-			.SetIcon(FSlateIcon(FEditorStyle::GetStyleSetName(), "Launcher.TabIcon"));
+			.SetIcon(FSlateIcon(FEditorStyle::GetStyleSetName(), "Launcher.TabIcon"))
+			.SetGroup(WorkspaceMenu::GetMenuStructure().GetToolsCategory());
+#endif
 	}
 
 	virtual void ShutdownModule( ) override
 	{
+#if WITH_EDITOR
 		FGlobalTabmanager::Get()->UnregisterTabSpawner(ProjectLauncherTabName);
+#else
+		FGlobalTabmanager::Get()->UnregisterNomadTabSpawner(ProjectLauncherTabName);
+#endif
 	}
 
 private:
diff --git a/Engine/Source/Developer/ProjectLauncher/ProjectLauncher.Build.cs b/Engine/Source/Developer/ProjectLauncher/ProjectLauncher.Build.cs
--- a/Engine/Source/Developer/ProjectLauncher/ProjectLauncher.Build.cs
+++ b/Engine/Source/Developer/ProjectLauncher/ProjectLauncher.Build.cs
@@ -27,6 +27,7 @@ public class ProjectLauncher : ModuleRules
 				"Slate",
 				"SlateCore",
                 "EditorStyle",
+                "WorkspaceMenuStructure",
 			}
 		);
 
diff --git a/Engine/Source/Developer/SessionFrontend/Private/SessionFrontendModule.cpp b/Engine/Source/Developer/SessionFrontend/Private/SessionFrontendModule.cpp
--- a/Engine/Source/Developer/SessionFrontend/Private/SessionFrontendModule.cpp
+++ b/Engine/Source/Developer/SessionFrontend/Private/SessionFrontendModule.cpp
@@ -1,7 +1,7 @@
 // Copyright 1998-2014 Epic Games, Inc. All Rights Reserved.
 
 #include "SessionFrontendPrivatePCH.h"
-
+#include "WorkspaceMenuStructureModule.h"
 
 static const FName SessionFrontendTabName("SessionFrontend");
 
@@ -32,15 +32,21 @@ public:
 	
 	virtual void StartupModule( ) override
 	{
-		FGlobalTabmanager::Get()->RegisterTabSpawner(SessionFrontendTabName, FOnSpawnTab::CreateRaw(this, &FSessionFrontendModule::SpawnSessionFrontendTab))
+		auto& TabSpawnerEntry = FGlobalTabmanager::Get()->RegisterNomadTabSpawner(SessionFrontendTabName, FOnSpawnTab::CreateRaw(this, &FSessionFrontendModule::SpawnSessionFrontendTab))
 			.SetDisplayName(NSLOCTEXT("FSessionFrontendModule", "FrontendTabTitle", "Session Frontend"))
 			.SetTooltipText(NSLOCTEXT("FSessionFrontendModule", "FrontendTooltipText", "Open the Session Frontend tab."))
 			.SetIcon(FSlateIcon(FEditorStyle::GetStyleSetName(), "SessionFrontEnd.TabIcon"));
+
+#if WITH_EDITOR
+		TabSpawnerEntry.SetGroup(WorkspaceMenu::GetMenuStructure().GetDeveloperToolsMiscCategory());
+#else
+		TabSpawnerEntry.SetGroup(WorkspaceMenu::GetMenuStructure().GetToolsCategory());
+#endif
 	}
 
 	virtual void ShutdownModule( ) override
 	{
-		FGlobalTabmanager::Get()->UnregisterTabSpawner(SessionFrontendTabName);
+		FGlobalTabmanager::Get()->UnregisterNomadTabSpawner(SessionFrontendTabName);
 	}
 
 private:
diff --git a/Engine/Source/Developer/SessionFrontend/Private/Widgets/SSessionFrontend.cpp b/Engine/Source/Developer/SessionFrontend/Private/Widgets/SSessionFrontend.cpp
--- a/Engine/Source/Developer/SessionFrontend/Private/Widgets/SSessionFrontend.cpp
+++ b/Engine/Source/Developer/SessionFrontend/Private/Widgets/SSessionFrontend.cpp
@@ -25,9 +25,7 @@ void SSessionFrontend::Construct( const FArguments& InArgs, const TSharedRef<SDo
 
 	// create & initialize tab manager
 	TabManager = FGlobalTabmanager::Get()->NewTabManager(ConstructUnderMajorTab);
-
-	TSharedRef<FWorkspaceItem> RootMenuGroup = FWorkspaceItem::NewGroup(LOCTEXT("RootMenuGroupName", "Root"));
-	TSharedRef<FWorkspaceItem> AppMenuGroup = RootMenuGroup->AddGroup(LOCTEXT("SessionFrontendMenuGroupName", "Session Frontend"));
+	TSharedRef<FWorkspaceItem> AppMenuGroup = TabManager->GetLocalWorkspaceMenuRoot()->AddGroup(LOCTEXT("SessionFrontendMenuGroupName", "Session Frontend"));
 	
 	TabManager->RegisterTabSpawner(AutomationTabId, FOnSpawnTab::CreateRaw(this, &SSessionFrontend::HandleTabManagerSpawnTab, AutomationTabId))
 		.SetDisplayName(LOCTEXT("AutomationTabTitle", "Automation"))
@@ -87,7 +85,7 @@ void SSessionFrontend::Construct( const FArguments& InArgs, const TSharedRef<SDo
 	MenuBarBuilder.AddPullDownMenu(
 		LOCTEXT("WindowMenuLabel", "Window"),
 		FText::GetEmpty(),
-		FNewMenuDelegate::CreateStatic(&SSessionFrontend::FillWindowMenu, RootMenuGroup, AppMenuGroup, TabManager),
+		FNewMenuDelegate::CreateStatic(&SSessionFrontend::FillWindowMenu, TabManager),
 		"Window"
 	);
 
@@ -113,7 +111,7 @@ void SSessionFrontend::Construct( const FArguments& InArgs, const TSharedRef<SDo
 /* SSessionFrontend implementation
  *****************************************************************************/
 
-void SSessionFrontend::FillWindowMenu( FMenuBuilder& MenuBuilder, TSharedRef<FWorkspaceItem> RootMenuGroup, TSharedRef<FWorkspaceItem> AppMenuGroup, const TSharedPtr<FTabManager> TabManager )
+void SSessionFrontend::FillWindowMenu( FMenuBuilder& MenuBuilder, const TSharedPtr<FTabManager> TabManager )
 {
 	if (!TabManager.IsValid())
 	{
@@ -121,18 +119,10 @@ void SSessionFrontend::FillWindowMenu( FMenuBuilder& MenuBuilder, TSharedRef<FWo
 	}
 
 #if !WITH_EDITOR
-	MenuBuilder.BeginSection("WindowGlobalTabSpawners", LOCTEXT("UfeMenuGroup", "Unreal Frontend"));
-	{
-		FGlobalTabmanager::Get()->PopulateTabSpawnerMenu(MenuBuilder, RootMenuGroup);
-	}
-	MenuBuilder.EndSection();
+	FGlobalTabmanager::Get()->PopulateTabSpawnerMenu(MenuBuilder, WorkspaceMenu::GetMenuStructure().GetStructureRoot());
 #endif //!WITH_EDITOR
 
-	MenuBuilder.BeginSection("WindowLocalTabSpawners", LOCTEXT("SessionFrontendMenuGroup", "Session Frontend"));
-	{
-		TabManager->PopulateTabSpawnerMenu(MenuBuilder, AppMenuGroup);
-	}
-	MenuBuilder.EndSection();
+	TabManager->PopulateLocalTabSpawnerMenu(MenuBuilder);
 }
 
 
diff --git a/Engine/Source/Developer/SessionFrontend/Private/Widgets/SSessionFrontend.h b/Engine/Source/Developer/SessionFrontend/Private/Widgets/SSessionFrontend.h
--- a/Engine/Source/Developer/SessionFrontend/Private/Widgets/SSessionFrontend.h
+++ b/Engine/Source/Developer/SessionFrontend/Private/Widgets/SSessionFrontend.h
@@ -31,11 +31,9 @@ protected:
 	 * Fills the Window menu with menu items.
 	 *
 	 * @param MenuBuilder The multi-box builder that should be filled with content for this pull-down menu.
-	 * @param RootMenuGroup The root menu group.
-	 * @param AppMenuGroup The application menu group.
 	 * @param TabManager A Tab Manager from which to populate tab spawner menu items.
 	 */
-	static void FillWindowMenu( FMenuBuilder& MenuBuilder, TSharedRef<FWorkspaceItem> RootMenuGroup, TSharedRef<FWorkspaceItem> AppMenuGroup, const TSharedPtr<FTabManager> TabManager );
+	static void FillWindowMenu( FMenuBuilder& MenuBuilder, const TSharedPtr<FTabManager> TabManager );
 
 	/**
 	 * Creates and initializes the controller classes.
diff --git a/Engine/Source/Developer/SessionFrontend/SessionFrontend.Build.cs b/Engine/Source/Developer/SessionFrontend/SessionFrontend.Build.cs
--- a/Engine/Source/Developer/SessionFrontend/SessionFrontend.Build.cs
+++ b/Engine/Source/Developer/SessionFrontend/SessionFrontend.Build.cs
@@ -33,6 +33,7 @@ public class SessionFrontend : ModuleRules
 				"ScreenShotComparisonTools",
 				"Profiler",
 				"TargetPlatform",
+                "WorkspaceMenuStructure",
 			}
 		);
 
diff --git a/Engine/Source/Developer/Toolbox/Private/Toolbox.cpp b/Engine/Source/Developer/Toolbox/Private/Toolbox.cpp
--- a/Engine/Source/Developer/Toolbox/Private/Toolbox.cpp
+++ b/Engine/Source/Developer/Toolbox/Private/Toolbox.cpp
@@ -154,14 +154,14 @@ void FToolboxModule::StartupModule()
 	FGlobalTabmanager::Get()->RegisterNomadTabSpawner( "DebugTools", FOnSpawnTab::CreateStatic( &CreateDebugToolsTab ) )
 		.SetDisplayName( NSLOCTEXT("Toolbox", "DebugTools", "Debug Tools") )
 		.SetTooltipText( NSLOCTEXT("Toolbox", "DebugToolsTooltipText", "Open the Debug Tools tab.") )
-		.SetGroup( MenuStructure.GetDeveloperToolsCategory() )
+		.SetGroup( MenuStructure.GetDeveloperToolsDebugCategory() )
 		.SetIcon( FSlateIcon(FEditorStyle::GetStyleSetName(), "DebugTools.TabIcon") );
 	if ( CanShowModulesTab() )
 	{
 		FGlobalTabmanager::Get()->RegisterNomadTabSpawner( "ModulesTab", FOnSpawnTab::CreateStatic( &CreateModulesTab ) )
 			.SetDisplayName( NSLOCTEXT("Toolbox", "Modules", "Modules") )
 			.SetTooltipText( NSLOCTEXT("Toolbox", "ModulesTooltipText", "Open the Modules tab.") )
-			.SetGroup( MenuStructure.GetDeveloperToolsCategory() )
+			.SetGroup( MenuStructure.GetDeveloperToolsMiscCategory() )
 			.SetIcon( FSlateIcon(FEditorStyle::GetStyleSetName(), "Modules.TabIcon") );
 	}
 }
diff --git a/Engine/Source/Editor/BehaviorTreeEditor/Private/BehaviorTreeEditorModes.cpp b/Engine/Source/Editor/BehaviorTreeEditor/Private/BehaviorTreeEditorModes.cpp
--- a/Engine/Source/Editor/BehaviorTreeEditor/Private/BehaviorTreeEditorModes.cpp
+++ b/Engine/Source/Editor/BehaviorTreeEditor/Private/BehaviorTreeEditorModes.cpp
@@ -6,8 +6,13 @@
 #include "BehaviorTreeEditorTabFactories.h"
 #include "BehaviorTreeEditorToolbar.h"
 
+/////////////////////////////////////////////////////
+// FBehaviorTreeEditorApplicationMode
+
+#define LOCTEXT_NAMESPACE "BehaviorTreeApplicationMode"
+
 FBehaviorTreeEditorApplicationMode::FBehaviorTreeEditorApplicationMode(TSharedPtr<class FBehaviorTreeEditor> InBehaviorTreeEditor)
-	: FApplicationMode(FBehaviorTreeEditor::BehaviorTreeMode)
+	: FApplicationMode(FBehaviorTreeEditor::BehaviorTreeMode, FBehaviorTreeEditor::GetLocalizedMode)
 {
 	BehaviorTreeEditor = InBehaviorTreeEditor;
 
@@ -94,8 +99,15 @@ void FBehaviorTreeEditorApplicationMode::PostActivateMode()
 	FApplicationMode::PostActivateMode();
 }
 
+#undef  LOCTEXT_NAMESPACE
+
+/////////////////////////////////////////////////////
+// FBlackboardEditorApplicationMode
+
+#define LOCTEXT_NAMESPACE "BlackboardApplicationMode"
+
 FBlackboardEditorApplicationMode::FBlackboardEditorApplicationMode(TSharedPtr<class FBehaviorTreeEditor> InBehaviorTreeEditor)
-	: FApplicationMode(FBehaviorTreeEditor::BlackboardMode)
+	: FApplicationMode(FBehaviorTreeEditor::BlackboardMode, FBehaviorTreeEditor::GetLocalizedMode)
 {
 	BehaviorTreeEditor = InBehaviorTreeEditor;
 	
@@ -154,3 +166,5 @@ void FBlackboardEditorApplicationMode::PostActivateMode()
 
 	FApplicationMode::PostActivateMode();
 }
+
+#undef LOCTEXT_NAMESPACE
\ No newline at end of file
diff --git a/Engine/Source/Editor/BehaviorTreeEditor/Private/BehaviorTreeEditorTabFactories.cpp b/Engine/Source/Editor/BehaviorTreeEditor/Private/BehaviorTreeEditorTabFactories.cpp
--- a/Engine/Source/Editor/BehaviorTreeEditor/Private/BehaviorTreeEditorTabFactories.cpp
+++ b/Engine/Source/Editor/BehaviorTreeEditor/Private/BehaviorTreeEditorTabFactories.cpp
@@ -15,7 +15,7 @@ FBlackboardSummoner::FBlackboardSummoner(TSharedPtr<class FBehaviorTreeEditor> I
 	, BehaviorTreeEditorPtr(InBehaviorTreeEditorPtr)
 {
 	TabLabel = LOCTEXT("BlackboardLabel", "Blackboard");
-	TabIcon = FEditorStyle::GetBrush("Kismet.Tabs.Components");
+	TabIcon = FSlateIcon(FEditorStyle::GetStyleSetName(), "Kismet.Tabs.Components");
 
 	bIsSingleton = true;
 
@@ -38,7 +38,7 @@ FBlackboardEditorSummoner::FBlackboardEditorSummoner(TSharedPtr<class FBehaviorT
 	, BehaviorTreeEditorPtr(InBehaviorTreeEditorPtr)
 {
 	TabLabel = LOCTEXT("BlackboardLabel", "Blackboard");
-	TabIcon = FEditorStyle::GetBrush("Kismet.Tabs.Components");
+	TabIcon = FSlateIcon(FEditorStyle::GetStyleSetName(), "Kismet.Tabs.Components");
 
 	bIsSingleton = true;
 
@@ -61,7 +61,7 @@ FBlackboardDetailsSummoner::FBlackboardDetailsSummoner(TSharedPtr<class FBehavio
 	, BehaviorTreeEditorPtr(InBehaviorTreeEditorPtr)
 {
 	TabLabel = LOCTEXT("BlackboardDetailsLabel", "Blackboard Details");
-	TabIcon = FEditorStyle::GetBrush("Kismet.Tabs.Components");
+	TabIcon = FSlateIcon(FEditorStyle::GetStyleSetName(), "Kismet.Tabs.Components");
 
 	bIsSingleton = true;
 
@@ -85,7 +85,7 @@ FBehaviorTreeDetailsSummoner::FBehaviorTreeDetailsSummoner(TSharedPtr<class FBeh
 	, BehaviorTreeEditorPtr(InBehaviorTreeEditorPtr)
 {
 	TabLabel = LOCTEXT("BehaviorTreeDetailsLabel", "Details");
-	TabIcon = FEditorStyle::GetBrush("Kismet.Tabs.Components");
+	TabIcon = FSlateIcon(FEditorStyle::GetStyleSetName(), "Kismet.Tabs.Components");
 
 	bIsSingleton = true;
 
@@ -109,7 +109,7 @@ FBehaviorTreeSearchSummoner::FBehaviorTreeSearchSummoner(TSharedPtr<class FBehav
 	, BehaviorTreeEditorPtr(InBehaviorTreeEditorPtr)
 {
 	TabLabel = LOCTEXT("BehaviorTreeSearchLabel", "Search");
-	TabIcon = FEditorStyle::GetBrush("Kismet.Tabs.Components");
+	TabIcon = FSlateIcon(FEditorStyle::GetStyleSetName(), "Kismet.Tabs.Components");
 
 	bIsSingleton = true;
 
diff --git a/Engine/Source/Editor/Blutility/Private/BlutilityModule.cpp b/Engine/Source/Editor/Blutility/Private/BlutilityModule.cpp
--- a/Engine/Source/Editor/Blutility/Private/BlutilityModule.cpp
+++ b/Engine/Source/Editor/Blutility/Private/BlutilityModule.cpp
@@ -41,10 +41,9 @@ public:
 		PropertyModule.RegisterCustomClassLayout("GlobalEditorUtilityBase", FOnGetDetailCustomizationInstance::CreateStatic(&FEditorUtilityInstanceDetails::MakeInstance));
 		PropertyModule.NotifyCustomizationModuleChanged();
 
-		if (GetDefault<UEditorExperimentalSettings>()->bEnableEditorUtilityBlueprints)
-		{
-			RegisterBlutilityShelfTabSpawner();
-		}
+		FGlobalTabmanager::Get()->RegisterTabSpawner(BlutilityModule::BlutilityShelfApp, FOnSpawnTab::CreateStatic(&SpawnBlutilityShelfTab))
+			.SetDisplayName(NSLOCTEXT("BlutilityShelf", "TabTitle", "Blutility Shelf"))
+			.SetGroup(WorkspaceMenu::GetMenuStructure().GetToolsCategory());
 	}
 
 	virtual void ShutdownModule() override
@@ -54,7 +53,7 @@ public:
 			return;
 		}
 
-		UnregisterBlutilityShelfTabSpawner();
+		FGlobalTabmanager::Get()->UnregisterTabSpawner(BlutilityModule::BlutilityShelfApp);
 
 		// Only unregister if the asset tools module is loaded.  We don't want to forcibly load it during shutdown phase.
 		check( EditorBlueprintAssetTypeActions.IsValid() );
@@ -83,35 +82,6 @@ protected:
 				SNew(SBlutilityShelf)
 			];
 	}
-
-	void RegisterBlutilityShelfTabSpawner()
-	{
-		FGlobalTabmanager::Get()->RegisterNomadTabSpawner( BlutilityModule::BlutilityShelfApp, FOnSpawnTab::CreateStatic( &SpawnBlutilityShelfTab ) )
-			.SetDisplayName(NSLOCTEXT("BlutilityShelf", "TabTitle", "Blutility Shelf"))
-			.SetGroup( WorkspaceMenu::GetMenuStructure().GetToolsCategory() );
-	}
-
-	void UnregisterBlutilityShelfTabSpawner()
-	{
-		// Unregister our shelf tab spawner
-		FGlobalTabmanager::Get()->UnregisterNomadTabSpawner(BlutilityModule::BlutilityShelfApp);
-	}
-
-	void HandleExperimentalSettingChanged(FName PropertyName)
-	{
-		if (PropertyName == TEXT("bEnableEditorUtilityBlueprints"))
-		{
-			UEditorUtilityBlueprintFactory* EditorUtilityBlueprintFactory = Cast<UEditorUtilityBlueprintFactory>(UEditorUtilityBlueprintFactory::StaticClass()->GetDefaultObject());
-			if (GetDefault<UEditorExperimentalSettings>()->bEnableEditorUtilityBlueprints)
-			{
-				RegisterBlutilityShelfTabSpawner();
-			}
-			else
-			{
-				UnregisterBlutilityShelfTabSpawner();
-			}
-		}
-	}
 };
 
 
diff --git a/Engine/Source/Editor/Cascade/Private/Cascade.cpp b/Engine/Source/Editor/Cascade/Private/Cascade.cpp
--- a/Engine/Source/Editor/Cascade/Private/Cascade.cpp
+++ b/Engine/Source/Editor/Cascade/Private/Cascade.cpp
@@ -44,25 +44,30 @@ DEFINE_LOG_CATEGORY(LogCascade);
 
 void FCascade::RegisterTabSpawners(const TSharedRef<class FTabManager>& TabManager)
 {
-	FAssetEditorToolkit::RegisterTabSpawners(TabManager);
+	WorkspaceMenuCategory = TabManager->GetLocalWorkspaceMenuRoot()->AddGroup(NSLOCTEXT("Cascade", "WorkspaceMenu_Cascade", "Cascade"));
+	auto WorkspaceMenuCategoryRef = WorkspaceMenuCategory.ToSharedRef();
 
-	const IWorkspaceMenuStructure& MenuStructure = WorkspaceMenu::GetMenuStructure();
+	FAssetEditorToolkit::RegisterTabSpawners(TabManager);
 
 	TabManager->RegisterTabSpawner( Cascade_PreviewViewportTab, FOnSpawnTab::CreateSP( this, &FCascade::SpawnTab, Cascade_PreviewViewportTab ) )
 		.SetDisplayName(NSLOCTEXT("Cascade", "SummonViewport", "Viewport"))
-		.SetGroup( MenuStructure.GetAssetEditorCategory() );
+		.SetGroup( WorkspaceMenuCategoryRef )
+		.SetIcon(FSlateIcon(FEditorStyle::GetStyleSetName(), "LevelEditor.Tabs.Viewports"));
 
 	TabManager->RegisterTabSpawner( Cascade_EmmitterCanvasTab, FOnSpawnTab::CreateSP( this, &FCascade::SpawnTab, Cascade_EmmitterCanvasTab ) )
 		.SetDisplayName(NSLOCTEXT("Cascade", "SummonCanvas", "Emitters"))
-		.SetGroup( MenuStructure.GetAssetEditorCategory() );
+		.SetGroup( WorkspaceMenuCategoryRef )
+		.SetIcon(FSlateIcon(FEditorStyle::GetStyleSetName(), "ClassIcon.Emitter"));
 
 	TabManager->RegisterTabSpawner( Cascade_PropertiesTab, FOnSpawnTab::CreateSP( this, &FCascade::SpawnTab, Cascade_PropertiesTab ) )
 		.SetDisplayName(NSLOCTEXT("Cascade", "SummonProperties", "Details"))
-		.SetGroup( MenuStructure.GetAssetEditorCategory() );
+		.SetGroup( WorkspaceMenuCategoryRef )
+		.SetIcon(FSlateIcon(FEditorStyle::GetStyleSetName(), "LevelEditor.Tabs.Details"));
 
 	TabManager->RegisterTabSpawner( Cascade_CurveEditorTab, FOnSpawnTab::CreateSP( this, &FCascade::SpawnTab, Cascade_CurveEditorTab ) )
 		.SetDisplayName(NSLOCTEXT("Cascade", "SummonCurveEditor", "CurveEditor"))
-		.SetGroup( MenuStructure.GetAssetEditorCategory() );
+		.SetGroup( WorkspaceMenuCategoryRef )
+		.SetIcon(FSlateIcon(FEditorStyle::GetStyleSetName(), "ClassIcon.CurveBase"));
 }
 
 void FCascade::UnregisterTabSpawners(const TSharedRef<class FTabManager>& TabManager)
@@ -5172,3 +5177,4 @@ void UCascadeConfiguration::CacheModuleRejections()
 		}
 	}
 }
+
diff --git a/Engine/Source/Editor/ClassViewer/Private/ClassViewerModule.cpp b/Engine/Source/Editor/ClassViewer/Private/ClassViewerModule.cpp
--- a/Engine/Source/Editor/ClassViewer/Private/ClassViewerModule.cpp
+++ b/Engine/Source/Editor/ClassViewer/Private/ClassViewerModule.cpp
@@ -32,8 +32,8 @@ void FClassViewerModule::StartupModule()
 {
 	FGlobalTabmanager::Get()->RegisterNomadTabSpawner( ClassViewerModule::ClassViewerApp, FOnSpawnTab::CreateStatic( &CreateClassPickerTab ) )
 		.SetDisplayName( NSLOCTEXT("ClassViewerApp", "TabTitle", "Class Viewer") )
-		.SetTooltipText( NSLOCTEXT("ClassViewerApp", "TooltipText", "Open the Class Viewer tab.") )
-		.SetGroup( WorkspaceMenu::GetMenuStructure().GetToolsCategory() )
+		.SetTooltipText( NSLOCTEXT("ClassViewerApp", "TooltipText", "Displays all classes that exist within this project.") )
+		.SetGroup( WorkspaceMenu::GetMenuStructure().GetDeveloperToolsMiscCategory() )
 		.SetIcon( FSlateIcon(FEditorStyle::GetStyleSetName(), "ClassViewer.TabIcon") );
 }
 
diff --git a/Engine/Source/Editor/CurveAssetEditor/Private/CurveAssetEditor.cpp b/Engine/Source/Editor/CurveAssetEditor/Private/CurveAssetEditor.cpp
--- a/Engine/Source/Editor/CurveAssetEditor/Private/CurveAssetEditor.cpp
+++ b/Engine/Source/Editor/CurveAssetEditor/Private/CurveAssetEditor.cpp
@@ -14,13 +14,14 @@ const FName FCurveAssetEditor::CurveTabId( TEXT( "CurveAssetEditor_Curve" ) );
 
 void FCurveAssetEditor::RegisterTabSpawners(const TSharedRef<class FTabManager>& TabManager)
 {
-	FAssetEditorToolkit::RegisterTabSpawners(TabManager);
+	WorkspaceMenuCategory = TabManager->GetLocalWorkspaceMenuRoot()->AddGroup(LOCTEXT("WorkspaceMenu_CurveAssetEditor", "Curve Asset Editor"));
 
-	const IWorkspaceMenuStructure& MenuStructure = WorkspaceMenu::GetMenuStructure();
+	FAssetEditorToolkit::RegisterTabSpawners(TabManager);
 
 	TabManager->RegisterTabSpawner( CurveTabId, FOnSpawnTab::CreateSP(this, &FCurveAssetEditor::SpawnTab_CurveAsset) )
 		.SetDisplayName( LOCTEXT("CurveTab", "Curve") )
-		.SetGroup( MenuStructure.GetAssetEditorCategory() );
+		.SetGroup(WorkspaceMenuCategory.ToSharedRef())
+		.SetIcon(FSlateIcon(FEditorStyle::GetStyleSetName(), "ClassIcon.CurveBase"));
 }
 
 void FCurveAssetEditor::UnregisterTabSpawners(const TSharedRef<class FTabManager>& TabManager)
diff --git a/Engine/Source/Editor/CurveTableEditor/Private/CurveTableEditor.cpp b/Engine/Source/Editor/CurveTableEditor/Private/CurveTableEditor.cpp
--- a/Engine/Source/Editor/CurveTableEditor/Private/CurveTableEditor.cpp
+++ b/Engine/Source/Editor/CurveTableEditor/Private/CurveTableEditor.cpp
@@ -12,11 +12,11 @@ const FName FCurveTableEditor::CurveTableTabId( TEXT( "CurveTableEditor_CurveTab
 
 void FCurveTableEditor::RegisterTabSpawners(const TSharedRef<class FTabManager>& TabManager)
 {
-	const IWorkspaceMenuStructure& MenuStructure = WorkspaceMenu::GetMenuStructure();
+	WorkspaceMenuCategory = TabManager->GetLocalWorkspaceMenuRoot()->AddGroup(LOCTEXT("WorkspaceMenu_UserDefinedEnumEditor", "User-Defined Enum Editor"));
 
 	TabManager->RegisterTabSpawner( CurveTableTabId, FOnSpawnTab::CreateSP(this, &FCurveTableEditor::SpawnTab_CurveTable) )
 		.SetDisplayName( LOCTEXT("CurveTableTab", "Curve Table") )
-		.SetGroup( MenuStructure.GetAssetEditorCategory() );
+		.SetGroup( WorkspaceMenuCategory.ToSharedRef() );
 }
 
 void FCurveTableEditor::UnregisterTabSpawners(const TSharedRef<class FTabManager>& TabManager)
diff --git a/Engine/Source/Editor/DataTableEditor/Private/DataTableEditor.cpp b/Engine/Source/Editor/DataTableEditor/Private/DataTableEditor.cpp
--- a/Engine/Source/Editor/DataTableEditor/Private/DataTableEditor.cpp
+++ b/Engine/Source/Editor/DataTableEditor/Private/DataTableEditor.cpp
@@ -12,11 +12,11 @@ const FName FDataTableEditor::DataTableTabId( TEXT( "DataTableEditor_DataTable"
 
 void FDataTableEditor::RegisterTabSpawners(const TSharedRef<class FTabManager>& TabManager)
 {
-	const IWorkspaceMenuStructure& MenuStructure = WorkspaceMenu::GetMenuStructure();
+	WorkspaceMenuCategory = TabManager->GetLocalWorkspaceMenuRoot()->AddGroup(LOCTEXT("WorkspaceMenu_Data Table Editor", "Data Table Editor"));
 
 	TabManager->RegisterTabSpawner( DataTableTabId, FOnSpawnTab::CreateSP(this, &FDataTableEditor::SpawnTab_DataTable) )
 		.SetDisplayName( LOCTEXT("DataTableTab", "Data Table") )
-		.SetGroup( MenuStructure.GetAssetEditorCategory() );
+		.SetGroup( WorkspaceMenuCategory.ToSharedRef() );
 }
 
 void FDataTableEditor::UnregisterTabSpawners(const TSharedRef<class FTabManager>& TabManager)
diff --git a/Engine/Source/Editor/DestructibleMeshEditor/Private/DestructibleMeshEditor.cpp b/Engine/Source/Editor/DestructibleMeshEditor/Private/DestructibleMeshEditor.cpp
--- a/Engine/Source/Editor/DestructibleMeshEditor/Private/DestructibleMeshEditor.cpp
+++ b/Engine/Source/Editor/DestructibleMeshEditor/Private/DestructibleMeshEditor.cpp
@@ -42,25 +42,27 @@ FDestructibleMeshEditor::~FDestructibleMeshEditor()
 
 void FDestructibleMeshEditor::RegisterTabSpawners(const TSharedRef<class FTabManager>& TabManager)
 {
-	FAssetEditorToolkit::RegisterTabSpawners(TabManager);
+	WorkspaceMenuCategory = TabManager->GetLocalWorkspaceMenuRoot()->AddGroup(LOCTEXT("WorkspaceMenu_DestructibleMeshEditor", "Destructible Mesh Editor"));
+	auto WorkspaceMenuCategoryRef = WorkspaceMenuCategory.ToSharedRef();
 
-	const IWorkspaceMenuStructure& MenuStructure = WorkspaceMenu::GetMenuStructure();
+	FAssetEditorToolkit::RegisterTabSpawners(TabManager);
 
 	TabManager->RegisterTabSpawner( ViewportTabId, FOnSpawnTab::CreateSP(this, &FDestructibleMeshEditor::SpawnTab_Viewport) )
 		.SetDisplayName( LOCTEXT("ViewportTab", "Viewport") )
-		.SetGroup( MenuStructure.GetAssetEditorCategory() );
-	
+		.SetGroup(WorkspaceMenuCategoryRef)
+		.SetIcon(FSlateIcon(FEditorStyle::GetStyleSetName(), "LevelEditor.Tabs.Viewports"));
+
 	TabManager->RegisterTabSpawner( PropertiesTabId, FOnSpawnTab::CreateSP(this, &FDestructibleMeshEditor::SpawnTab_Properties) )
 		.SetDisplayName( LOCTEXT("PropertiesTab", "Destructible Settings") )
-		.SetGroup( MenuStructure.GetAssetEditorCategory() );
-	
+		.SetGroup(WorkspaceMenuCategoryRef);
+
 	TabManager->RegisterTabSpawner( FractureSettingsTabId, FOnSpawnTab::CreateSP(this, &FDestructibleMeshEditor::SpawnTab_FractureSettings) )
 		.SetDisplayName( LOCTEXT("FractureSettingsTab", "Fracture Settings") )
-		.SetGroup( MenuStructure.GetAssetEditorCategory() );
-	
+		.SetGroup(WorkspaceMenuCategoryRef);
+
 	TabManager->RegisterTabSpawner( ChunkParametersTabId, FOnSpawnTab::CreateSP(this, &FDestructibleMeshEditor::SpawnTab_ChunkParameters) )
 		.SetDisplayName( LOCTEXT("ChunkParametersTab", "Chunk Parameters") )
-		.SetGroup( MenuStructure.GetAssetEditorCategory() );
+		.SetGroup(WorkspaceMenuCategoryRef);
 }
 
 void FDestructibleMeshEditor::UnregisterTabSpawners(const TSharedRef<class FTabManager>& TabManager)
diff --git a/Engine/Source/Editor/DeviceProfileEditor/Private/DeviceProfileEditorModule.cpp b/Engine/Source/Editor/DeviceProfileEditor/Private/DeviceProfileEditorModule.cpp
--- a/Engine/Source/Editor/DeviceProfileEditor/Private/DeviceProfileEditorModule.cpp
+++ b/Engine/Source/Editor/DeviceProfileEditor/Private/DeviceProfileEditorModule.cpp
@@ -14,9 +14,9 @@ static const FName DeviceProfileEditorName("DeviceProfileEditor");
 void FDeviceProfileEditorModule::StartupModule()
 {
 	FGlobalTabmanager::Get()->RegisterNomadTabSpawner(DeviceProfileEditorName, FOnSpawnTab::CreateStatic(&FDeviceProfileEditorModule::SpawnDeviceProfileEditorTab))
-		.SetDisplayName( NSLOCTEXT("DeviceProfileEditor", "DeviceProfileEditorTitle", "Device Profile Editor") )
+		.SetDisplayName( NSLOCTEXT("DeviceProfileEditor", "DeviceProfileEditorTitle", "Device Profiles") )
 		.SetIcon(FSlateIcon(FEditorStyle::GetStyleSetName(), "DeviceDetails.Tabs.ProfileEditor"))
-		.SetGroup( WorkspaceMenu::GetMenuStructure().GetToolsCategory() );
+		.SetGroup( WorkspaceMenu::GetMenuStructure().GetDeveloperToolsMiscCategory() );
 }
 
 
diff --git a/Engine/Source/Editor/EditorStyle/Private/SlateEditorStyle.cpp b/Engine/Source/Editor/EditorStyle/Private/SlateEditorStyle.cpp
--- a/Engine/Source/Editor/EditorStyle/Private/SlateEditorStyle.cpp
+++ b/Engine/Source/Editor/EditorStyle/Private/SlateEditorStyle.cpp
@@ -1239,7 +1239,7 @@ void FSlateEditorStyle::FStyle::SetupGeneralStyles()
 		Set( "AssetEditor.SaveAsset", new IMAGE_BRUSH( "Icons/icon_SaveAsset_40x", Icon40x40 ) );
 		Set( "AssetEditor.SaveAsset.Small", new IMAGE_BRUSH( "Icons/icon_SaveAsset_40x", Icon20x20 ) );
 		Set( "AssetEditor.ReimportAsset", new IMAGE_BRUSH( "Icons/icon_TextureEd_Reimport_40x", Icon40x40 ) );
-		Set( "AssetEditor.ReimportAsset.Small", new IMAGE_BRUSH( "Icons/icon_TextureEd_Reimport_40x", Icon20x20 ) );		
+		Set( "AssetEditor.ReimportAsset.Small", new IMAGE_BRUSH( "Icons/icon_TextureEd_Reimport_40x", Icon20x20 ) );
 	}
 		
 	// Asset Thumbnail
@@ -1728,7 +1728,7 @@ void FSlateEditorStyle::FStyle::SetupGeneralStyles()
 	// ToolBar
 	{
 		Set( "ToolBar.Background", new BOX_BRUSH( "Common/GroupBorder", FMargin(4.0f/16.0f) ) );
-		Set( "ToolBar.Icon", new IMAGE_BRUSH( "Icons/icon_tab_toolbar_16px", Icon16x16 ) );
+		Set( "ToolBar.Icon", new IMAGE_BRUSH( "Icons/icon_tab_Toolbars_16x", Icon16x16 ) );
 		Set( "ToolBar.Expand", new IMAGE_BRUSH( "Icons/toolbar_expand_16x", Icon16x16) );
 		Set( "ToolBar.SubMenuIndicator", new IMAGE_BRUSH( "Common/SubmenuArrow", Icon8x8 ) );
 		Set( "ToolBar.SToolBarComboButtonBlock.Padding", FMargin(4.0f));
@@ -4399,9 +4399,10 @@ void FSlateEditorStyle::FStyle::SetupPersonaStyle()
 			Set( "Kismet.Tabs.Variables", new IMAGE_BRUSH( "/Icons/pill_16x", Icon16x16 ) );
 			Set( "Kismet.Tabs.Explorer", new IMAGE_BRUSH( "/Icons/hiererchy_16x", Icon16x16 ) );
 			Set( "Kismet.Tabs.Palette", new IMAGE_BRUSH( "/Icons/hiererchy_16x", Icon16x16 ) );
-			Set( "Kismet.Tabs.CompilerResults", new IMAGE_BRUSH( "/Icons/hiererchy_16x", Icon16x16 ) );
+			Set( "Kismet.Tabs.CompilerResults", new IMAGE_BRUSH( "Icons/icon_tab_OutputLog_16x", Icon16x16 ) );
 			Set( "Kismet.Tabs.FindResults", new IMAGE_BRUSH( "/Icons/lens_16x", Icon16x16 ) );
 			Set( "Kismet.Tabs.Components", new IMAGE_BRUSH( "Graph/Icons/Hammer", Icon16x16 ) );
+			Set( "Kismet.Tabs.BlueprintDefaults", new IMAGE_BRUSH( "Icons/icon_BlueprintEditor_Defaults_40x", Icon16x16 ) );
 		}
 
 		Set("Kismet.Palette.Favorites", new IMAGE_BRUSH("Icons/Star_16x", Icon16x16, FLinearColor(0.4f, 0.4, 0.4f, 1.f)));
@@ -4589,6 +4590,11 @@ void FSlateEditorStyle::FStyle::SetupPersonaStyle()
 		Set( "Kismet.DebuggerOverlay.InstructionPointer", new IMAGE_BRUSH( "Old/Kismet2/IP_Normal", FVector2D(128,96)) );
 		Set( "Kismet.DebuggerOverlay.InstructionPointerBreakpoint", new IMAGE_BRUSH( "Old/Kismet2/IP_Breakpoint", FVector2D(128,96)) );
 	}
+
+	// Asset context menu
+	{
+		Set("Persona.AssetActions.CreateAnimAsset", new IMAGE_BRUSH(TEXT("Icons/icon_Persona_CreateAsset_40x"), Icon16x16));
+	}
 #endif // WITH_EDITOR
 }
 	
@@ -4936,6 +4942,19 @@ void FSlateEditorStyle::FStyle::SetupContentBrowserStyle()
 
 		Set( "ContentBrowser.SaveDirtyPackages", new IMAGE_BRUSH( "Icons/icon_file_saveall_40x", Icon25x25 ) );
 		Set( "ContentBrowser.ImportPackage", new IMAGE_BRUSH( "Icons/icon_Import_40x", Icon25x25 ) );
+		Set( "ContentBrowser.ImportIcon", new IMAGE_BRUSH( "Icons/icon_Import_16x", Icon16x16 ) );
+
+		// Asset Context Menu
+		Set( "ContentBrowser.AssetActions.Edit", new IMAGE_BRUSH( "Icons/icon_Editor_Modes_16x", Icon16x16 ) );
+		Set( "ContentBrowser.AssetActions.Delete", new IMAGE_BRUSH( "Icons/icon_delete_16px", Icon16x16, FLinearColor( 0.4f, 0.5f, 0.7f, 1.0f ) ) );
+		Set( "ContentBrowser.AssetActions.ReimportAsset", new IMAGE_BRUSH( "Icons/icon_TextureEd_Reimport_40x", Icon16x16 ) );
+		Set( "ContentBrowser.AssetActions", new IMAGE_BRUSH( "Icons/icon_tab_Tools_16x", Icon16x16 ) );
+		Set( "ContentBrowser.AssetActions.GoToCodeForAsset", new IMAGE_BRUSH( "GameProjectDialog/feature_code_32x", Icon16x16 ) );
+		Set( "ContentBrowser.AssetActions.FindAssetInWorld", new IMAGE_BRUSH( "/Icons/lens_16x", Icon16x16 ) );
+
+		Set( "MediaAsset.AssetActions.Play", new IMAGE_BRUSH( "Icons/icon_SCueEd_PlayCue_40x", Icon16x16 ) );
+		Set( "MediaAsset.AssetActions.Stop", new IMAGE_BRUSH( "Icons/icon_SCueEd_Stop_40x", Icon16x16 ) );
+		Set( "MediaAsset.AssetActions.Pause", new IMAGE_BRUSH( "Icons/icon_SCueEd_Pause_40x", Icon16x16 ) );
 
 		// Misc
 		Set( "ContentBrowser.ThumbnailShadow", new BOX_BRUSH( "ContentBrowser/ThumbnailShadow" , FMargin( 4.0f / 64.0f ) ) );
@@ -5017,55 +5036,55 @@ void FSlateEditorStyle::FStyle::SetupLandscapeEditorStyle()
 		Set("LandscapeEditor.NoiseTool",        new IMAGE_BRUSH("Icons/icon_Landscape_Tool_Noise_40x",            Icon40x40));
 		Set("LandscapeEditor.RetopologizeTool", new IMAGE_BRUSH("Icons/icon_Landscape_Tool_Retopologize_40x",     Icon40x40));
 		Set("LandscapeEditor.VisibilityTool",   new IMAGE_BRUSH("Icons/icon_Landscape_Tool_Visibility_40x",       Icon40x40));
-		Set("LandscapeEditor.SculptTool.Small",       new IMAGE_BRUSH("Icons/icon_Landscape_Tool_Sculpt_20x",           Icon20x20));
-		Set("LandscapeEditor.PaintTool.Small",        new IMAGE_BRUSH("Icons/icon_Landscape_Tool_Paint_20x",            Icon20x20));
-		Set("LandscapeEditor.SmoothTool.Small",       new IMAGE_BRUSH("Icons/icon_Landscape_Tool_Smooth_20x",           Icon20x20));
-		Set("LandscapeEditor.FlattenTool.Small",      new IMAGE_BRUSH("Icons/icon_Landscape_Tool_Flatten_20x",          Icon20x20));
-		Set("LandscapeEditor.RampTool.Small",         new IMAGE_BRUSH("Icons/icon_Landscape_Tool_Ramp_20x",             Icon20x20));
-		Set("LandscapeEditor.ErosionTool.Small",      new IMAGE_BRUSH("Icons/icon_Landscape_Tool_Erosion_20x",          Icon20x20));
+		Set("LandscapeEditor.SculptTool.Small", new IMAGE_BRUSH("Icons/icon_Landscape_Tool_Sculpt_20x", Icon20x20));
+		Set("LandscapeEditor.PaintTool.Small", new IMAGE_BRUSH("Icons/icon_Landscape_Tool_Paint_20x", Icon20x20));
+		Set("LandscapeEditor.SmoothTool.Small", new IMAGE_BRUSH("Icons/icon_Landscape_Tool_Smooth_20x", Icon20x20));
+		Set("LandscapeEditor.FlattenTool.Small", new IMAGE_BRUSH("Icons/icon_Landscape_Tool_Flatten_20x", Icon20x20));
+		Set("LandscapeEditor.RampTool.Small", new IMAGE_BRUSH("Icons/icon_Landscape_Tool_Ramp_20x", Icon20x20));
+		Set("LandscapeEditor.ErosionTool.Small", new IMAGE_BRUSH("Icons/icon_Landscape_Tool_Erosion_20x", Icon20x20));
 		Set("LandscapeEditor.HydroErosionTool.Small", new IMAGE_BRUSH("Icons/icon_Landscape_Tool_HydraulicErosion_20x", Icon20x20));
-		Set("LandscapeEditor.NoiseTool.Small",        new IMAGE_BRUSH("Icons/icon_Landscape_Tool_Noise_20x",            Icon20x20));
-		Set("LandscapeEditor.RetopologizeTool.Small", new IMAGE_BRUSH("Icons/icon_Landscape_Tool_Retopologize_20x",     Icon20x20));
-		Set("LandscapeEditor.VisibilityTool.Small",   new IMAGE_BRUSH("Icons/icon_Landscape_Tool_Visibility_20x",       Icon20x20));
+		Set("LandscapeEditor.NoiseTool.Small", new IMAGE_BRUSH("Icons/icon_Landscape_Tool_Noise_20x", Icon20x20));
+		Set("LandscapeEditor.RetopologizeTool.Small", new IMAGE_BRUSH("Icons/icon_Landscape_Tool_Retopologize_20x", Icon20x20));
+		Set("LandscapeEditor.VisibilityTool.Small", new IMAGE_BRUSH("Icons/icon_Landscape_Tool_Visibility_20x", Icon20x20));
 
-		Set("LandscapeEditor.SelectComponentTool", new IMAGE_BRUSH("Icons/icon_Landscape_Tool_Selection_40x",       Icon40x40));
-		Set("LandscapeEditor.AddComponentTool",    new IMAGE_BRUSH("Icons/icon_Landscape_Tool_AddComponent_40x",    Icon40x40));
+		Set("LandscapeEditor.SelectComponentTool", new IMAGE_BRUSH("Icons/icon_Landscape_Tool_Selection_40x", Icon40x40));
+		Set("LandscapeEditor.AddComponentTool", new IMAGE_BRUSH("Icons/icon_Landscape_Tool_AddComponent_40x", Icon40x40));
 		Set("LandscapeEditor.DeleteComponentTool", new IMAGE_BRUSH("Icons/icon_Landscape_Tool_DeleteComponent_40x", Icon40x40));
-		Set("LandscapeEditor.MoveToLevelTool",     new IMAGE_BRUSH("Icons/icon_Landscape_Tool_MoveToLevel_40x",     Icon40x40));
-		Set("LandscapeEditor.SelectComponentTool.Small", new IMAGE_BRUSH("Icons/icon_Landscape_Tool_Selection_20x",       Icon20x20));
-		Set("LandscapeEditor.AddComponentTool.Small",    new IMAGE_BRUSH("Icons/icon_Landscape_Tool_AddComponent_20x",    Icon20x20));
+		Set("LandscapeEditor.MoveToLevelTool", new IMAGE_BRUSH("Icons/icon_Landscape_Tool_MoveToLevel_40x", Icon40x40));
+		Set("LandscapeEditor.SelectComponentTool.Small", new IMAGE_BRUSH("Icons/icon_Landscape_Tool_Selection_20x", Icon20x20));
+		Set("LandscapeEditor.AddComponentTool.Small", new IMAGE_BRUSH("Icons/icon_Landscape_Tool_AddComponent_20x", Icon20x20));
 		Set("LandscapeEditor.DeleteComponentTool.Small", new IMAGE_BRUSH("Icons/icon_Landscape_Tool_DeleteComponent_20x", Icon20x20));
-		Set("LandscapeEditor.MoveToLevelTool.Small",     new IMAGE_BRUSH("Icons/icon_Landscape_Tool_MoveToLevel_20x",     Icon20x20));
+		Set("LandscapeEditor.MoveToLevelTool.Small", new IMAGE_BRUSH("Icons/icon_Landscape_Tool_MoveToLevel_20x", Icon20x20));
 
-		Set("LandscapeEditor.RegionSelectTool",    new IMAGE_BRUSH("Icons/icon_Landscape_Tool_Mask_40x", Icon40x40));
+		Set("LandscapeEditor.RegionSelectTool", new IMAGE_BRUSH("Icons/icon_Landscape_Tool_Mask_40x", Icon40x40));
 		Set("LandscapeEditor.RegionCopyPasteTool", new IMAGE_BRUSH("Icons/icon_Landscape_Tool_Copy_40x", Icon40x40));
-		Set("LandscapeEditor.RegionSelectTool.Small",    new IMAGE_BRUSH("Icons/icon_Landscape_Tool_Mask_20x", Icon20x20));
+		Set("LandscapeEditor.RegionSelectTool.Small", new IMAGE_BRUSH("Icons/icon_Landscape_Tool_Mask_20x", Icon20x20));
 		Set("LandscapeEditor.RegionCopyPasteTool.Small", new IMAGE_BRUSH("Icons/icon_Landscape_Tool_Copy_20x", Icon20x20));
 
 		Set("LandscapeEditor.SplineTool", new IMAGE_BRUSH("Icons/icon_Landscape_Tool_Splines_40x", Icon40x40));
 		Set("LandscapeEditor.SplineTool.Small", new IMAGE_BRUSH("Icons/icon_Landscape_Tool_Splines_20x", Icon20x20));
 
 		// Brush Sets
-		Set("LandscapeEditor.CircleBrush",        new IMAGE_BRUSH("Icons/icon_Landscape_Brush_Circle_smooth_40x", Icon40x40));
-		Set("LandscapeEditor.AlphaBrush",         new IMAGE_BRUSH("Icons/icon_Landscape_Brush_Alpha_40x",         Icon40x40));
-		Set("LandscapeEditor.AlphaBrush_Pattern", new IMAGE_BRUSH("Icons/icon_Landscape_Brush_Pattern_40x",       Icon40x40));
-		Set("LandscapeEditor.ComponentBrush",     new IMAGE_BRUSH("Icons/icon_Landscape_Brush_Component_40x",     Icon40x40));
-		Set("LandscapeEditor.GizmoBrush",         new IMAGE_BRUSH("Icons/icon_Landscape_Brush_Gizmo_40x",         Icon40x40));
-		Set("LandscapeEditor.CircleBrush.Small",        new IMAGE_BRUSH("Icons/icon_Landscape_Brush_Circle_smooth_20x", Icon20x20));
-		Set("LandscapeEditor.AlphaBrush.Small",         new IMAGE_BRUSH("Icons/icon_Landscape_Brush_Alpha_20x",         Icon20x20));
-		Set("LandscapeEditor.AlphaBrush_Pattern.Small", new IMAGE_BRUSH("Icons/icon_Landscape_Brush_Pattern_20x",       Icon20x20));
-		Set("LandscapeEditor.ComponentBrush.Small",     new IMAGE_BRUSH("Icons/icon_Landscape_Brush_Component_20x",     Icon20x20));
-		Set("LandscapeEditor.GizmoBrush.Small",         new IMAGE_BRUSH("Icons/icon_Landscape_Brush_Gizmo_20x",         Icon20x20));
+		Set("LandscapeEditor.CircleBrush", new IMAGE_BRUSH("Icons/icon_Landscape_Brush_Circle_smooth_40x", Icon40x40));
+		Set("LandscapeEditor.AlphaBrush", new IMAGE_BRUSH("Icons/icon_Landscape_Brush_Alpha_40x", Icon40x40));
+		Set("LandscapeEditor.AlphaBrush_Pattern", new IMAGE_BRUSH("Icons/icon_Landscape_Brush_Pattern_40x", Icon40x40));
+		Set("LandscapeEditor.ComponentBrush", new IMAGE_BRUSH("Icons/icon_Landscape_Brush_Component_40x", Icon40x40));
+		Set("LandscapeEditor.GizmoBrush", new IMAGE_BRUSH("Icons/icon_Landscape_Brush_Gizmo_40x", Icon40x40));
+		Set("LandscapeEditor.CircleBrush.Small", new IMAGE_BRUSH("Icons/icon_Landscape_Brush_Circle_smooth_20x", Icon20x20));
+		Set("LandscapeEditor.AlphaBrush.Small", new IMAGE_BRUSH("Icons/icon_Landscape_Brush_Alpha_20x", Icon20x20));
+		Set("LandscapeEditor.AlphaBrush_Pattern.Small", new IMAGE_BRUSH("Icons/icon_Landscape_Brush_Pattern_20x", Icon20x20));
+		Set("LandscapeEditor.ComponentBrush.Small", new IMAGE_BRUSH("Icons/icon_Landscape_Brush_Component_20x", Icon20x20));
+		Set("LandscapeEditor.GizmoBrush.Small", new IMAGE_BRUSH("Icons/icon_Landscape_Brush_Gizmo_20x", Icon20x20));
 
 		// Brushes
-		Set("LandscapeEditor.CircleBrush_Smooth",    new IMAGE_BRUSH("Icons/icon_Landscape_Brush_Circle_smooth_40x",    Icon40x40));
-		Set("LandscapeEditor.CircleBrush_Linear",    new IMAGE_BRUSH("Icons/icon_Landscape_Brush_Circle_linear_40x",    Icon40x40));
+		Set("LandscapeEditor.CircleBrush_Smooth", new IMAGE_BRUSH("Icons/icon_Landscape_Brush_Circle_smooth_40x", Icon40x40));
+		Set("LandscapeEditor.CircleBrush_Linear", new IMAGE_BRUSH("Icons/icon_Landscape_Brush_Circle_linear_40x", Icon40x40));
 		Set("LandscapeEditor.CircleBrush_Spherical", new IMAGE_BRUSH("Icons/icon_Landscape_Brush_Circle_spherical_40x", Icon40x40));
-		Set("LandscapeEditor.CircleBrush_Tip",       new IMAGE_BRUSH("Icons/icon_Landscape_Brush_Circle_tip_40x",       Icon40x40));
-		Set("LandscapeEditor.CircleBrush_Smooth.Small",    new IMAGE_BRUSH("Icons/icon_Landscape_Brush_Circle_smooth_20x",    Icon20x20));
-		Set("LandscapeEditor.CircleBrush_Linear.Small",    new IMAGE_BRUSH("Icons/icon_Landscape_Brush_Circle_linear_20x",    Icon20x20));
+		Set("LandscapeEditor.CircleBrush_Tip", new IMAGE_BRUSH("Icons/icon_Landscape_Brush_Circle_tip_40x", Icon40x40));
+		Set("LandscapeEditor.CircleBrush_Smooth.Small", new IMAGE_BRUSH("Icons/icon_Landscape_Brush_Circle_smooth_20x", Icon20x20));
+		Set("LandscapeEditor.CircleBrush_Linear.Small", new IMAGE_BRUSH("Icons/icon_Landscape_Brush_Circle_linear_20x", Icon20x20));
 		Set("LandscapeEditor.CircleBrush_Spherical.Small", new IMAGE_BRUSH("Icons/icon_Landscape_Brush_Circle_spherical_20x", Icon20x20));
-		Set("LandscapeEditor.CircleBrush_Tip.Small",       new IMAGE_BRUSH("Icons/icon_Landscape_Brush_Circle_tip_20x",       Icon20x20));
+		Set("LandscapeEditor.CircleBrush_Tip.Small", new IMAGE_BRUSH("Icons/icon_Landscape_Brush_Circle_tip_20x", Icon20x20));
 
 		Set("LandscapeEditor.Brushes.Alpha.UseRChannel", new IMAGE_BRUSH("Icons/icon_Landscape_Brush_Alpha_UseRChannel_20x", Icon20x20));
 		Set("LandscapeEditor.Brushes.Alpha.UseGChannel", new IMAGE_BRUSH("Icons/icon_Landscape_Brush_Alpha_UseGChannel_20x", Icon20x20));
@@ -5073,18 +5092,18 @@ void FSlateEditorStyle::FStyle::SetupLandscapeEditorStyle()
 		Set("LandscapeEditor.Brushes.Alpha.UseAChannel", new IMAGE_BRUSH("Icons/icon_Landscape_Brush_Alpha_UseAChannel_20x", Icon20x20));
 
 		// Target List
-		Set("LandscapeEditor.TargetList.RowBackground",        new FSlateNoResource());
-		Set("LandscapeEditor.TargetList.RowBackgroundHovered", new BOX_BRUSH("Common/RoundedSelection_16x", 4.0f/16.0f, FLinearColor(1.0f, 1.0f, 1.0f, 0.1f)));
-		Set("LandscapeEditor.TargetList.RowSelected",          new BOX_BRUSH("Common/RoundedSelection_16x", 4.0f/16.0f, SelectionColor_Pressed));
-		Set("LandscapeEditor.TargetList.RowSelectedHovered",   new BOX_BRUSH("Common/RoundedSelection_16x", 4.0f/16.0f, SelectionColor));
-		
-		Set("LandscapeEditor.Target_Heightmap",  new IMAGE_BRUSH("Icons/icon_Landscape_Target_Heightmap_48x",  Icon48x48));
+		Set("LandscapeEditor.TargetList.RowBackground", new FSlateNoResource());
+		Set("LandscapeEditor.TargetList.RowBackgroundHovered", new BOX_BRUSH("Common/RoundedSelection_16x", 4.0f / 16.0f, FLinearColor(1.0f, 1.0f, 1.0f, 0.1f)));
+		Set("LandscapeEditor.TargetList.RowSelected", new BOX_BRUSH("Common/RoundedSelection_16x", 4.0f / 16.0f, SelectionColor_Pressed));
+		Set("LandscapeEditor.TargetList.RowSelectedHovered", new BOX_BRUSH("Common/RoundedSelection_16x", 4.0f / 16.0f, SelectionColor));
+
+		Set("LandscapeEditor.Target_Heightmap", new IMAGE_BRUSH("Icons/icon_Landscape_Target_Heightmap_48x", Icon48x48));
 		Set("LandscapeEditor.Target_Visibility", new IMAGE_BRUSH("Icons/icon_Landscape_Target_Visibility_48x", Icon48x48));
-		Set("LandscapeEditor.Target_Invalid",    new IMAGE_BRUSH("Icons/icon_Landscape_Target_Invalid_48x",    Icon48x48));
+		Set("LandscapeEditor.Target_Invalid", new IMAGE_BRUSH("Icons/icon_Landscape_Target_Invalid_48x", Icon48x48));
 
-		Set("LandscapeEditor.Target_Create",     new IMAGE_BRUSH("Icons/icon_Landscape_Target_Create_12x", Icon12x12));
-		Set("LandscapeEditor.Target_MakePublic", new IMAGE_BRUSH("Icons/assign_right_12x",                 Icon12x12));
-		Set("LandscapeEditor.Target_Delete",     new IMAGE_BRUSH("Icons/Cross_12x",                        Icon12x12));
+		Set("LandscapeEditor.Target_Create", new IMAGE_BRUSH("Icons/icon_Landscape_Target_Create_12x", Icon12x12));
+		Set("LandscapeEditor.Target_MakePublic", new IMAGE_BRUSH("Icons/assign_right_12x", Icon12x12));
+		Set("LandscapeEditor.Target_Delete", new IMAGE_BRUSH("Icons/Cross_12x", Icon12x12));
 	}
 #endif
 }
@@ -5100,7 +5119,7 @@ void FSlateEditorStyle::FStyle::SetupToolkitStyles()
 
 		Set("ProjectBrowser.Toolbar.Text", FTextBlockStyle(NormalText)
 			.SetFont(TTF_CORE_FONT("Fonts/Roboto-Bold", 12))
-			.SetShadowColorAndOpacity(FLinearColor(0,0,0,0.5f))
+			.SetShadowColorAndOpacity(FLinearColor(0, 0, 0, 0.5f))
 			.SetShadowOffset(FVector2D(0, 1)));
 
 		Set("ProjectBrowser.VersionOverlayText", FTextBlockStyle(NormalText)
@@ -5111,21 +5130,21 @@ void FSlateEditorStyle::FStyle::SetupToolkitStyles()
 		Set("ProjectBrowser.Background", new BOX_BRUSH("Common/FlatColorSquare", FVector2D(1.0f, 1.0f), FMargin(0), FLinearColor(FColor(0xff404040))));
 		Set("ProjectBrowser.Tab.ActiveBackground", new BOX_BRUSH("Common/FlatColorSquare", FVector2D(1.0f, 1.0f), FMargin(0), FLinearColor(FColor(0xff404040))));
 		Set("ProjectBrowser.Tab.Background", new BOX_BRUSH("Common/FlatColorSquare", FVector2D(1.0f, 1.0f), FMargin(0), FLinearColor(FColor(0xff272727))));
-		Set("ProjectBrowser.Tab.ActiveHighlight", new BOX_BRUSH("Common/FlatColorSquare", FVector2D(1.0f, 1.0f), FMargin(0), SelectionColor) );
+		Set("ProjectBrowser.Tab.ActiveHighlight", new BOX_BRUSH("Common/FlatColorSquare", FVector2D(1.0f, 1.0f), FMargin(0), SelectionColor));
 		Set("ProjectBrowser.Tab.Highlight", new BOX_BRUSH("Common/FlatColorSquare", FVector2D(1.0f, 1.0f), FMargin(0), SelectionColor_Inactive));
 		Set("ProjectBrowser.Tab.PressedHighlight", new BOX_BRUSH("Common/FlatColorSquare", FVector2D(1.0f, 1.0f), FMargin(0), SelectionColor_Pressed));
 
-		Set( "ProjectBrowser.TileViewTooltip.ToolTipBorder", new FSlateColorBrush( FLinearColor::Black ) );
-		Set( "ProjectBrowser.TileViewTooltip.NonContentBorder", new BOX_BRUSH( "/Docking/TabContentArea", FMargin(4/16.0f) ) );
-		Set( "ProjectBrowser.TileViewTooltip.ContentBorder", new BOX_BRUSH( "Common/GroupBorder", FMargin(4.0f/16.0f) ) );
-		Set( "ProjectBrowser.TileViewTooltip.NameFont", TTF_CORE_FONT( "Fonts/Roboto-Regular", 12 ) );
+		Set("ProjectBrowser.TileViewTooltip.ToolTipBorder", new FSlateColorBrush(FLinearColor::Black));
+		Set("ProjectBrowser.TileViewTooltip.NonContentBorder", new BOX_BRUSH("/Docking/TabContentArea", FMargin(4 / 16.0f)));
+		Set("ProjectBrowser.TileViewTooltip.ContentBorder", new BOX_BRUSH("Common/GroupBorder", FMargin(4.0f / 16.0f)));
+		Set("ProjectBrowser.TileViewTooltip.NameFont", TTF_CORE_FONT("Fonts/Roboto-Regular", 12));
 	}
 
 	// Toolkit Display
 	{
-		Set( "ToolkitDisplay.UnsavedChangeIcon", new IMAGE_BRUSH( "Common/UnsavedChange", Icon8x8 ) );
-		Set( "ToolkitDisplay.MenuDropdown", new IMAGE_BRUSH( "Common/ComboArrow", Icon8x8 ) );
-		Set( "ToolkitDisplay.ColorOverlay", new BOX_BRUSH( "/Docking/Tab_ColorOverlay", 4/16.0f ) );
+		Set("ToolkitDisplay.UnsavedChangeIcon", new IMAGE_BRUSH("Common/UnsavedChange", Icon8x8));
+		Set("ToolkitDisplay.MenuDropdown", new IMAGE_BRUSH("Common/ComboArrow", Icon8x8));
+		Set("ToolkitDisplay.ColorOverlay", new BOX_BRUSH("/Docking/Tab_ColorOverlay", 4 / 16.0f));
 
 		FComboButtonStyle ComboButton = FComboButtonStyle()
 			.SetButtonStyle(GetWidgetStyle<FButtonStyle>("Button"))
@@ -5133,27 +5152,27 @@ void FSlateEditorStyle::FStyle::SetupToolkitStyles()
 			// Multiboxes draw their own border so we don't want a default content border
 			.SetMenuBorderBrush(*GetBrush("NoBorder"))
 			.SetMenuBorderPadding(FMargin(0.0f));
-		Set( "ToolkitDisplay.ComboButton", ComboButton );
+		Set("ToolkitDisplay.ComboButton", ComboButton);
 	}
-	
+
 	// Generic Editor
 	{
-		Set( "GenericEditor.Tabs.Properties", new IMAGE_BRUSH( "/Icons/properties_16x", Icon16x16 ) );
+		Set( "GenericEditor.Tabs.Properties", new IMAGE_BRUSH( "/Icons/icon_tab_SelectionDetails_16x", Icon16x16 ) );
 	}
 
 	// CurveTable Editor
 	{
-		Set( "CurveTableEditor.Tabs.Properties", new IMAGE_BRUSH( "/Icons/pillarray_16x", Icon16x16 ) );
+		Set( "CurveTableEditor.Tabs.Properties", new IMAGE_BRUSH( "/Icons/icon_tab_SelectionDetails_16x", Icon16x16 ) );
 	}
 
 	// DataTable Editor
 	{
-		Set( "DataTableEditor.Tabs.Properties", new IMAGE_BRUSH( "/Icons/pillarray_16x", Icon16x16 ) );
+		Set( "DataTableEditor.Tabs.Properties", new IMAGE_BRUSH( "/Icons/icon_tab_SelectionDetails_16x", Icon16x16 ) );
 	}
 
 	// DestructibleMesh Editor
 	{
-		Set( "DestructibleMeshEditor.Tabs.Properties", new IMAGE_BRUSH( "/Icons/pillarray_16x", Icon16x16 ) );
+		Set( "DestructibleMeshEditor.Tabs.Properties", new IMAGE_BRUSH( "/Icons/icon_tab_SelectionDetails_16x", Icon16x16 ) );
 		Set( "DestructibleMeshEditor.Fracture", new IMAGE_BRUSH( "Icons/icon_DestructibleMeshEd_Fracture_40x", Icon40x40 ) );
 		Set( "DestructibleMeshEditor.Fracture.Small", new IMAGE_BRUSH( "Icons/icon_DestructibleMeshEd_Fracture_20x", Icon20x20 ) );
 		Set( "DestructibleMeshEditor.Refresh", new IMAGE_BRUSH( "Icons/icon_DestructibleMeshEd_Refresh_40x", Icon40x40 ) );
@@ -5169,7 +5188,7 @@ void FSlateEditorStyle::FStyle::SetupToolkitStyles()
 	{
 		Set( "MaterialEditor.NormalFont", TTF_CORE_FONT( "Fonts/Roboto-Regular", 9 ) );
 		Set( "MaterialEditor.BoldFont", TTF_CORE_FONT( "Fonts/Roboto-Bold", 9 ) );
-		
+
 		Set( "MaterialEditor.Apply", new IMAGE_BRUSH( "Icons/icon_MatEd_Apply_40x", Icon40x40 ) );
 		Set( "MaterialEditor.Apply.Small", new IMAGE_BRUSH( "Icons/icon_MatEd_Apply_40x", Icon20x20 ) );
 
@@ -5201,26 +5220,26 @@ void FSlateEditorStyle::FStyle::SetupToolkitStyles()
 		Set( "MaterialEditor.ToggleRealtimeExpressions.Small", new IMAGE_BRUSH( "Icons/icon_MatEd_LiveNodes_40x", Icon20x20 ) );
 		Set( "MaterialEditor.AlwaysRefreshAllPreviews", new IMAGE_BRUSH( "Icons/icon_MatEd_Refresh_40x", Icon40x40 ) );
 		Set( "MaterialEditor.AlwaysRefreshAllPreviews.Small", new IMAGE_BRUSH( "Icons/icon_MatEd_Refresh_40x", Icon20x20 ) );
-		Set( "MaterialEditor.ToggleLivePreview", new IMAGE_BRUSH("Icons/icon_MatEd_LivePreview_40x", Icon40x40) );
-		Set( "MaterialEditor.ToggleLivePreview.Small", new IMAGE_BRUSH("Icons/icon_MatEd_LivePreview_40x", Icon20x20) );
+		Set( "MaterialEditor.ToggleLivePreview", new IMAGE_BRUSH( "Icons/icon_MatEd_LivePreview_40x", Icon40x40 ) );
+		Set( "MaterialEditor.ToggleLivePreview.Small", new IMAGE_BRUSH( "Icons/icon_MatEd_LivePreview_40x", Icon20x20 ) );
 		Set( "MaterialEditor.ShowHideConnectors", new IMAGE_BRUSH( "Icons/icon_MatEd_Connectors_40x", Icon40x40 ) );
 		Set( "MaterialEditor.ShowHideConnectors.Small", new IMAGE_BRUSH( "Icons/icon_MatEd_Connectors_40x", Icon20x20 ) );
 		Set( "MaterialEditor.CameraHome", new IMAGE_BRUSH( "Icons/icon_MatEd_Home_40x", Icon40x40 ) );
 		Set( "MaterialEditor.CameraHome.Small", new IMAGE_BRUSH( "Icons/icon_MatEd_Home_40x", Icon20x20 ) );
-		Set( "MaterialEditor.FindInMaterial", new IMAGE_BRUSH("Icons/icon_Blueprint_Find_40px", Icon40x40));
-		Set( "MaterialEditor.FindInMaterial.Small", new IMAGE_BRUSH("Icons/icon_Blueprint_Find_40px", Icon20x20));
-		Set( "MaterialEditor.GoToDocumentation", new IMAGE_BRUSH( "Common/icon_Help_Hover_16x", Icon16x16 ));
+		Set( "MaterialEditor.FindInMaterial", new IMAGE_BRUSH( "Icons/icon_Blueprint_Find_40px", Icon40x40 ) );
+		Set( "MaterialEditor.FindInMaterial.Small", new IMAGE_BRUSH( "Icons/icon_Blueprint_Find_40px", Icon20x20 ) );
+		Set( "MaterialEditor.GoToDocumentation", new IMAGE_BRUSH( "Common/icon_Help_Hover_16x", Icon16x16 ) );
 	}
 
 	// Material Instance Editor
 	{
-		Set( "MaterialInstanceEditor.Tabs.Properties", new IMAGE_BRUSH( "/Icons/properties_16x", Icon16x16 ) );
+		Set( "MaterialInstanceEditor.Tabs.Properties", new IMAGE_BRUSH( "/Icons/icon_tab_SelectionDetails_16x", Icon16x16 ) );
 		Set( "MaterialInstanceEditor.Tabs.Parents", new IMAGE_BRUSH( "/Icons/layers_16x", Icon16x16 ) );
 	}
 
 	// Sound Class Editor
 	{
-		Set( "SoundClassEditor.Tabs.Properties", new IMAGE_BRUSH( "/Icons/properties_16x", Icon16x16 ) );
+		Set( "SoundClassEditor.Tabs.Properties", new IMAGE_BRUSH( "/Icons/icon_tab_SelectionDetails_16x", Icon16x16 ) );
 	}
 
 	// Font Editor
@@ -5228,7 +5247,7 @@ void FSlateEditorStyle::FStyle::SetupToolkitStyles()
 		// Tab icons
 		{
 			Set( "FontEditor.Tabs.Preview", new IMAGE_BRUSH( "/Icons/lens_16x", Icon16x16 ) );
-			Set( "FontEditor.Tabs.Properties", new IMAGE_BRUSH( "/Icons/properties_16x", Icon16x16 ) );
+			Set( "FontEditor.Tabs.Properties", new IMAGE_BRUSH( "/Icons/icon_tab_SelectionDetails_16x", Icon16x16 ) );
 			Set( "FontEditor.Tabs.PageProperties", new IMAGE_BRUSH( "/Icons/properties_16x", Icon16x16 ) );
 		}
 
@@ -5249,18 +5268,18 @@ void FSlateEditorStyle::FStyle::SetupToolkitStyles()
 
 	// SoundCueGraph Editor
 	{
-		Set("SoundCueGraphEditor.PlayCue", new IMAGE_BRUSH( "Icons/icon_SCueEd_PlayCue_40x", Icon40x40));
-		Set("SoundCueGraphEditor.PlayCue.Small", new IMAGE_BRUSH( "Icons/icon_SCueEd_PlayCue_40x", Icon20x20));
-		Set("SoundCueGraphEditor.PlayNode", new IMAGE_BRUSH( "Icons/icon_SCueEd_PlayNode_40x", Icon40x40));
-		Set("SoundCueGraphEditor.PlayNode.Small", new IMAGE_BRUSH( "Icons/icon_SCueEd_PlayNode_40x", Icon20x20));
-		Set("SoundCueGraphEditor.StopCueNode", new IMAGE_BRUSH( "Icons/icon_SCueEd_Stop_40x", Icon40x40));
-		Set("SoundCueGraphEditor.StopCueNode.Small", new IMAGE_BRUSH( "Icons/icon_SCueEd_Stop_40x", Icon20x20));
+		Set( "SoundCueGraphEditor.PlayCue", new IMAGE_BRUSH( "Icons/icon_SCueEd_PlayCue_40x", Icon40x40 ) );
+		Set( "SoundCueGraphEditor.PlayCue.Small", new IMAGE_BRUSH( "Icons/icon_SCueEd_PlayCue_40x", Icon20x20 ) );
+		Set( "SoundCueGraphEditor.PlayNode", new IMAGE_BRUSH( "Icons/icon_SCueEd_PlayNode_40x", Icon40x40 ) );
+		Set( "SoundCueGraphEditor.PlayNode.Small", new IMAGE_BRUSH( "Icons/icon_SCueEd_PlayNode_40x", Icon20x20 ) );
+		Set( "SoundCueGraphEditor.StopCueNode", new IMAGE_BRUSH( "Icons/icon_SCueEd_Stop_40x", Icon40x40 ) );
+		Set( "SoundCueGraphEditor.StopCueNode.Small", new IMAGE_BRUSH( "Icons/icon_SCueEd_Stop_40x", Icon20x20 ) );
 	}
 
 	// Static Mesh Editor
 	{
 		Set( "StaticMeshEditor.NormalFont", TTF_CORE_FONT( "Fonts/Roboto-Regular", 8 ) );
-		Set( "StaticMeshEditor.Tabs.Properties", new IMAGE_BRUSH( "/Icons/properties_16x", Icon16x16 ) );
+		Set( "StaticMeshEditor.Tabs.Properties", new IMAGE_BRUSH( "/Icons/icon_tab_SelectionDetails_16x", Icon16x16 ) );
 		Set( "StaticMeshEditor.SetShowWireframe", new IMAGE_BRUSH( "Icons/icon_StaticMeshEd_Wireframe_40x", Icon40x40 ) );
 		Set( "StaticMeshEditor.SetShowWireframe.Small", new IMAGE_BRUSH( "Icons/icon_StaticMeshEd_Wireframe_40x", Icon20x20 ) );
 		Set( "StaticMeshEditor.SetShowVertexColor", new IMAGE_BRUSH( "Icons/icon_StaticMeshEd_VertColor_40x", Icon40x40 ) );
@@ -5291,17 +5310,17 @@ void FSlateEditorStyle::FStyle::SetupToolkitStyles()
 		Set( "StaticMeshEditor.SetShowBinormals.Small", new IMAGE_BRUSH( "Icons/icon_StaticMeshEd_Binormals_40x", Icon20x20 ) );
 		Set( "StaticMeshEditor.SetDrawAdditionalData", new IMAGE_BRUSH( "Icons/icon_StaticMeshEd_AdditionalData_40x", Icon40x40 ) );
 		Set( "StaticMeshEditor.SetDrawAdditionalData.Small", new IMAGE_BRUSH( "Icons/icon_StaticMeshEd_AdditionalData_40x", Icon20x20 ) );
-		Set( "StaticMeshEditor.GroupSection",  new BOX_BRUSH( "Common/RoundedSelection_16x", FMargin(4.0f/16.0f) ) );
+		Set( "StaticMeshEditor.GroupSection", new BOX_BRUSH( "Common/RoundedSelection_16x", FMargin( 4.0f / 16.0f ) ) );
 	}
 
 	// Skeletal Mesh Editor
 	{
-		Set( "SkeletalMeshEditor.GroupSection",  new BOX_BRUSH( "Common/RoundedSelection_16x", FMargin(4.0f/16.0f) ) );
+		Set( "SkeletalMeshEditor.GroupSection", new BOX_BRUSH( "Common/RoundedSelection_16x", FMargin( 4.0f / 16.0f ) ) );
 	}
 
 	// Texture Editor
 	{
-		Set( "TextureEditor.Tabs.Properties", new IMAGE_BRUSH( "/Icons/properties_16x", Icon16x16 ) );
+		Set("TextureEditor.Tabs.Properties", new IMAGE_BRUSH("/Icons/icon_tab_SelectionDetails_16x", Icon16x16));
 		
 		Set( "TextureEditor.RedChannel", new IMAGE_BRUSH( "Icons/icon_TextureEd_RedChannel_40x", Icon40x40 ) );
 		Set( "TextureEditor.RedChannel.Small", new IMAGE_BRUSH( "Icons/icon_TextureEd_RedChannel_40x", Icon20x20 ) );
@@ -5322,7 +5341,7 @@ void FSlateEditorStyle::FStyle::SetupToolkitStyles()
 
 	// Cascade
 	{
-		Set( "Cascade.Tabs.Properties", new IMAGE_BRUSH( "/Icons/properties_16x", Icon16x16 ) );
+		Set( "Cascade.Tabs.Properties", new IMAGE_BRUSH( "/Icons/icon_tab_SelectionDetails_16x", Icon16x16 ) );
 		
 		Set( "Cascade.RestartSimulation", new IMAGE_BRUSH( "Icons/icon_Cascade_RestartSim_40x", Icon40x40 ) );
 		Set( "Cascade.RestartInLevel", new IMAGE_BRUSH( "Icons/icon_Cascade_RestartInLevel_40x", Icon40x40 ) );
@@ -5425,7 +5444,7 @@ void FSlateEditorStyle::FStyle::SetupToolkitStyles()
 
 	// PhAT
 	{
-		Set( "PhAT.Tabs.Properties", new IMAGE_BRUSH( "/Icons/properties_16x", Icon16x16 ) );
+		Set( "PhAT.Tabs.Properties", new IMAGE_BRUSH( "/Icons/icon_tab_SelectionDetails_16x", Icon16x16 ) );
 		Set( "PhAT.Tabs.Hierarchy", new IMAGE_BRUSH( "/Icons/hiererchy_16x", Icon16x16 ) );
 
 		Set( "PhAT.EditingMode_Body", new IMAGE_BRUSH( "/Icons/icon_PHatMode_Body_40x", Icon40x40) );
@@ -5575,6 +5594,16 @@ void FSlateEditorStyle::FStyle::SetupSourceControlStyles()
 		Set( "SourceControl.StatusIcon.Unknown", new IMAGE_BRUSH( "Icons/SourceControlUnknown_16x", Icon16x16 ) );
 		Set( "SourceControl.LoginWindow.Font", TTF_CORE_FONT( "Fonts/Roboto-Regular", 8 ) );
 		Set( "SourceControl.ProgressWindow.Warning", new IMAGE_BRUSH( "Icons/alert", Icon32x32) );
+
+		// Menu commands
+		// Sync
+		// Check in
+		// Diff
+		// Revert
+		// Connect
+		Set(  "SourceControl.Actions.CheckOut", new IMAGE_BRUSH( "ContentBrowser/SCC_CheckedOut", Icon16x16 ) );
+		Set( "SourceControl.Actions.Add", new IMAGE_BRUSH( "ContentBrowser/SCC_ContentAdd", Icon16x16 ) );
+		Set( "SourceControl.Actions.Refresh", new IMAGE_BRUSH( "Icons/icon_Refresh_16x", Icon16x16 ) );
 	}
 #endif // WITH_EDITOR || IS_PROGRAM
 
diff --git a/Engine/Source/Editor/EnvironmentQueryEditor/Private/EnvironmentQueryEditor.cpp b/Engine/Source/Editor/EnvironmentQueryEditor/Private/EnvironmentQueryEditor.cpp
--- a/Engine/Source/Editor/EnvironmentQueryEditor/Private/EnvironmentQueryEditor.cpp
+++ b/Engine/Source/Editor/EnvironmentQueryEditor/Private/EnvironmentQueryEditor.cpp
@@ -14,22 +14,27 @@
 #include "Editor/PropertyEditor/Public/PropertyEditorModule.h"
 #include "Editor/PropertyEditor/Public/IDetailsView.h"
  
+#define LOCTEXT_NAMESPACE "EnvironmentQueryEditor"
+
 const FName FEnvironmentQueryEditor::EQSUpdateGraphTabId( TEXT( "EnvironmentQueryEditor_UpdateGraph" ) );
 const FName FEnvironmentQueryEditor::EQSPropertiesTabId( TEXT( "EnvironmentQueryEditor_Properties" ) );
 
 void FEnvironmentQueryEditor::RegisterTabSpawners(const TSharedRef<class FTabManager>& TabManager)
 {
-	FAssetEditorToolkit::RegisterTabSpawners(TabManager);
+	WorkspaceMenuCategory = TabManager->GetLocalWorkspaceMenuRoot()->AddGroup(LOCTEXT("WorkspaceMenu_EnvironmentQueryEditor", "Environment Query Editor"));
+	auto WorkspaceMenuCategoryRef = WorkspaceMenuCategory.ToSharedRef();
 
-	const IWorkspaceMenuStructure& MenuStructure = WorkspaceMenu::GetMenuStructure();
+	FAssetEditorToolkit::RegisterTabSpawners(TabManager);
 
 	TabManager->RegisterTabSpawner( EQSUpdateGraphTabId, FOnSpawnTab::CreateSP(this, &FEnvironmentQueryEditor::SpawnTab_UpdateGraph) )
 		.SetDisplayName( NSLOCTEXT("EnvironmentQueryEditor", "Graph", "Graph") )
-		.SetGroup( MenuStructure.GetAssetEditorCategory() );
+		.SetGroup(WorkspaceMenuCategoryRef)
+		.SetIcon(FSlateIcon(FEditorStyle::GetStyleSetName(), "GraphEditor.EventGraph_16x"));
 
 	TabManager->RegisterTabSpawner( EQSPropertiesTabId, FOnSpawnTab::CreateSP(this, &FEnvironmentQueryEditor::SpawnTab_Properties) )
 		.SetDisplayName( NSLOCTEXT("EnvironmentQueryEditor", "PropertiesTab", "Details" ) )
-		.SetGroup( MenuStructure.GetAssetEditorCategory() );
+		.SetGroup(WorkspaceMenuCategoryRef)
+		.SetIcon(FSlateIcon(FEditorStyle::GetStyleSetName(), "LevelEditor.Tabs.Details"));
 }
 
 void FEnvironmentQueryEditor::UnregisterTabSpawners(const TSharedRef<class FTabManager>& TabManager)
@@ -722,3 +727,5 @@ bool FEnvironmentQueryEditor::CanDuplicateNodes() const
 {
 	return CanCopyNodes();
 }
+
+#undef LOCTEXT_NAMESPACE
\ No newline at end of file
diff --git a/Engine/Source/Editor/FontEditor/Private/FontEditor.cpp b/Engine/Source/Editor/FontEditor/Private/FontEditor.cpp
--- a/Engine/Source/Editor/FontEditor/Private/FontEditor.cpp
+++ b/Engine/Source/Editor/FontEditor/Private/FontEditor.cpp
@@ -71,25 +71,30 @@ void FFontEditorCommands::RegisterCommands()
 
 void FFontEditor::RegisterTabSpawners(const TSharedRef<class FTabManager>& TabManager)
 {
-	FAssetEditorToolkit::RegisterTabSpawners(TabManager);
+	WorkspaceMenuCategory = TabManager->GetLocalWorkspaceMenuRoot()->AddGroup(LOCTEXT("WorkspaceMenu_FontEditor", "Font Editor"));
+	auto WorkspaceMenuCategoryRef = WorkspaceMenuCategory.ToSharedRef();
 
-	const IWorkspaceMenuStructure& MenuStructure = WorkspaceMenu::GetMenuStructure();
+	FAssetEditorToolkit::RegisterTabSpawners(TabManager);
 
 	TabManager->RegisterTabSpawner( ViewportTabId,		FOnSpawnTab::CreateSP(this, &FFontEditor::SpawnTab_Viewport) )
 		.SetDisplayName( LOCTEXT("ViewportTab", "Viewport") )
-		.SetGroup( MenuStructure.GetAssetEditorCategory() );
+		.SetGroup(WorkspaceMenuCategoryRef)
+		.SetIcon(FSlateIcon(FEditorStyle::GetStyleSetName(), "LevelEditor.Tabs.Viewports"));
 
 	TabManager->RegisterTabSpawner( PreviewTabId,		FOnSpawnTab::CreateSP(this, &FFontEditor::SpawnTab_Preview) )
 		.SetDisplayName( LOCTEXT("PreviewTab", "Preview") )
-		.SetGroup( MenuStructure.GetAssetEditorCategory() );
+		.SetGroup(WorkspaceMenuCategoryRef)
+		.SetIcon(FSlateIcon(FEditorStyle::GetStyleSetName(), "FontEditor.Tabs.Preview"));
 
 	TabManager->RegisterTabSpawner( PropertiesTabId,	FOnSpawnTab::CreateSP(this, &FFontEditor::SpawnTab_Properties) )
 		.SetDisplayName( LOCTEXT("PropertiesTabId", "Details") )
-		.SetGroup( MenuStructure.GetAssetEditorCategory() );
+		.SetGroup(WorkspaceMenuCategoryRef)
+		.SetIcon(FSlateIcon(FEditorStyle::GetStyleSetName(), "LevelEditor.Tabs.Details"));
 
 	TabManager->RegisterTabSpawner( PagePropertiesTabId,FOnSpawnTab::CreateSP(this, &FFontEditor::SpawnTab_PageProperties) )
 		.SetDisplayName( LOCTEXT("PagePropertiesTab", "Page Details") )
-		.SetGroup( MenuStructure.GetAssetEditorCategory() );
+		.SetGroup(WorkspaceMenuCategoryRef)
+		.SetIcon(FSlateIcon(FEditorStyle::GetStyleSetName(), "FontEditor.Tabs.PageProperties"));
 }
 
 void FFontEditor::UnregisterTabSpawners(const TSharedRef<class FTabManager>& TabManager)
diff --git a/Engine/Source/Editor/Kismet/Private/ApplicationMode.cpp b/Engine/Source/Editor/Kismet/Private/ApplicationMode.cpp
--- a/Engine/Source/Editor/Kismet/Private/ApplicationMode.cpp
+++ b/Engine/Source/Editor/Kismet/Private/ApplicationMode.cpp
@@ -4,9 +4,25 @@
 #include "WorkflowOrientedApp/WorkflowTabFactory.h"
 #include "WorkflowOrientedApp/ApplicationMode.h"
 
+#define LOCTEXT_NAMESPACE "ApplicationMode"
+
 /////////////////////////////////////////////////////
 // FApplicationMode
 
+FApplicationMode::FApplicationMode(FName InModeName)
+	: ModeName(InModeName)
+{
+	ToolbarExtender = MakeShareable(new FExtender);
+	WorkspaceMenuCategory = FWorkspaceItem::NewGroup(LOCTEXT("WorkspaceMenu_ApplicationMode", "Asset Editor"));
+}
+
+FApplicationMode::FApplicationMode(FName InModeName, FText(*GetLocalizedMode)(const FName))
+	: ModeName(InModeName)
+{
+	ToolbarExtender = MakeShareable(new FExtender);
+	WorkspaceMenuCategory = FWorkspaceItem::NewGroup(FText::Format(LOCTEXT("WorkspaceMenu_ApplicationMode", "{0} Editor"), GetLocalizedMode(InModeName)));
+}
+
 void FApplicationMode::DeactivateMode(TSharedPtr<FTabManager> InTabManager)
 {
 	// Save the layout to INI
@@ -42,3 +58,5 @@ TSharedRef<FTabManager::FLayout> FApplicationMode::ActivateMode(TSharedPtr<FTabM
 	check(TabLayout.IsValid());
 	return FLayoutSaveRestore::LoadFromConfig(GEditorLayoutIni, TabLayout.ToSharedRef());
 }
+
+#undef LOCTEXT_NAMESPACE
\ No newline at end of file
diff --git a/Engine/Source/Editor/Kismet/Private/BlueprintEditor.cpp b/Engine/Source/Editor/Kismet/Private/BlueprintEditor.cpp
--- a/Engine/Source/Editor/Kismet/Private/BlueprintEditor.cpp
+++ b/Engine/Source/Editor/Kismet/Private/BlueprintEditor.cpp
@@ -88,7 +88,7 @@ FSelectionDetailsSummoner::FSelectionDetailsSummoner(TSharedPtr<class FAssetEdit
 	: FWorkflowTabFactory(FBlueprintEditorTabs::DetailsID, InHostingApp)
 {
 	TabLabel = LOCTEXT("DetailsView_TabTitle", "Details");
-	TabIcon = FEditorStyle::GetBrush("LevelEditor.Tabs.Details");
+	TabIcon = FSlateIcon(FEditorStyle::GetStyleSetName(), "LevelEditor.Tabs.Details");
 
 	bIsSingleton = true;
 
@@ -1301,14 +1301,14 @@ void FBlueprintEditor::RegisterApplicationModes(const TArray<UBlueprint*>& InBlu
 		{
 			AddApplicationMode(
 				FBlueprintEditorApplicationModes::StandardBlueprintEditorMode,
-				MakeShareable(new FBlueprintEditorApplicationMode(SharedThis(this), FBlueprintEditorApplicationModes::StandardBlueprintEditorMode)));
+				MakeShareable(new FBlueprintEditorApplicationMode(SharedThis(this), FBlueprintEditorApplicationModes::StandardBlueprintEditorMode, FBlueprintEditorApplicationModes::GetLocalizedMode)));
 			SetCurrentMode(FBlueprintEditorApplicationModes::StandardBlueprintEditorMode);
 		}
 		else
 		{
 			AddApplicationMode(
 				FBlueprintEditorApplicationModes::StandardBlueprintEditorMode,
-				MakeShareable(new FBlueprintEditorApplicationMode(SharedThis(this), FBlueprintEditorApplicationModes::StandardBlueprintEditorMode)));
+				MakeShareable(new FBlueprintEditorApplicationMode(SharedThis(this), FBlueprintEditorApplicationModes::StandardBlueprintEditorMode, FBlueprintEditorApplicationModes::GetLocalizedMode)));
 			AddApplicationMode(
 				FBlueprintEditorApplicationModes::BlueprintDefaultsMode,
 				MakeShareable(new FBlueprintDefaultsApplicationMode(SharedThis(this))));
diff --git a/Engine/Source/Editor/Kismet/Private/BlueprintEditorModes.cpp b/Engine/Source/Editor/Kismet/Private/BlueprintEditorModes.cpp
--- a/Engine/Source/Editor/Kismet/Private/BlueprintEditorModes.cpp
+++ b/Engine/Source/Editor/Kismet/Private/BlueprintEditorModes.cpp
@@ -41,8 +41,8 @@ const FName FBlueprintEditorApplicationModes::BlueprintComponentsMode( TEXT("Com
 const FName FBlueprintEditorApplicationModes::BlueprintInterfaceMode( TEXT("InterfaceName") );
 const FName FBlueprintEditorApplicationModes::BlueprintMacroMode( TEXT("MacroName") );
 
-FBlueprintEditorApplicationMode::FBlueprintEditorApplicationMode(TSharedPtr<class FBlueprintEditor> InBlueprintEditor, FName InModeName, const bool bRegisterViewport, const bool bRegisterDefaultsTab)
-	: FApplicationMode(InModeName)
+FBlueprintEditorApplicationMode::FBlueprintEditorApplicationMode(TSharedPtr<class FBlueprintEditor> InBlueprintEditor, FName InModeName, FText(*GetLocalizedMode)(const FName), const bool bRegisterViewport, const bool bRegisterDefaultsTab)
+	: FApplicationMode(InModeName, GetLocalizedMode)
 {
 	MyBlueprintEditor = InBlueprintEditor;
 
@@ -168,7 +168,7 @@ void FBlueprintEditorApplicationMode::PostActivateMode()
 
 
 FBlueprintDefaultsApplicationMode::FBlueprintDefaultsApplicationMode(TSharedPtr<class FBlueprintEditor> InBlueprintEditor)
-	: FApplicationMode(FBlueprintEditorApplicationModes::BlueprintDefaultsMode)
+	: FApplicationMode(FBlueprintEditorApplicationModes::BlueprintDefaultsMode, FBlueprintEditorApplicationModes::GetLocalizedMode)
 {
 	MyBlueprintEditor = InBlueprintEditor;
 	
@@ -223,7 +223,7 @@ void FBlueprintDefaultsApplicationMode::PostActivateMode()
 
 
 FBlueprintComponentsApplicationMode::FBlueprintComponentsApplicationMode(TSharedPtr<class FBlueprintEditor> InBlueprintEditor)
-	: FApplicationMode(FBlueprintEditorApplicationModes::BlueprintComponentsMode)
+	: FApplicationMode(FBlueprintEditorApplicationModes::BlueprintComponentsMode, FBlueprintEditorApplicationModes::GetLocalizedMode)
 {
 	MyBlueprintEditor = InBlueprintEditor;
 	
@@ -342,7 +342,7 @@ void FBlueprintComponentsApplicationMode::PostActivateMode()
 ////////////////////////////////////////
 //
 FBlueprintInterfaceApplicationMode::FBlueprintInterfaceApplicationMode(TSharedPtr<class FBlueprintEditor> InBlueprintEditor)
-	: FApplicationMode(FBlueprintEditorApplicationModes::BlueprintInterfaceMode)
+	: FApplicationMode(FBlueprintEditorApplicationModes::BlueprintInterfaceMode, FBlueprintEditorApplicationModes::GetLocalizedMode)
 {
 	MyBlueprintEditor = InBlueprintEditor;
 	
@@ -423,7 +423,7 @@ void FBlueprintInterfaceApplicationMode::RegisterTabFactories(TSharedPtr<FTabMan
 ////////////////////////////////////////
 //
 FBlueprintMacroApplicationMode::FBlueprintMacroApplicationMode(TSharedPtr<class FBlueprintEditor> InBlueprintEditor)
-	: FApplicationMode(FBlueprintEditorApplicationModes::BlueprintMacroMode)
+	: FApplicationMode(FBlueprintEditorApplicationModes::BlueprintMacroMode, FBlueprintEditorApplicationModes::GetLocalizedMode)
 {
 	MyBlueprintEditor = InBlueprintEditor;
 	
diff --git a/Engine/Source/Editor/Kismet/Private/BlueprintEditorModule.cpp b/Engine/Source/Editor/Kismet/Private/BlueprintEditorModule.cpp
--- a/Engine/Source/Editor/Kismet/Private/BlueprintEditorModule.cpp
+++ b/Engine/Source/Editor/Kismet/Private/BlueprintEditorModule.cpp
@@ -74,7 +74,7 @@ void FBlueprintEditorModule::StartupModule()
 	FGlobalTabmanager::Get()->RegisterNomadTabSpawner( DebuggerAppName, FOnSpawnTab::CreateStatic(&CreateBluprintDebuggerTab) )
 		.SetDisplayName( NSLOCTEXT("BlueprintDebugger", "TabTitle", "Blueprint Debugger") )
 		.SetTooltipText( NSLOCTEXT("BlueprintDebugger", "TooltipText", "Open the Blueprint Debugger tab.") )
-		.SetGroup( MenuStructure.GetToolsCategory() )
+		.SetGroup( MenuStructure.GetDeveloperToolsDebugCategory() )
 		.SetIcon(FSlateIcon(FEditorStyle::GetStyleSetName(), "BlueprintDebugger.TabIcon"));
 
 	// Have to check GIsEditor because right now editor modules can be loaded by the game
diff --git a/Engine/Source/Editor/Kismet/Private/BlueprintEditorTabFactories.h b/Engine/Source/Editor/Kismet/Private/BlueprintEditorTabFactories.h
--- a/Engine/Source/Editor/Kismet/Private/BlueprintEditorTabFactories.h
+++ b/Engine/Source/Editor/Kismet/Private/BlueprintEditorTabFactories.h
@@ -56,6 +56,7 @@ public:
 		: FWorkflowTabFactory(FBlueprintEditorTabs::DebugID, InHostingApp)
 	{
 		TabLabel = LOCTEXT("DebugTabTitle", "Debug");
+		TabIcon = FSlateIcon(FEditorStyle::GetStyleSetName(), "DebugTools.TabIcon");
 
 		EnableTabPadding();
 		bIsSingleton = true;
@@ -82,7 +83,7 @@ public:
 		: FWorkflowTabFactory(FBlueprintEditorTabs::DefaultEditorID, InHostingApp)
 	{
 		TabLabel = LOCTEXT("BlueprintDefaultsTabTitle", "Blueprint Defaults"); //@TODO: ANIMATION: GetDefaultEditorTitle(); !!!
-		TabIcon = FEditorStyle::GetBrush("LevelEditor.Tabs.Details");
+		TabIcon = FSlateIcon(FEditorStyle::GetStyleSetName(), "Kismet.Tabs.BlueprintDefaults");
 
 		bIsSingleton = true;
 
@@ -297,7 +298,7 @@ public:
 		: FWorkflowTabFactory(FBlueprintEditorTabs::ConstructionScriptEditorID, InHostingApp)
 	{
 		TabLabel = LOCTEXT("ComponentsTabLabel", "Components");
-		TabIcon = FEditorStyle::GetBrush("Kismet.Tabs.Components");
+		TabIcon = FSlateIcon(FEditorStyle::GetStyleSetName(), "Kismet.Tabs.Components");
 
 		bIsSingleton = true;
 
@@ -328,6 +329,7 @@ public:
 		: FWorkflowTabFactory(FBlueprintEditorTabs::SCSViewportID, InHostingApp)
 	{
 		TabLabel = LOCTEXT("SCSViewportTabLabel", "Viewport");
+		TabIcon = FSlateIcon(FEditorStyle::GetStyleSetName(), "LevelEditor.Tabs.Viewports");
 
 		bIsSingleton = true;
 
@@ -368,7 +370,7 @@ public:
 		: FWorkflowTabFactory(FBlueprintEditorTabs::PaletteID, InHostingApp)
 	{
 		TabLabel = LOCTEXT("PaletteTabTitle", "Palette");
-		TabIcon = FEditorStyle::GetBrush("Kismet.Tabs.Palette");
+		TabIcon = FSlateIcon(FEditorStyle::GetStyleSetName(), "Kismet.Tabs.Palette");
 
 		bIsSingleton = true;
 
@@ -399,7 +401,7 @@ public:
 		: FWorkflowTabFactory(FBlueprintEditorTabs::MyBlueprintID, InHostingApp)
 	{
 		TabLabel = LOCTEXT("MyBlueprintTabLabel", "My Blueprint");
-		TabIcon = FEditorStyle::GetBrush("Kismet.Tabs.Palette");
+		TabIcon = FSlateIcon(FEditorStyle::GetStyleSetName(), "ClassIcon.BlueprintCore");
 
 		bIsSingleton = true;
 
@@ -430,7 +432,7 @@ public:
 		: FWorkflowTabFactory(FBlueprintEditorTabs::CompilerResultsID, InHostingApp)
 	{
 		TabLabel = LOCTEXT("CompilerResultsTabTitle", "Compiler Results");
-		TabIcon = FEditorStyle::GetBrush("Kismet.Tabs.CompilerResults");
+		TabIcon = FSlateIcon(FEditorStyle::GetStyleSetName(), "Kismet.Tabs.CompilerResults");
 
 		bIsSingleton = true;
 
@@ -461,7 +463,7 @@ public:
 		: FWorkflowTabFactory(FBlueprintEditorTabs::FindResultsID, InHostingApp)
 	{
 		TabLabel = LOCTEXT("FindResultsTabTitle", "Find Results");
-		TabIcon = FEditorStyle::GetBrush("Kismet.Tabs.FindResults");
+		TabIcon = FSlateIcon(FEditorStyle::GetStyleSetName(), "Kismet.Tabs.FindResults");
 
 		bIsSingleton = true;
 
diff --git a/Engine/Source/Editor/Kismet/Private/UserDefinedEnumEditor.cpp b/Engine/Source/Editor/Kismet/Private/UserDefinedEnumEditor.cpp
--- a/Engine/Source/Editor/Kismet/Private/UserDefinedEnumEditor.cpp
+++ b/Engine/Source/Editor/Kismet/Private/UserDefinedEnumEditor.cpp
@@ -13,13 +13,14 @@ const FName FUserDefinedEnumEditor::UserDefinedEnumEditorAppIdentifier( TEXT( "U
 
 void FUserDefinedEnumEditor::RegisterTabSpawners(const TSharedRef<class FTabManager>& TabManager)
 {
-	FAssetEditorToolkit::RegisterTabSpawners(TabManager);
+	WorkspaceMenuCategory = TabManager->GetLocalWorkspaceMenuRoot()->AddGroup(LOCTEXT("WorkspaceMenu_UserDefinedEnumEditor", "User-Defined Enum Editor"));
 
-	const IWorkspaceMenuStructure& MenuStructure = WorkspaceMenu::GetMenuStructure();
+	FAssetEditorToolkit::RegisterTabSpawners(TabManager);
 
 	TabManager->RegisterTabSpawner( EnumeratorsTabId, FOnSpawnTab::CreateSP(this, &FUserDefinedEnumEditor::SpawnEnumeratorsTab) )
 		.SetDisplayName( LOCTEXT("EnumeratorEditor", "Enumerators") )
-		.SetGroup( MenuStructure.GetAssetEditorCategory() );
+		.SetGroup(WorkspaceMenuCategory.ToSharedRef())
+		.SetIcon(FSlateIcon(FEditorStyle::GetStyleSetName(), "GraphEditor.Enum_16x"));
 }
 
 void FUserDefinedEnumEditor::UnregisterTabSpawners(const TSharedRef<class FTabManager>& TabManager)
diff --git a/Engine/Source/Editor/Kismet/Private/UserDefinedStructureEditor.cpp b/Engine/Source/Editor/Kismet/Private/UserDefinedStructureEditor.cpp
--- a/Engine/Source/Editor/Kismet/Private/UserDefinedStructureEditor.cpp
+++ b/Engine/Source/Editor/Kismet/Private/UserDefinedStructureEditor.cpp
@@ -167,13 +167,14 @@ const FName FUserDefinedStructureEditor::UserDefinedStructureEditorAppIdentifier
 
 void FUserDefinedStructureEditor::RegisterTabSpawners(const TSharedRef<class FTabManager>& TabManager)
 {
-	FAssetEditorToolkit::RegisterTabSpawners(TabManager);
+	WorkspaceMenuCategory = TabManager->GetLocalWorkspaceMenuRoot()->AddGroup(LOCTEXT("WorkspaceMenu_UserDefinedStructureEditor", "User-Defined Structure Editor"));
 
-	const IWorkspaceMenuStructure& MenuStructure = WorkspaceMenu::GetMenuStructure();
+	FAssetEditorToolkit::RegisterTabSpawners(TabManager);
 
 	TabManager->RegisterTabSpawner( MemberVariablesTabId, FOnSpawnTab::CreateSP(this, &FUserDefinedStructureEditor::SpawnStructureTab) )
 		.SetDisplayName( LOCTEXT("MemberVariablesEditor", "Member Variables") )
-		.SetGroup( MenuStructure.GetAssetEditorCategory() );
+		.SetGroup(WorkspaceMenuCategory.ToSharedRef())
+		.SetIcon(FSlateIcon(FEditorStyle::GetStyleSetName(), "Kismet.Tabs.Variables"));
 }
 
 void FUserDefinedStructureEditor::UnregisterTabSpawners(const TSharedRef<class FTabManager>& TabManager)
diff --git a/Engine/Source/Editor/Kismet/Private/WorkflowCentricApplication.cpp b/Engine/Source/Editor/Kismet/Private/WorkflowCentricApplication.cpp
--- a/Engine/Source/Editor/Kismet/Private/WorkflowCentricApplication.cpp
+++ b/Engine/Source/Editor/Kismet/Private/WorkflowCentricApplication.cpp
@@ -58,6 +58,11 @@ void FWorkflowCentricApplication::SetCurrentMode(FName NewMode)
 			//@TODO: Should do some validation here
 			CurrentAppModePtr = NewModePtr;
 
+			// Establish the workspace menu category
+			auto LocalWorkspaceRoot = TabManager->GetLocalWorkspaceMenuRoot();
+			LocalWorkspaceRoot->ClearItems();
+			LocalWorkspaceRoot->AddItem(CurrentAppModePtr->GetWorkspaceMenuCategory());
+
 			// Activate the new layout
 			const TSharedRef<FTabManager::FLayout> NewLayout = CurrentAppModePtr->ActivateMode(TabManager);
 			RestoreFromLayout(NewLayout);
@@ -86,10 +91,17 @@ void FWorkflowCentricApplication::PushTabFactories(FWorkflowAllowedTabSet& Facto
 	for (auto FactoryIt = FactorySetToPush.CreateIterator(); FactoryIt; ++FactoryIt)
 	{
 		TSharedPtr<FWorkflowTabFactory> SomeFactory = FactoryIt.Value();
-		TabManager->RegisterTabSpawner( SomeFactory->GetIdentifier(), FOnSpawnTab::CreateRaw( this, &FWorkflowCentricApplication::CreatePanelTab, SomeFactory ) )
-			.SetDisplayName( SomeFactory->ConstructTabName( SpawnInfo ).Get() )
-			.SetGroup(WorkspaceMenu::GetMenuStructure().GetAssetEditorCategory());
-	}
+		FTabSpawnerEntry& SpawnerEntry = TabManager->RegisterTabSpawner(SomeFactory->GetIdentifier(), FOnSpawnTab::CreateRaw(this, &FWorkflowCentricApplication::CreatePanelTab, SomeFactory))
+			.SetDisplayName(SomeFactory->ConstructTabName(SpawnInfo).Get())
+			.SetGroup(CurrentAppModePtr->GetWorkspaceMenuCategory());
+		
+		// Add the tab icon to the menu entry if one was provided
+		const FSlateIcon& TabSpawnerIcon = SomeFactory->GetTabSpawnerIcon(SpawnInfo);
+		if (TabSpawnerIcon.IsSet())
+		{
+			SpawnerEntry.SetIcon(TabSpawnerIcon);
+		}
+	} 
 }
 
 bool FWorkflowCentricApplication::OnRequestClose()
diff --git a/Engine/Source/Editor/Kismet/Private/WorkflowOrientedApp/WorkflowTabFactory.cpp b/Engine/Source/Editor/Kismet/Private/WorkflowOrientedApp/WorkflowTabFactory.cpp
--- a/Engine/Source/Editor/Kismet/Private/WorkflowOrientedApp/WorkflowTabFactory.cpp
+++ b/Engine/Source/Editor/Kismet/Private/WorkflowOrientedApp/WorkflowTabFactory.cpp
@@ -10,7 +10,7 @@
 FWorkflowTabFactory::FWorkflowTabFactory(FName InIdentifier, TSharedPtr<class FAssetEditorToolkit> InHostingApp)
 	: TabIdentifier(InIdentifier)
 	, TabRole(ETabRole::PanelTab)
-	, TabIcon(NULL)
+	, TabIcon(FSlateIcon())
 	, InsideTabPadding(0.0f)
 	, bIsSingleton(false)
 	, bShouldAutosize(false)
@@ -111,9 +111,13 @@ TSharedRef<SWidget> FWorkflowTabFactory::CreateTabBody(const FWorkflowTabSpawnIn
 
 const FSlateBrush* FWorkflowTabFactory::GetTabIcon(const FWorkflowTabSpawnInfo& Info) const
 {
-	return TabIcon;
+	return TabIcon.GetIcon();
 }
 
+const FSlateIcon& FWorkflowTabFactory::GetTabSpawnerIcon(const FWorkflowTabSpawnInfo& Info) const
+{
+	return TabIcon;
+}
 
 void FWorkflowTabFactory::CreateViewMenuEntry(FMenuBuilder& MenuBuilder, const FUIAction& Action) const
 {
diff --git a/Engine/Source/Editor/Kismet/Public/BlueprintEditorModes.h b/Engine/Source/Editor/Kismet/Public/BlueprintEditorModes.h
--- a/Engine/Source/Editor/Kismet/Public/BlueprintEditorModes.h
+++ b/Engine/Source/Editor/Kismet/Public/BlueprintEditorModes.h
@@ -41,7 +41,7 @@ private:
 class KISMET_API FBlueprintEditorApplicationMode : public FApplicationMode
 {
 public:
-	FBlueprintEditorApplicationMode(TSharedPtr<class FBlueprintEditor> InBlueprintEditor, FName InModeName, const bool bRegisterViewport = true, const bool bRegisterDefaultsTab = true);
+	FBlueprintEditorApplicationMode(TSharedPtr<class FBlueprintEditor> InBlueprintEditor, FName InModeName, FText(*GetLocalizedMode)(const FName), const bool bRegisterViewport = true, const bool bRegisterDefaultsTab = true);
 
 	virtual void RegisterTabFactories(TSharedPtr<FTabManager> InTabManager) override;
 	virtual void PreDeactivateMode() override;
diff --git a/Engine/Source/Editor/Kismet/Public/WorkflowOrientedApp/ApplicationMode.h b/Engine/Source/Editor/Kismet/Public/WorkflowOrientedApp/ApplicationMode.h
--- a/Engine/Source/Editor/Kismet/Public/WorkflowOrientedApp/ApplicationMode.h
+++ b/Engine/Source/Editor/Kismet/Public/WorkflowOrientedApp/ApplicationMode.h
@@ -24,12 +24,13 @@ protected:
 	/** The toolbar extension for this mode */
 	TSharedPtr<FExtender> ToolbarExtender;
 
+	/** The workspace menu category for this mode */
+	TSharedPtr<FWorkspaceItem> WorkspaceMenuCategory;
+
 public:
-	FApplicationMode(FName InModeName)
-		: ModeName(InModeName)
-	{
-		ToolbarExtender = MakeShareable(new FExtender);
-	}
+	FApplicationMode(FName InModeName);
+
+	FApplicationMode(FName InModeName, FText(*GetLocalizedMode)(const FName));
 
 	virtual ~FApplicationMode() {}
 
@@ -46,4 +47,7 @@ public:
 	}
 
 	TSharedPtr<FExtender> GetToolbarExtender() { return ToolbarExtender; }
+
+	/** @return The the workspace category for this asset editor */
+	TSharedRef<FWorkspaceItem> GetWorkspaceMenuCategory() { return WorkspaceMenuCategory.ToSharedRef(); }
 };
diff --git a/Engine/Source/Editor/Kismet/Public/WorkflowOrientedApp/WorkflowTabFactory.h b/Engine/Source/Editor/Kismet/Public/WorkflowOrientedApp/WorkflowTabFactory.h
--- a/Engine/Source/Editor/Kismet/Public/WorkflowOrientedApp/WorkflowTabFactory.h
+++ b/Engine/Source/Editor/Kismet/Public/WorkflowOrientedApp/WorkflowTabFactory.h
@@ -61,7 +61,7 @@ protected:
 	ETabRole::Type TabRole;
 
 	FText TabLabel;
-	const FSlateBrush* TabIcon;
+	FSlateIcon TabIcon;
 
 	float InsideTabPadding;
 
@@ -120,6 +120,9 @@ public:
 
 	// Virtual to create an icon for the tab
 	virtual const FSlateBrush* GetTabIcon(const FWorkflowTabSpawnInfo& Info) const;
+	
+	// Virtual to create an icon for the menu entry
+	virtual const FSlateIcon& GetTabSpawnerIcon(const FWorkflowTabSpawnInfo& Info) const;
 
 	// Create a tooltip widget for the tab
 	virtual TSharedPtr<SToolTip> CreateTabToolTipWidget(const FWorkflowTabSpawnInfo& Info) const;
diff --git a/Engine/Source/Editor/LevelEditor/Private/LevelEditor.cpp b/Engine/Source/Editor/LevelEditor/Private/LevelEditor.cpp
--- a/Engine/Source/Editor/LevelEditor/Private/LevelEditor.cpp
+++ b/Engine/Source/Editor/LevelEditor/Private/LevelEditor.cpp
@@ -171,7 +171,7 @@ void FLevelEditorModule::StartupModule()
 	FGlobalTabmanager::Get()->RegisterTabSpawner("LevelEditor", FOnSpawnTab::CreateRaw( this, &FLevelEditorModule::SpawnLevelEditor ) )
 		.SetDisplayName( NSLOCTEXT("LevelEditor", "LevelEditorTab", "Level Editor") );
 
-	FModuleManager::LoadModuleChecked<ISlateReflectorModule>("SlateReflector").RegisterTabSpawner(MenuStructure.GetDeveloperToolsCategory());
+	FModuleManager::LoadModuleChecked<ISlateReflectorModule>("SlateReflector").RegisterTabSpawner(MenuStructure.GetDeveloperToolsMiscCategory());
 
 	FMessageLogModule& MessageLogModule = FModuleManager::LoadModuleChecked<FMessageLogModule>("MessageLog");
 	MessageLogModule.RegisterLogListing("BuildAndSubmitErrors", LOCTEXT("BuildAndSubmitErrors", "Build and Submit Errors"));
diff --git a/Engine/Source/Editor/LevelEditor/Private/LevelEditorToolBar.cpp b/Engine/Source/Editor/LevelEditor/Private/LevelEditorToolBar.cpp
--- a/Engine/Source/Editor/LevelEditor/Private/LevelEditorToolBar.cpp
+++ b/Engine/Source/Editor/LevelEditor/Private/LevelEditorToolBar.cpp
@@ -1472,7 +1472,13 @@ TSharedRef< SWidget > FLevelEditorToolBar::GenerateBuildMenuContent( TSharedRef<
 
 	MenuBuilder.BeginSection("LevelEditorAutomation", LOCTEXT( "AutomationHeading", "Automation" ) );
 	{
-		MenuBuilder.AddMenuEntry( FLevelEditorCommands::Get().BuildAndSubmitToSourceControl );
+		MenuBuilder.AddMenuEntry( 
+			FLevelEditorCommands::Get().BuildAndSubmitToSourceControl,
+			NAME_None,
+			TAttribute<FText>(),
+			TAttribute<FText>(),
+			FSlateIcon(FEditorStyle::GetStyleSetName(), "LevelEditor.Tabs.BuildAndSubmit")
+			);
 	}
 	MenuBuilder.EndSection();
 
@@ -1928,7 +1934,7 @@ TSharedRef< SWidget > FLevelEditorToolBar::GenerateMatineeMenuContent( TSharedRe
 
 	MenuBuilder.BeginSection("LevelEditorNewMatinee", LOCTEXT("MatineeMenuCombo_NewHeading", "New"));
 	{
-		MenuBuilder.AddMenuEntry( FLevelEditorCommands::Get().AddMatinee );
+		MenuBuilder.AddMenuEntry( FLevelEditorCommands::Get().AddMatinee, NAME_None, TAttribute<FText>(), TAttribute<FText>(), FSlateIcon(FEditorStyle::GetStyleSetName(), "ClassIcon.MatineeActor") );
 	}
 	MenuBuilder.EndSection();
 
diff --git a/Engine/Source/Editor/LevelEditor/Private/SLevelEditor.cpp b/Engine/Source/Editor/LevelEditor/Private/SLevelEditor.cpp
--- a/Engine/Source/Editor/LevelEditor/Private/SLevelEditor.cpp
+++ b/Engine/Source/Editor/LevelEditor/Private/SLevelEditor.cpp
@@ -1048,15 +1048,6 @@ TSharedRef<SWidget> SLevelEditor::RestoreContentArea( const TSharedRef<SDockTab>
 				.SetGroup( MenuStructure.GetLevelEditorCategory() );
 		}
 
-		{
-			const FSlateIcon BuildAndSubmitIcon(FEditorStyle::GetStyleSetName(), "LevelEditor.Tabs.BuildAndSubmit");
-			LevelEditorTabManager->RegisterTabSpawner( LevelEditorBuildAndSubmitTab, FOnSpawnTab::CreateSP<SLevelEditor, FName, FString>(this, &SLevelEditor::SpawnLevelEditorTab, LevelEditorBuildAndSubmitTab, FString()) )
-				.SetDisplayName(NSLOCTEXT("LevelEditorTabs", "BuildAndSubmit", "Build And Submit"))
-				.SetTooltipText(NSLOCTEXT("LevelEditorTabs", "BuildAndSubmitTooltip", "Opens a tab allowing the user to build all levels and optionally submit the result to source control."))
-				.SetGroup( MenuStructure.GetLevelEditorCategory() )
-				.SetIcon( BuildAndSubmitIcon );
-		}
-
 		{
 			const FSlateIcon WorldPropertiesIcon(FEditorStyle::GetStyleSetName(), "LevelEditor.WorldProperties.Tab");
 			LevelEditorTabManager->RegisterTabSpawner( "WorldSettingsTab", FOnSpawnTab::CreateSP<SLevelEditor, FName, FString>(this, &SLevelEditor::SpawnLevelEditorTab, FName("WorldSettingsTab"), FString()) )
diff --git a/Engine/Source/Editor/MainFrame/MainFrame.Build.cs b/Engine/Source/Editor/MainFrame/MainFrame.Build.cs
--- a/Engine/Source/Editor/MainFrame/MainFrame.Build.cs
+++ b/Engine/Source/Editor/MainFrame/MainFrame.Build.cs
@@ -10,7 +10,7 @@ public class MainFrame : ModuleRules
         PublicDependencyModuleNames.AddRange(
             new string[]
 			{
-                "Documentation"
+                "Documentation",
 			}
         );
 
@@ -37,6 +37,8 @@ public class MainFrame : ModuleRules
 //				"SearchUI",
 				"TranslationEditor",
 				"Projects",
+                "DeviceProfileEditor",
+                "UndoHistory",
 			}
 		);
 
diff --git a/Engine/Source/Editor/MainFrame/Private/Menus/MainMenu.cpp b/Engine/Source/Editor/MainFrame/Private/Menus/MainMenu.cpp
--- a/Engine/Source/Editor/MainFrame/Private/Menus/MainMenu.cpp
+++ b/Engine/Source/Editor/MainFrame/Private/Menus/MainMenu.cpp
@@ -9,7 +9,8 @@
 #include "SourceCodeNavigation.h"
 #include "Toolkits/AssetEditorToolkit.h"
 #include "TranslationEditorMenu.h"
-
+#include "Editor/DeviceProfileEditor/Public/DeviceProfileEditorModule.h"
+#include "UndoHistoryModule.h"
 
 #define LOCTEXT_NAMESPACE "MainFileMenu"
 
@@ -73,6 +74,14 @@ void FMainMenu::FillEditMenu( FMenuBuilder& MenuBuilder, const TSharedRef< FExte
 		TAttribute< FText > DynamicRedoLabel;
 		DynamicRedoLabel.BindStatic( &Local::GetRedoLabelText );
 		MenuBuilder.AddMenuEntry(FGenericCommands::Get().Redo, "Redo", DynamicRedoLabel); // TAttribute< FString >::Create( &Local::GetRedoLabelText ) );
+
+		// Show undo history
+		MenuBuilder.AddMenuEntry(
+			LOCTEXT("UndoHistoryTabTitle", "Undo History"),
+			LOCTEXT("UndoHistoryTooltipText", "View the entire undo history."),
+			FSlateIcon(FEditorStyle::GetStyleSetName(), "UndoHistory.TabIcon"),
+			FUIAction(FExecuteAction::CreateStatic(&FUndoHistoryModule::ExecuteOpenUndoHistory))
+			);
 	}
 	MenuBuilder.EndSection();
 
@@ -122,7 +131,7 @@ void FMainMenu::FillEditMenu( FMenuBuilder& MenuBuilder, const TSharedRef< FExte
 #if !PLATFORM_MAC // Handled by app's menu in menu bar
 			MenuBuilder.AddMenuEntry(
 				LOCTEXT("EditorPreferencesMenuLabel", "Editor Preferences..."),
-				LOCTEXT("EditorPreferencesMenuToolTip", "Configure the behavior and features of this Editor"),
+				LOCTEXT("EditorPreferencesMenuToolTip", "Configure the behavior and features of the Unreal Editor."),
 				FSlateIcon(),
 				FUIAction(FExecuteAction::CreateStatic(&FSettingsMenu::OpenSettings, FName("Editor"), FName("General"), FName("Appearance")))
 			);
@@ -130,7 +139,7 @@ void FMainMenu::FillEditMenu( FMenuBuilder& MenuBuilder, const TSharedRef< FExte
 
 			MenuBuilder.AddMenuEntry(
 				LOCTEXT("ProjectSettingsMenuLabel", "Project Settings..."),
-				LOCTEXT("ProjectSettingsMenuToolTip", "Change the settings of the currently loaded project"),
+				LOCTEXT("ProjectSettingsMenuToolTip", "Change the settings of the currently loaded project."),
 				FSlateIcon(),
 				FUIAction(FExecuteAction::CreateStatic(&FSettingsMenu::OpenSettings, FName("Project"), FName("Project"), FName("General")))
 			);
@@ -145,53 +154,83 @@ void FMainMenu::FillEditMenu( FMenuBuilder& MenuBuilder, const TSharedRef< FExte
 
 void FMainMenu::FillWindowMenu( FMenuBuilder& MenuBuilder, const TSharedRef< FExtender > Extender, const TSharedPtr<FTabManager> TabManager )
 {
-	MenuBuilder.BeginSection("WindowLocalTabSpawners");
-	{
-		// Automatically populate tab spawners from TabManager
-		if (TabManager.IsValid())
-		{
-			const IWorkspaceMenuStructure& MenuStructure = WorkspaceMenu::GetMenuStructure();
-			TabManager->PopulateTabSpawnerMenu(MenuBuilder, MenuStructure.GetStructureRoot());
-		}
-	}
-	MenuBuilder.EndSection();
-
-	MenuBuilder.BeginSection("WindowGlobalTabSpawners");
+	// Automatically populate tab spawners from TabManager
+	if (TabManager.IsValid())
 	{
-		//@todo The tab system needs to be able to be extendable by plugins [9/3/2013 Justin.Sargent]
-		if (IModularFeatures::Get().IsModularFeatureAvailable(EditorFeatures::PluginsEditor ) )
-		{
-			FGlobalTabmanager::Get()->PopulateTabSpawnerMenu(MenuBuilder, "PluginsEditor");
-		}
+		// Local editor tabs
+		TabManager->PopulateLocalTabSpawnerMenu(MenuBuilder);
+	
+		// General tabs
+		const IWorkspaceMenuStructure& MenuStructure = WorkspaceMenu::GetMenuStructure();
+		TabManager->PopulateTabSpawnerMenu(MenuBuilder, MenuStructure.GetStructureRoot());
 	}
-	MenuBuilder.EndSection();
 
-	MenuBuilder.BeginSection("UnrealFrontendTabs", NSLOCTEXT("MainAppMenu", "UnrealFrontendHeader", "Unreal Frontend"));
 	{
-		FGlobalTabmanager::Get()->PopulateTabSpawnerMenu(MenuBuilder, "DeviceManager");
-
-		if (GetDefault<UEditorExperimentalSettings>()->bMessagingDebugger)
+		// This is a temporary home for the spawners of experimental features that must be explicitly enabled.
+		// When the feature becomes permanent and need not check a flag, register a nomad spawner for it in the proper WorkspaceMenu category
+		bool bProjectLauncher = GetDefault<UEditorExperimentalSettings>()->bProjectLauncher;
+		bool bMessagingDebugger = GetDefault<UEditorExperimentalSettings>()->bMessagingDebugger;
+		bool bBlutility = GetDefault<UEditorExperimentalSettings>()->bEnableEditorUtilityBlueprints;
+		bool bTranslationEditor = GetDefault<UEditorExperimentalSettings>()->bEnableTranslationEditor;
+
+		// Make sure at least one is enabled before creating the section
+		if (bProjectLauncher || bMessagingDebugger || bBlutility || bTranslationEditor)
 		{
-			if (IModularFeatures::Get().IsModularFeatureAvailable("MessagingDebugger"))
+			MenuBuilder.BeginSection("ExperimentalTabSpawners", LOCTEXT("ExperimentalTabSpawnersHeading", "Experimental"));
 			{
-				FGlobalTabmanager::Get()->PopulateTabSpawnerMenu(MenuBuilder, "MessagingDebugger");
-			}
-		}
+				// Project Launcher
+				if (bProjectLauncher)
+				{
+					MenuBuilder.AddMenuEntry(
+						LOCTEXT("ProjectLauncherLabel", "Project Launcher"),
+						LOCTEXT("ProjectLauncherToolTip", "The Project Launcher provides advanced workflows for packaging, deploying and launching your projects."),
+						FSlateIcon(FEditorStyle::GetStyleSetName(), "Launcher.TabIcon"),
+						FUIAction(FExecuteAction::CreateStatic(&FMainMenu::OpenProjectLauncher))
+						);
+				}
 
-		FGlobalTabmanager::Get()->PopulateTabSpawnerMenu(MenuBuilder, "SessionFrontend");
+				// Messaging Debugger
+				if (bMessagingDebugger)
+				{
+					MenuBuilder.AddMenuEntry(
+						LOCTEXT("MessagingDebuggerLabel", "Messaging Debugger"),
+						LOCTEXT("MessagingDebuggerToolTip", "The Messaging Debugger provides a visual utility for debugging the messaging system."),
+						FSlateIcon(), // Icon lives in the plugin dir for the debugger
+						FUIAction(FExecuteAction::CreateStatic(&FMainMenu::OpenMessagingDebugger))
+						);
+				}
 
-		if (GetDefault<UEditorExperimentalSettings>()->bGameLauncher)
-		{
-			FGlobalTabmanager::Get()->PopulateTabSpawnerMenu(MenuBuilder, "ProjectLauncher");
+				// Blutility
+				if (bBlutility)
+				{
+					MenuBuilder.AddMenuEntry(
+						LOCTEXT("BlutilityShelfLabel", "Blutility Shelf"),
+						LOCTEXT("BlutilityShelfToolTip", "Open the blutility shelf."),
+						FSlateIcon(),
+						FUIAction(FExecuteAction::CreateStatic(&FMainMenu::OpenBlutilityShelf))
+						);
+				}
+
+				// Translation Editor
+				if (bTranslationEditor)
+				{
+					MenuBuilder.AddSubMenu(
+						LOCTEXT("TranslationEditorSubMenuLabel", "Translation Editor"),
+						LOCTEXT("EditorPreferencesSubMenuToolTip", "Open the Translation Editor for a Given Project and Language"),
+						FNewMenuDelegate::CreateStatic(&FMainFrameTranslationEditorMenu::MakeMainFrameTranslationEditorSubMenu)
+						);
+				}
+			}
+			MenuBuilder.EndSection();
 		}
+	}
 
-		if (GetDefault<UEditorExperimentalSettings>()->bEnableTranslationEditor)
+	MenuBuilder.BeginSection("WindowGlobalTabSpawners");
+	{
+		//@todo The tab system needs to be able to be extendable by plugins [9/3/2013 Justin.Sargent]
+		if (IModularFeatures::Get().IsModularFeatureAvailable(EditorFeatures::PluginsEditor ) )
 		{
-			MenuBuilder.AddSubMenu(
-				LOCTEXT("TranslationEditorSubMenuLabel", "TranslationEditor"),
-				LOCTEXT("EditorPreferencesSubMenuToolTip", "Open the Translation Editor for a Given Project and Language"),
-				FNewMenuDelegate::CreateStatic(&FMainFrameTranslationEditorMenu::MakeMainFrameTranslationEditorSubMenu)
-			);
+			FGlobalTabmanager::Get()->PopulateTabSpawnerMenu(MenuBuilder, "PluginsEditor");
 		}
 	}
 	MenuBuilder.EndSection();
@@ -422,4 +461,4 @@ TSharedRef< SWidget > FMainMenu::MakeMainTabMenu( const TSharedPtr<FTabManager>&
 }
 
 
-#undef LOCTEXT_NAMESPACE
+#undef LOCTEXT_NAMESPACE
\ No newline at end of file
diff --git a/Engine/Source/Editor/MainFrame/Private/Menus/MainMenu.h b/Engine/Source/Editor/MainFrame/Private/Menus/MainMenu.h
--- a/Engine/Source/Editor/MainFrame/Private/Menus/MainMenu.h
+++ b/Engine/Source/Editor/MainFrame/Private/Menus/MainMenu.h
@@ -64,4 +64,32 @@ protected:
 	 * @param Extender Extensibility support for this menu.
 	 */
 	static void FillHelpMenu( FMenuBuilder& MenuBuilder, const TSharedRef<FExtender> Extender );
+
+private:
+	/** 
+	* Opens the experimental project launcher tab.
+	* Remove this when it is is no longer experimental.
+	*/
+	static void OpenProjectLauncher()
+	{
+		FGlobalTabmanager::Get()->InvokeTab(FName(TEXT("ProjectLauncher")));
+	}
+
+	/**
+	* Opens the experimental messaging debugger tab.
+	* Remove this when it is no longer experimental.
+	*/
+	static void OpenMessagingDebugger()
+	{
+		FGlobalTabmanager::Get()->InvokeTab(FName(TEXT("MessagingDebugger")));
+	}
+
+	/**
+	* Opens the experimental blutility shelf tab.
+	* Remove this when it is no longer experimental.
+	*/
+	static void OpenBlutilityShelf()
+	{
+		FGlobalTabmanager::Get()->InvokeTab(FName(TEXT("BlutilityShelfApp")));
+	}
 };
diff --git a/Engine/Source/Editor/MaterialEditor/Private/MaterialEditor.cpp b/Engine/Source/Editor/MaterialEditor/Private/MaterialEditor.cpp
--- a/Engine/Source/Editor/MaterialEditor/Private/MaterialEditor.cpp
+++ b/Engine/Source/Editor/MaterialEditor/Private/MaterialEditor.cpp
@@ -169,37 +169,44 @@ void FMatExpressionPreview::NotifyCompilationFinished()
 
 void FMaterialEditor::RegisterTabSpawners(const TSharedRef<class FTabManager>& TabManager)
 {
-	FAssetEditorToolkit::RegisterTabSpawners(TabManager);
-
-	const IWorkspaceMenuStructure& MenuStructure = WorkspaceMenu::GetMenuStructure();
+	WorkspaceMenuCategory = TabManager->GetLocalWorkspaceMenuRoot()->AddGroup(LOCTEXT("WorkspaceMenu_MaterialEditor", "Material Editor"));
+	auto WorkspaceMenuCategoryRef = WorkspaceMenuCategory.ToSharedRef();
 
+	FAssetEditorToolkit::RegisterTabSpawners(TabManager);
+	
 	TabManager->RegisterTabSpawner( PreviewTabId, FOnSpawnTab::CreateSP(this, &FMaterialEditor::SpawnTab_Preview) )
 		.SetDisplayName( LOCTEXT("ViewportTab", "Viewport") )
-		.SetGroup( MenuStructure.GetAssetEditorCategory() );
+		.SetGroup( WorkspaceMenuCategoryRef )
+		.SetIcon(FSlateIcon(FEditorStyle::GetStyleSetName(), "LevelEditor.Tabs.Viewports"));
 	
 	TabManager->RegisterTabSpawner( GraphCanvasTabId, FOnSpawnTab::CreateSP(this, &FMaterialEditor::SpawnTab_GraphCanvas) )
 		.SetDisplayName( LOCTEXT("GraphCanvasTab", "Graph") )
-		.SetGroup( MenuStructure.GetAssetEditorCategory() );
+		.SetGroup( WorkspaceMenuCategoryRef )
+		.SetIcon(FSlateIcon(FEditorStyle::GetStyleSetName(), "GraphEditor.EventGraph_16x"));
 	
 	TabManager->RegisterTabSpawner( PropertiesTabId, FOnSpawnTab::CreateSP(this, &FMaterialEditor::SpawnTab_MaterialProperties) )
 		.SetDisplayName( LOCTEXT("DetailsTab", "Details") )
-		.SetGroup( MenuStructure.GetAssetEditorCategory() );
+		.SetGroup( WorkspaceMenuCategoryRef )
+		.SetIcon(FSlateIcon(FEditorStyle::GetStyleSetName(), "LevelEditor.Tabs.Details"));
 
 	TabManager->RegisterTabSpawner( PaletteTabId, FOnSpawnTab::CreateSP(this, &FMaterialEditor::SpawnTab_Palette) )
 		.SetDisplayName( LOCTEXT("PaletteTab", "Palette") )
-		.SetGroup( MenuStructure.GetAssetEditorCategory() );
+		.SetGroup( WorkspaceMenuCategoryRef )
+		.SetIcon(FSlateIcon(FEditorStyle::GetStyleSetName(), "Kismet.Tabs.Palette"));
 
 	TabManager->RegisterTabSpawner( StatsTabId, FOnSpawnTab::CreateSP(this, &FMaterialEditor::SpawnTab_Stats) )
 		.SetDisplayName( LOCTEXT("StatsTab", "Stats") )
-		.SetGroup( MenuStructure.GetAssetEditorCategory() );
+		.SetGroup( WorkspaceMenuCategoryRef )
+		.SetIcon(FSlateIcon(FEditorStyle::GetStyleSetName(), "LevelEditor.Tabs.StatsViewer"));
 	
 	TabManager->RegisterTabSpawner(FindTabId, FOnSpawnTab::CreateSP(this, &FMaterialEditor::SpawnTab_Find))
 		.SetDisplayName(LOCTEXT("FindTab", "Find Results"))
-		.SetGroup(MenuStructure.GetAssetEditorCategory());
+		.SetGroup( WorkspaceMenuCategoryRef )
+		.SetIcon(FSlateIcon(FEditorStyle::GetStyleSetName(), "Kismet.Tabs.FindResults"));
 
 	TabManager->RegisterTabSpawner( HLSLCodeTabId, FOnSpawnTab::CreateSP(this, &FMaterialEditor::SpawnTab_HLSLCode) )
 		.SetDisplayName( LOCTEXT("HLSLCodeTab", "HLSL Code") )
-		.SetGroup( MenuStructure.GetAssetEditorCategory() );
+		.SetGroup( WorkspaceMenuCategoryRef );
 }
 
 
@@ -214,6 +221,8 @@ void FMaterialEditor::UnregisterTabSpawners(const TSharedRef<class FTabManager>&
 	TabManager->UnregisterTabSpawner( StatsTabId );
 	TabManager->UnregisterTabSpawner( FindTabId );
 	TabManager->UnregisterTabSpawner( HLSLCodeTabId );
+
+	TabManager->GetLocalWorkspaceMenuRoot()->ClearItems();
 }
 
 void FMaterialEditor::InitEditorForMaterial(UMaterial* InMaterial)
@@ -398,7 +407,7 @@ void FMaterialEditor::InitMaterialEditor( const EToolkitMode::Type Mode, const T
 	ObjectsToEdit.Add(ObjectToEdit);
 	ObjectsToEdit.Add(Material);
 	FAssetEditorToolkit::InitAssetEditor( Mode, InitToolkitHost, MaterialEditorAppIdentifier, StandaloneDefaultLayout, bCreateDefaultStandaloneMenu, bCreateDefaultToolbar, ObjectsToEdit, false );
-	
+
 	IMaterialEditorModule* MaterialEditorModule = &FModuleManager::LoadModuleChecked<IMaterialEditorModule>( "MaterialEditor" );
 	AddMenuExtender(MaterialEditorModule->GetMenuExtensibilityManager()->GetAllExtenders(GetToolkitCommands(), GetEditingObjects()));
 
diff --git a/Engine/Source/Editor/MaterialEditor/Private/MaterialInstanceEditor.cpp b/Engine/Source/Editor/MaterialEditor/Private/MaterialInstanceEditor.cpp
--- a/Engine/Source/Editor/MaterialEditor/Private/MaterialInstanceEditor.cpp
+++ b/Engine/Source/Editor/MaterialEditor/Private/MaterialInstanceEditor.cpp
@@ -108,21 +108,25 @@ protected:
 
 void FMaterialInstanceEditor::RegisterTabSpawners(const TSharedRef<class FTabManager>& TabManager)
 {
-	FAssetEditorToolkit::RegisterTabSpawners(TabManager);
+	WorkspaceMenuCategory = TabManager->GetLocalWorkspaceMenuRoot()->AddGroup(LOCTEXT("WorkspaceMenu_MaterialInstanceEditor", "Material Instance Editor"));
+	auto WorkspaceMenuCategoryRef = WorkspaceMenuCategory.ToSharedRef();
 
-	const IWorkspaceMenuStructure& MenuStructure = WorkspaceMenu::GetMenuStructure();
+	FAssetEditorToolkit::RegisterTabSpawners(TabManager);
 	
-	TabManager->RegisterTabSpawner( PreviewTabId,		FOnSpawnTab::CreateSP(this, &FMaterialInstanceEditor::SpawnTab_Preview) )
-		.SetDisplayName( LOCTEXT("ViewportTab", "Viewport") )
-		.SetGroup( MenuStructure.GetAssetEditorCategory() );
+	TabManager->RegisterTabSpawner( PreviewTabId, FOnSpawnTab::CreateSP( this, &FMaterialInstanceEditor::SpawnTab_Preview ) )
+		.SetDisplayName( LOCTEXT( "ViewportTab", "Viewport" ) )
+		.SetGroup( WorkspaceMenuCategoryRef )
+		.SetIcon( FSlateIcon( FEditorStyle::GetStyleSetName(), "LevelEditor.Tabs.Viewports" ) );
 	
-	TabManager->RegisterTabSpawner( PropertiesTabId,	FOnSpawnTab::CreateSP(this, &FMaterialInstanceEditor::SpawnTab_Properties) )
-		.SetDisplayName( LOCTEXT("PropertiesTab", "Details") )
-		.SetGroup( MenuStructure.GetAssetEditorCategory() );
+	TabManager->RegisterTabSpawner( PropertiesTabId, FOnSpawnTab::CreateSP( this, &FMaterialInstanceEditor::SpawnTab_Properties ) )
+		.SetDisplayName( LOCTEXT( "PropertiesTab", "Details" ) )
+		.SetGroup( WorkspaceMenuCategoryRef )
+		.SetIcon( FSlateIcon( FEditorStyle::GetStyleSetName(), "LevelEditor.Tabs.Details" ) );
 	
-	TabManager->RegisterTabSpawner( ParentsTabId,		FOnSpawnTab::CreateSP(this, &FMaterialInstanceEditor::SpawnTab_Parents) )
+	TabManager->RegisterTabSpawner( ParentsTabId, FOnSpawnTab::CreateSP(this, &FMaterialInstanceEditor::SpawnTab_Parents) )
 		.SetDisplayName( LOCTEXT("ParentsTab", "Parents") )
-		.SetGroup( MenuStructure.GetAssetEditorCategory() );
+		.SetGroup( WorkspaceMenuCategoryRef )
+		.SetIcon( FSlateIcon( FEditorStyle::GetStyleSetName(), "Kismet.Tabs.Palette" ) );
 }
 
 void FMaterialInstanceEditor::UnregisterTabSpawners(const TSharedRef<class FTabManager>& TabManager)
diff --git a/Engine/Source/Editor/Matinee/Private/Matinee.cpp b/Engine/Source/Editor/Matinee/Private/Matinee.cpp
--- a/Engine/Source/Editor/Matinee/Private/Matinee.cpp
+++ b/Engine/Source/Editor/Matinee/Private/Matinee.cpp
@@ -32,7 +32,6 @@
 #include "LevelEditorActions.h"
 #include "EditorSupportDelegates.h"
 
-#include "Editor/WorkspaceMenuStructure/Public/WorkspaceMenuStructureModule.h"
 #include "MessageLog.h"
 #include "Editor/PropertyEditor/Public/IDetailsView.h"
 
@@ -109,24 +108,36 @@ FLinearColor FMatinee::GetWorldCentricTabColorScale() const
 	return FLinearColor(0.3f, 0.2f, 0.5f, 0.5f);
 }
 
+static const FName MatineeRecordingViewportName("Matinee_RecordingViewport");
 static const FName MatineeCurveEdName("Matinee_CurveEditor");
 static const FName MatineeTrackWindowName("Matinee_TrackWindow");
 static const FName MatineePropertyWindowName("Matinee_PropertyWindow");
 
 void FMatinee::RegisterTabSpawners(const TSharedRef<class FTabManager>& TabManager)
 {
+	WorkspaceMenuCategory = TabManager->GetLocalWorkspaceMenuRoot()->AddGroup(LOCTEXT("WorkspaceMenu_MatineeEditor", "Matinee"));
+	auto WorkspaceMenuCategoryRef = WorkspaceMenuCategory.ToSharedRef();
+
 	FAssetEditorToolkit::RegisterTabSpawners(TabManager);
 
-	const IWorkspaceMenuStructure& MenuStructure = WorkspaceMenu::GetMenuStructure();
+	TabManager->RegisterTabSpawner(MatineeRecordingViewportName, FOnSpawnTab::CreateRaw(this, &FMatinee::SpawnRecordingViewport))
+		.SetDisplayName(NSLOCTEXT("Matinee", "RecordingViewport", "Matinee Recorder"))
+		.SetGroup(WorkspaceMenuCategoryRef);
 
 	TabManager->RegisterTabSpawner(MatineeCurveEdName, FOnSpawnTab::CreateSP(this, &FMatinee::SpawnTab, MatineeCurveEdName))
-		.SetDisplayName(NSLOCTEXT("Matinee", "CurveEditorTitle", "Curve Editor"));
+		.SetDisplayName(NSLOCTEXT("Matinee", "CurveEditorTitle", "Curve Editor"))
+		.SetGroup(WorkspaceMenuCategoryRef);
 
 	TabManager->RegisterTabSpawner(MatineeTrackWindowName, FOnSpawnTab::CreateSP(this, &FMatinee::SpawnTab, MatineeTrackWindowName))
-		.SetDisplayName(NSLOCTEXT("Matinee", "TrackViewEditorTitle", "Tracks"));
+		.SetDisplayName(NSLOCTEXT("Matinee", "TrackViewEditorTitle", "Tracks"))
+		.SetGroup(WorkspaceMenuCategoryRef);
+
 
 	TabManager->RegisterTabSpawner( MatineePropertyWindowName, FOnSpawnTab::CreateSP(this, &FMatinee::SpawnTab, MatineePropertyWindowName) )
-		.SetDisplayName(NSLOCTEXT("Matinee", "PropertiesEditorTitle", "Details"));
+		.SetDisplayName(NSLOCTEXT("Matinee", "PropertiesEditorTitle", "Details"))
+		.SetGroup(WorkspaceMenuCategoryRef)
+		.SetIcon(FSlateIcon(FEditorStyle::GetStyleSetName(), "LevelEditor.Tabs.Details"));
+
 }
 
 void FMatinee::UnregisterTabSpawners(const TSharedRef<class FTabManager>& TabManager)
@@ -919,12 +930,6 @@ void FMatinee::InitMatinee(const EToolkitMode::Type Mode, const TSharedPtr< clas
 		}
 	}
 
-	const IWorkspaceMenuStructure& MenuStructure = WorkspaceMenu::GetMenuStructure();
-
-	FGlobalTabmanager::Get()->RegisterTabSpawner("RecordingViewport", FOnSpawnTab::CreateRaw( this, &FMatinee::SpawnRecordingViewport ) )
-		.SetDisplayName( NSLOCTEXT("MatineeEditor", "RecordingViewport", "Matinee Recorder") )
-		.SetGroup( MenuStructure.GetAssetEditorCategory() );
-
 	TSharedRef<FTabManager::FLayout> Layout = FTabManager::NewLayout( "RecordingViewport_Layout" )
 		->AddArea
 		(
diff --git a/Engine/Source/Editor/NiagaraEditor/Private/NiagaraEditor.cpp b/Engine/Source/Editor/NiagaraEditor/Private/NiagaraEditor.cpp
--- a/Engine/Source/Editor/NiagaraEditor/Private/NiagaraEditor.cpp
+++ b/Engine/Source/Editor/NiagaraEditor/Private/NiagaraEditor.cpp
@@ -12,13 +12,13 @@ const FName FNiagaraEditor::UpdateGraphTabId( TEXT( "NiagaraEditor_UpdateGraph"
 
 void FNiagaraEditor::RegisterTabSpawners(const TSharedRef<class FTabManager>& TabManager)
 {
-	FAssetEditorToolkit::RegisterTabSpawners(TabManager);
+	WorkspaceMenuCategory = TabManager->GetLocalWorkspaceMenuRoot()->AddGroup(LOCTEXT("WorkspaceMenu_NiagaraEditor", "Niagara"));
 
-	const IWorkspaceMenuStructure& MenuStructure = WorkspaceMenu::GetMenuStructure();
+	FAssetEditorToolkit::RegisterTabSpawners(TabManager);
 
 	TabManager->RegisterTabSpawner( UpdateGraphTabId, FOnSpawnTab::CreateSP(this, &FNiagaraEditor::SpawnTab_UpdateGraph) )
 		.SetDisplayName( LOCTEXT("UpdateGraph", "Update Graph") )
-		.SetGroup( MenuStructure.GetAssetEditorCategory() );
+		.SetGroup( WorkspaceMenuCategory.ToSharedRef() );
 }
 
 void FNiagaraEditor::UnregisterTabSpawners(const TSharedRef<class FTabManager>& TabManager)
@@ -35,7 +35,7 @@ FNiagaraEditor::~FNiagaraEditor()
 
 
 void FNiagaraEditor::InitNiagaraEditor( const EToolkitMode::Type Mode, const TSharedPtr< class IToolkitHost >& InitToolkitHost, UNiagaraScript* InScript )
-{
+{	
 	Script = InScript;
 	check(Script != NULL);
 	Source = CastChecked<UNiagaraScriptSource>(Script->Source);
@@ -62,7 +62,7 @@ void FNiagaraEditor::InitNiagaraEditor( const EToolkitMode::Type Mode, const TSh
 	const bool bCreateDefaultStandaloneMenu = true;
 	const bool bCreateDefaultToolbar = true;
 	FAssetEditorToolkit::InitAssetEditor( Mode, InitToolkitHost, FNiagaraEditorModule::NiagaraEditorAppIdentifier, StandaloneDefaultLayout, bCreateDefaultStandaloneMenu, bCreateDefaultToolbar, Script );
-	
+
 	FNiagaraEditorModule& NiagaraEditorModule = FModuleManager::LoadModuleChecked<FNiagaraEditorModule>( "NiagaraEditor" );
 	AddMenuExtender(NiagaraEditorModule.GetMenuExtensibilityManager()->GetAllExtenders(GetToolkitCommands(), GetEditingObjects()));
 
diff --git a/Engine/Source/Editor/NiagaraEditor/Private/NiagaraEffectEditor.cpp b/Engine/Source/Editor/NiagaraEditor/Private/NiagaraEffectEditor.cpp
--- a/Engine/Source/Editor/NiagaraEditor/Private/NiagaraEffectEditor.cpp
+++ b/Engine/Source/Editor/NiagaraEditor/Private/NiagaraEffectEditor.cpp
@@ -13,13 +13,13 @@ const FName FNiagaraEffectEditor::UpdateTabId(TEXT("NiagaraEditor_Effect"));
 
 void FNiagaraEffectEditor::RegisterTabSpawners(const TSharedRef<class FTabManager>& TabManager)
 {
-	FAssetEditorToolkit::RegisterTabSpawners(TabManager);
+	WorkspaceMenuCategory = TabManager->GetLocalWorkspaceMenuRoot()->AddGroup(LOCTEXT("WorkspaceMenu_NiagaraEffectEditor", "Niagara Effect"));
 
-	const IWorkspaceMenuStructure& MenuStructure = WorkspaceMenu::GetMenuStructure();
+	FAssetEditorToolkit::RegisterTabSpawners(TabManager);
 
 	TabManager->RegisterTabSpawner(UpdateTabId, FOnSpawnTab::CreateSP(this, &FNiagaraEffectEditor::SpawnTab))
 		.SetDisplayName(LOCTEXT("NiagaraEffect", "Niagara Effect"))
-		.SetGroup(MenuStructure.GetAssetEditorCategory());
+		.SetGroup( WorkspaceMenuCategory.ToSharedRef() );
 }
 
 void FNiagaraEffectEditor::UnregisterTabSpawners(const TSharedRef<class FTabManager>& TabManager)
@@ -61,7 +61,7 @@ void FNiagaraEffectEditor::InitNiagaraEffectEditor(const EToolkitMode::Type Mode
 	const bool bCreateDefaultStandaloneMenu = true;
 	const bool bCreateDefaultToolbar = true;
 	FAssetEditorToolkit::InitAssetEditor(Mode, InitToolkitHost, FNiagaraEditorModule::NiagaraEditorAppIdentifier, StandaloneDefaultLayout, bCreateDefaultStandaloneMenu, bCreateDefaultToolbar, Effect);
-
+	
 	FNiagaraEditorModule& NiagaraEditorModule = FModuleManager::LoadModuleChecked<FNiagaraEditorModule>("NiagaraEditor");
 	AddMenuExtender(NiagaraEditorModule.GetMenuExtensibilityManager()->GetAllExtenders(GetToolkitCommands(), GetEditingObjects()));
 
diff --git a/Engine/Source/Editor/Persona/Private/BlueprintMode/AnimBlueprintMode.cpp b/Engine/Source/Editor/Persona/Private/BlueprintMode/AnimBlueprintMode.cpp
--- a/Engine/Source/Editor/Persona/Private/BlueprintMode/AnimBlueprintMode.cpp
+++ b/Engine/Source/Editor/Persona/Private/BlueprintMode/AnimBlueprintMode.cpp
@@ -17,7 +17,7 @@
 // FAnimBlueprintEditAppMode
 
 FAnimBlueprintEditAppMode::FAnimBlueprintEditAppMode(TSharedPtr<FPersona> InPersona)
-	: FBlueprintEditorApplicationMode(StaticCastSharedPtr<FBlueprintEditor>(InPersona), FPersonaModes::AnimBlueprintEditMode, false, false)
+	: FBlueprintEditorApplicationMode(StaticCastSharedPtr<FBlueprintEditor>(InPersona), FPersonaModes::AnimBlueprintEditMode, FPersonaModes::GetLocalizedMode, false, false)
 {
 	TabLayout = FTabManager::NewLayout( "Persona_AnimBlueprintEditMode_Layout_v6" )
 		->AddArea
diff --git a/Engine/Source/Editor/Persona/Private/MeshMode/MeshMode.cpp b/Engine/Source/Editor/Persona/Private/MeshMode/MeshMode.cpp
--- a/Engine/Source/Editor/Persona/Private/MeshMode/MeshMode.cpp
+++ b/Engine/Source/Editor/Persona/Private/MeshMode/MeshMode.cpp
@@ -51,7 +51,7 @@ FMeshPropertiesSummoner::FMeshPropertiesSummoner(TSharedPtr<class FAssetEditorTo
 	: FWorkflowTabFactory(FPersonaTabs::MeshAssetPropertiesID, InHostingApp)
 {
 	TabLabel = LOCTEXT("MeshProperties_TabTitle", "Mesh Details");
-	//TabIcon = FEditorStyle::GetBrush("LevelEditor.Tabs.ContentBrowser");
+	TabIcon = FSlateIcon(FEditorStyle::GetStyleSetName(), "ClassIcon.SkeletalMesh");
 
 	bIsSingleton = true;
 
diff --git a/Engine/Source/Editor/Persona/Private/Persona.cpp b/Engine/Source/Editor/Persona/Private/Persona.cpp
--- a/Engine/Source/Editor/Persona/Private/Persona.cpp
+++ b/Engine/Source/Editor/Persona/Private/Persona.cpp
@@ -724,11 +724,7 @@ TSharedRef< SWidget > FPersona::GenerateCreateAssetMenu( USkeleton* Skeleton ) c
 	
 	Skeletons.Add(Skeleton);
 	
-	MenuBuilder.BeginSection(NAME_None, LOCTEXT("Persona_CreateAsset", "Create Asset"));
-	{
-		AnimationEditorUtils::FillCreateAssetMenu(MenuBuilder, Skeletons, FAnimAssetCreated::CreateSP(this, &FPersona::OnAssetCreated), false);
-	}
-	MenuBuilder.EndSection();
+	AnimationEditorUtils::FillCreateAssetMenu(MenuBuilder, Skeletons, FAnimAssetCreated::CreateSP(this, &FPersona::OnAssetCreated), false);
 
 	return MenuBuilder.MakeWidget();
 }
diff --git a/Engine/Source/Editor/Persona/Private/Shared/PersonaMode.cpp b/Engine/Source/Editor/Persona/Private/Shared/PersonaMode.cpp
--- a/Engine/Source/Editor/Persona/Private/Shared/PersonaMode.cpp
+++ b/Engine/Source/Editor/Persona/Private/Shared/PersonaMode.cpp
@@ -56,7 +56,7 @@ const FName FPersonaModes::AnimBlueprintEditMode( "GraphName" );
 // FPersonaAppMode
 
 FPersonaAppMode::FPersonaAppMode(TSharedPtr<class FPersona> InPersona, FName InModeName)
-	: FApplicationMode(InModeName)
+	: FApplicationMode(InModeName, FPersonaModes::GetLocalizedMode)
 {
 	MyPersona = InPersona;
 
@@ -162,7 +162,7 @@ FAnimationAssetBrowserSummoner::FAnimationAssetBrowserSummoner(TSharedPtr<class
 	: FWorkflowTabFactory(FPersonaTabs::AssetBrowserID, InHostingApp)
 {
 	TabLabel = LOCTEXT("AssetBrowserTabTitle", "Asset Browser");
-	TabIcon = FEditorStyle::GetBrush("LevelEditor.Tabs.ContentBrowser");
+	TabIcon = FSlateIcon(FEditorStyle::GetStyleSetName(), "ContentBrowser.TabIcon");
 
 	EnableTabPadding();
 	bIsSingleton = true;
@@ -186,6 +186,7 @@ FPreviewViewportSummoner::FPreviewViewportSummoner(TSharedPtr<class FAssetEditor
 	: FWorkflowTabFactory(FPersonaTabs::PreviewViewportID, InHostingApp)
 {
 	TabLabel = LOCTEXT("ViewportTabTitle", "Viewport");
+	TabIcon = FSlateIcon(FEditorStyle::GetStyleSetName(), "LevelEditor.Tabs.Viewports");
 
 	bIsSingleton = true;
 
diff --git a/Engine/Source/Editor/PhAT/Private/PhAT.cpp b/Engine/Source/Editor/PhAT/Private/PhAT.cpp
--- a/Engine/Source/Editor/PhAT/Private/PhAT.cpp
+++ b/Engine/Source/Editor/PhAT/Private/PhAT.cpp
@@ -74,18 +74,25 @@ static const FName PhATHierarchyName("PhAT_Hierarchy");
 
 void FPhAT::RegisterTabSpawners(const TSharedRef<class FTabManager>& TabManager)
 {
-	FAssetEditorToolkit::RegisterTabSpawners(TabManager);
+	WorkspaceMenuCategory = TabManager->GetLocalWorkspaceMenuRoot()->AddGroup( LOCTEXT( "WorkspaceMenu_MaterialEditor", "Material Editor" ) );
+	auto WorkspaceMenuCategoryRef = WorkspaceMenuCategory.ToSharedRef();
 
-	const IWorkspaceMenuStructure& MenuStructure = WorkspaceMenu::GetMenuStructure();
+	FAssetEditorToolkit::RegisterTabSpawners( TabManager );
 
-	TabManager->RegisterTabSpawner( PhATPreviewViewportName, FOnSpawnTab::CreateSP(this, &FPhAT::SpawnTab, PhATPreviewViewportName) )
-		.SetDisplayName( LOCTEXT( "ViewportTab", "Viewport" ) );
+	TabManager->RegisterTabSpawner( PhATPreviewViewportName, FOnSpawnTab::CreateSP( this, &FPhAT::SpawnTab, PhATPreviewViewportName ) )
+		.SetDisplayName( LOCTEXT( "ViewportTab", "Viewport" ) )
+		.SetGroup( WorkspaceMenuCategoryRef )
+		.SetIcon( FSlateIcon( FEditorStyle::GetStyleSetName(), "LevelEditor.Tabs.Viewports" ) );
 
-	TabManager->RegisterTabSpawner( PhATPropertiesName, FOnSpawnTab::CreateSP(this, &FPhAT::SpawnTab, PhATPropertiesName) )
-		.SetDisplayName( LOCTEXT( "PropertiesTab", "Details" ) );
+	TabManager->RegisterTabSpawner( PhATPropertiesName, FOnSpawnTab::CreateSP( this, &FPhAT::SpawnTab, PhATPropertiesName ) )
+		.SetDisplayName( LOCTEXT( "PropertiesTab", "Details" ) )
+		.SetGroup( WorkspaceMenuCategoryRef )
+		.SetIcon( FSlateIcon( FEditorStyle::GetStyleSetName(), "LevelEditor.Tabs.Details" ) );
 
-	TabManager->RegisterTabSpawner( PhATHierarchyName, FOnSpawnTab::CreateSP(this, &FPhAT::SpawnTab, PhATHierarchyName) )
-		.SetDisplayName( LOCTEXT( "HierarchyTab", "Hierarchy" ) );
+	TabManager->RegisterTabSpawner( PhATHierarchyName, FOnSpawnTab::CreateSP( this, &FPhAT::SpawnTab, PhATHierarchyName ) )
+		.SetDisplayName( LOCTEXT( "HierarchyTab", "Hierarchy" ) )
+		.SetGroup( WorkspaceMenuCategoryRef )
+		.SetIcon( FSlateIcon( FEditorStyle::GetStyleSetName(), "Kismet.Tabs.Palette" ) );
 }
 
 void FPhAT::UnregisterTabSpawners(const TSharedRef<class FTabManager>& TabManager)
diff --git a/Engine/Source/Editor/PropertyEditor/Private/PropertyEditorToolkit.cpp b/Engine/Source/Editor/PropertyEditor/Private/PropertyEditorToolkit.cpp
--- a/Engine/Source/Editor/PropertyEditor/Private/PropertyEditorToolkit.cpp
+++ b/Engine/Source/Editor/PropertyEditor/Private/PropertyEditorToolkit.cpp
@@ -21,15 +21,17 @@ const FName FPropertyEditorToolkit::TreePinAsColumnHeaderId( TEXT( "PropertyEdit
 
 void FPropertyEditorToolkit::RegisterTabSpawners(const TSharedRef<class FTabManager>& TabManager)
 {
-	const IWorkspaceMenuStructure& MenuStructure = WorkspaceMenu::GetMenuStructure();
+	WorkspaceMenuCategory = TabManager->GetLocalWorkspaceMenuRoot()->AddGroup(LOCTEXT("WorkspaceMenu_PropertyEditorToolkit", "Property Editor"));
 
 	TabManager->RegisterTabSpawner( GridTabId, FOnSpawnTab::CreateSP(this, &FPropertyEditorToolkit::SpawnTab_PropertyTable) )
 		.SetDisplayName( LOCTEXT("PropertyTableTab", "Grid") )
-		.SetGroup( MenuStructure.GetAssetEditorCategory() );
+		.SetGroup( WorkspaceMenuCategory.ToSharedRef() )
+		.SetIcon(FSlateIcon(FEditorStyle::GetStyleSetName(), "LevelEditor.Tabs.Viewports"));
 
 	TabManager->RegisterTabSpawner( TreeTabId, FOnSpawnTab::CreateSP(this, &FPropertyEditorToolkit::SpawnTab_PropertyTree) )
 		.SetDisplayName( LOCTEXT("PropertiesTab", "Details") )
-		.SetGroup( MenuStructure.GetAssetEditorCategory() );
+		.SetGroup( WorkspaceMenuCategory.ToSharedRef() )
+		.SetIcon(FSlateIcon(FEditorStyle::GetStyleSetName(), "PropertyEditor.Grid.TabIcon"));
 }
 
 void FPropertyEditorToolkit::UnregisterTabSpawners(const TSharedRef<class FTabManager>& TabManager)
diff --git a/Engine/Source/Editor/Sequencer/Private/SequencerAssetEditor.cpp b/Engine/Source/Editor/Sequencer/Private/SequencerAssetEditor.cpp
--- a/Engine/Source/Editor/Sequencer/Private/SequencerAssetEditor.cpp
+++ b/Engine/Source/Editor/Sequencer/Private/SequencerAssetEditor.cpp
@@ -23,13 +23,13 @@ namespace SequencerDefs
 }
 void FSequencerAssetEditor::RegisterTabSpawners(const TSharedRef<class FTabManager>& TabManager)
 {
-	const IWorkspaceMenuStructure& MenuStructure = WorkspaceMenu::GetMenuStructure();
-
 	if( FSequencer::IsSequencerEnabled() && !IsWorldCentricAssetEditor() )
 	{
+		WorkspaceMenuCategory = TabManager->GetLocalWorkspaceMenuRoot()->AddGroup(LOCTEXT("WorkspaceMenu_SequencerAssetEditor", "Sequencer"));
+
 		TabManager->RegisterTabSpawner( SequencerMainTabId, FOnSpawnTab::CreateSP(this, &FSequencerAssetEditor::SpawnTab_SequencerMain) )
 			.SetDisplayName( LOCTEXT("SequencerMainTab", "Sequencer") )
-			.SetGroup( MenuStructure.GetAssetEditorCategory() );
+			.SetGroup( WorkspaceMenuCategory.ToSharedRef() );
 	}
 
 }
@@ -48,7 +48,6 @@ void FSequencerAssetEditor::UnregisterTabSpawners(const TSharedRef<class FTabMan
 
 void FSequencerAssetEditor::InitSequencerAssetEditor( const EToolkitMode::Type Mode, const FSequencerViewParams& InViewParams, const TSharedPtr< class IToolkitHost >& InitToolkitHost, UMovieScene* InRootMovieScene, const TArray<FOnCreateTrackEditor>& TrackEditorDelegates, bool bEditWithinLevelEditor )
 {
-	
 	{
 		const TSharedRef<FTabManager::FLayout> StandaloneDefaultLayout = FTabManager::NewLayout("Standalone_Sequencer_Layout")
 		->AddArea
diff --git a/Engine/Source/Editor/SoundClassEditor/Private/SoundClassEditor.cpp b/Engine/Source/Editor/SoundClassEditor/Private/SoundClassEditor.cpp
--- a/Engine/Source/Editor/SoundClassEditor/Private/SoundClassEditor.cpp
+++ b/Engine/Source/Editor/SoundClassEditor/Private/SoundClassEditor.cpp
@@ -29,17 +29,20 @@ const FName FSoundClassEditor::PropertiesTabId( TEXT( "SoundClassEditor_Properti
 
 void FSoundClassEditor::RegisterTabSpawners(const TSharedRef<class FTabManager>& TabManager)
 {
-	FAssetEditorToolkit::RegisterTabSpawners(TabManager);
+	WorkspaceMenuCategory = TabManager->GetLocalWorkspaceMenuRoot()->AddGroup(LOCTEXT("WorkspaceMenu_SoundClassEditor", "Sound Class Editor"));
+	auto WorkspaceMenuCategoryRef = WorkspaceMenuCategory.ToSharedRef();
 
-	const IWorkspaceMenuStructure& MenuStructure = WorkspaceMenu::GetMenuStructure();
+	FAssetEditorToolkit::RegisterTabSpawners(TabManager);
 
 	TabManager->RegisterTabSpawner( GraphCanvasTabId, FOnSpawnTab::CreateSP(this, &FSoundClassEditor::SpawnTab_GraphCanvas) )
 		.SetDisplayName( LOCTEXT( "GraphCanvasTab", "Graph" ) )
-		.SetGroup( MenuStructure.GetAssetEditorCategory() );
+		.SetGroup( WorkspaceMenuCategoryRef )
+		.SetIcon(FSlateIcon(FEditorStyle::GetStyleSetName(), "GraphEditor.EventGraph_16x"));
 
 	TabManager->RegisterTabSpawner( PropertiesTabId, FOnSpawnTab::CreateSP(this, &FSoundClassEditor::SpawnTab_Properties) )
 		.SetDisplayName( LOCTEXT( "PropertiesTab", "Details" ) )
-		.SetGroup( MenuStructure.GetAssetEditorCategory() );
+		.SetGroup( WorkspaceMenuCategoryRef )
+		.SetIcon(FSlateIcon(FEditorStyle::GetStyleSetName(), "LevelEditor.Tabs.Details"));
 }
 
 void FSoundClassEditor::UnregisterTabSpawners(const TSharedRef<class FTabManager>& TabManager)
@@ -115,7 +118,7 @@ void FSoundClassEditor::InitSoundClassEditor( const EToolkitMode::Type Mode, con
 	const bool bCreateDefaultStandaloneMenu = true;
 	const bool bCreateDefaultToolbar = true;
 	FAssetEditorToolkit::InitAssetEditor( Mode, InitToolkitHost, SoundClassEditorAppIdentifier, StandaloneDefaultLayout, bCreateDefaultStandaloneMenu, bCreateDefaultToolbar, ObjectToEdit );
-	
+
 	ISoundClassEditorModule* SoundClassEditorModule = &FModuleManager::LoadModuleChecked<ISoundClassEditorModule>( "SoundClassEditor" );
 	AddMenuExtender(SoundClassEditorModule->GetMenuExtensibilityManager()->GetAllExtenders(GetToolkitCommands(), GetEditingObjects()));
 	AddToolbarExtender(SoundClassEditorModule->GetToolBarExtensibilityManager()->GetAllExtenders(GetToolkitCommands(), GetEditingObjects()));
diff --git a/Engine/Source/Editor/SoundCueEditor/Private/SoundCueEditor.cpp b/Engine/Source/Editor/SoundCueEditor/Private/SoundCueEditor.cpp
--- a/Engine/Source/Editor/SoundCueEditor/Private/SoundCueEditor.cpp
+++ b/Engine/Source/Editor/SoundCueEditor/Private/SoundCueEditor.cpp
@@ -25,21 +25,25 @@ const FName FSoundCueEditor::PaletteTabId( TEXT( "SoundCueEditor_Palette" ) );
 
 void FSoundCueEditor::RegisterTabSpawners(const TSharedRef<class FTabManager>& TabManager)
 {
-	FAssetEditorToolkit::RegisterTabSpawners(TabManager);
+	WorkspaceMenuCategory = TabManager->GetLocalWorkspaceMenuRoot()->AddGroup(LOCTEXT("WorkspaceMenu_SoundCueEditor", "Sound Cue Editor"));
+	auto WorkspaceMenuCategoryRef = WorkspaceMenuCategory.ToSharedRef();
 
-	const IWorkspaceMenuStructure& MenuStructure = WorkspaceMenu::GetMenuStructure();
+	FAssetEditorToolkit::RegisterTabSpawners(TabManager);
 
 	TabManager->RegisterTabSpawner( GraphCanvasTabId, FOnSpawnTab::CreateSP(this, &FSoundCueEditor::SpawnTab_GraphCanvas) )
 		.SetDisplayName( LOCTEXT("GraphCanvasTab", "Viewport") )
-		.SetGroup( MenuStructure.GetAssetEditorCategory() );
+		.SetGroup(WorkspaceMenuCategoryRef)
+		.SetIcon(FSlateIcon(FEditorStyle::GetStyleSetName(), "GraphEditor.EventGraph_16x"));
 
 	TabManager->RegisterTabSpawner( PropertiesTabId, FOnSpawnTab::CreateSP(this, &FSoundCueEditor::SpawnTab_Properties) )
 		.SetDisplayName( LOCTEXT("DetailsTab", "Details") )
-		.SetGroup( MenuStructure.GetAssetEditorCategory() );
+		.SetGroup(WorkspaceMenuCategoryRef)
+		.SetIcon(FSlateIcon(FEditorStyle::GetStyleSetName(), "LevelEditor.Tabs.Details"));
 
 	TabManager->RegisterTabSpawner( PaletteTabId, FOnSpawnTab::CreateSP(this, &FSoundCueEditor::SpawnTab_Palette) )
 		.SetDisplayName( LOCTEXT("PaletteTab", "Palette") )
-		.SetGroup( MenuStructure.GetAssetEditorCategory() );
+		.SetGroup(WorkspaceMenuCategoryRef)
+		.SetIcon(FSlateIcon(FEditorStyle::GetStyleSetName(), "Kismet.Tabs.Palette"));
 }
 
 void FSoundCueEditor::UnregisterTabSpawners(const TSharedRef<class FTabManager>& TabManager)
diff --git a/Engine/Source/Editor/StaticMeshEditor/Private/StaticMeshEditor.cpp b/Engine/Source/Editor/StaticMeshEditor/Private/StaticMeshEditor.cpp
--- a/Engine/Source/Editor/StaticMeshEditor/Private/StaticMeshEditor.cpp
+++ b/Engine/Source/Editor/StaticMeshEditor/Private/StaticMeshEditor.cpp
@@ -37,25 +37,28 @@ const FName FStaticMeshEditor::CollisionTabId( TEXT( "StaticMeshEditor_Collision
 
 void FStaticMeshEditor::RegisterTabSpawners(const TSharedRef<class FTabManager>& TabManager)
 {
-	FAssetEditorToolkit::RegisterTabSpawners(TabManager);
+	WorkspaceMenuCategory = TabManager->GetLocalWorkspaceMenuRoot()->AddGroup(LOCTEXT("WorkspaceMenu_StaticMeshEditor", "Static Mesh Editor"));
+	auto WorkspaceMenuCategoryRef = WorkspaceMenuCategory.ToSharedRef();
 
-	const IWorkspaceMenuStructure& MenuStructure = WorkspaceMenu::GetMenuStructure();
+	FAssetEditorToolkit::RegisterTabSpawners(TabManager);
 
 	TabManager->RegisterTabSpawner( ViewportTabId, FOnSpawnTab::CreateSP(this, &FStaticMeshEditor::SpawnTab_Viewport) )
 		.SetDisplayName( LOCTEXT("ViewportTab", "Viewport") )
-		.SetGroup( MenuStructure.GetAssetEditorCategory() );
-	
+		.SetGroup(WorkspaceMenuCategoryRef)
+		.SetIcon(FSlateIcon(FEditorStyle::GetStyleSetName(), "LevelEditor.Tabs.Viewports"));
+
 	TabManager->RegisterTabSpawner( PropertiesTabId, FOnSpawnTab::CreateSP(this, &FStaticMeshEditor::SpawnTab_Properties) )
 		.SetDisplayName( LOCTEXT("PropertiesTab", "Details") )
-		.SetGroup( MenuStructure.GetAssetEditorCategory() );
-	
+		.SetGroup(WorkspaceMenuCategoryRef)
+		.SetIcon(FSlateIcon(FEditorStyle::GetStyleSetName(), "LevelEditor.Tabs.Details"));
+
 	TabManager->RegisterTabSpawner( SocketManagerTabId, FOnSpawnTab::CreateSP(this, &FStaticMeshEditor::SpawnTab_SocketManager) )
 		.SetDisplayName( LOCTEXT("SocketManagerTab", "Socket Manager") )
-		.SetGroup( MenuStructure.GetAssetEditorCategory() );
-	
+		.SetGroup(WorkspaceMenuCategoryRef);
+
 	TabManager->RegisterTabSpawner( CollisionTabId, FOnSpawnTab::CreateSP(this, &FStaticMeshEditor::SpawnTab_Collision) )
 		.SetDisplayName( LOCTEXT("CollisionTab", "Convex Decomposition") )
-		.SetGroup( MenuStructure.GetAssetEditorCategory() );
+		.SetGroup(WorkspaceMenuCategoryRef);
 }
 
 void FStaticMeshEditor::UnregisterTabSpawners(const TSharedRef<class FTabManager>& TabManager)
diff --git a/Engine/Source/Editor/TaskBrowser/Private/TaskBrowserModule.cpp b/Engine/Source/Editor/TaskBrowser/Private/TaskBrowserModule.cpp
--- a/Engine/Source/Editor/TaskBrowser/Private/TaskBrowserModule.cpp
+++ b/Engine/Source/Editor/TaskBrowser/Private/TaskBrowserModule.cpp
@@ -27,7 +27,7 @@ void FTaskBrowserModule::StartupModule()
 {
 	FGlobalTabmanager::Get()->RegisterNomadTabSpawner( TaskBrowserModule::TaskBrowserApp, FOnSpawnTab::CreateStatic( &SpawnTaskBrowserTab ) )
 		.SetDisplayName(NSLOCTEXT("TaskBrowser", "TabTitle", "Task Browser"))
-		.SetGroup( WorkspaceMenu::GetMenuStructure().GetDeveloperToolsCategory() );
+		.SetGroup( WorkspaceMenu::GetMenuStructure().GetDeveloperToolsMiscCategory() );
 }
 
 
diff --git a/Engine/Source/Editor/TextureEditor/Private/TextureEditorToolkit.cpp b/Engine/Source/Editor/TextureEditor/Private/TextureEditorToolkit.cpp
--- a/Engine/Source/Editor/TextureEditor/Private/TextureEditorToolkit.cpp
+++ b/Engine/Source/Editor/TextureEditor/Private/TextureEditorToolkit.cpp
@@ -40,17 +40,20 @@ FString FTextureEditorToolkit::GetDocumentationLink( ) const
 
 void FTextureEditorToolkit::RegisterTabSpawners( const TSharedRef<class FTabManager>& TabManager )
 {
-	FAssetEditorToolkit::RegisterTabSpawners(TabManager);
+	WorkspaceMenuCategory = TabManager->GetLocalWorkspaceMenuRoot()->AddGroup(LOCTEXT("WorkspaceMenu_TextureEditor", "Texture Editor"));
+	auto WorkspaceMenuCategoryRef = WorkspaceMenuCategory.ToSharedRef();
 
-	const IWorkspaceMenuStructure& MenuStructure = WorkspaceMenu::GetMenuStructure();
+	FAssetEditorToolkit::RegisterTabSpawners(TabManager);
 
 	TabManager->RegisterTabSpawner(ViewportTabId, FOnSpawnTab::CreateSP(this, &FTextureEditorToolkit::HandleTabSpawnerSpawnViewport))
 		.SetDisplayName(LOCTEXT("ViewportTab", "Viewport"))
-		.SetGroup(MenuStructure.GetAssetEditorCategory());
-	
+		.SetGroup(WorkspaceMenuCategoryRef)
+		.SetIcon(FSlateIcon(FEditorStyle::GetStyleSetName(), "LevelEditor.Tabs.Viewports"));
+
 	TabManager->RegisterTabSpawner(PropertiesTabId, FOnSpawnTab::CreateSP(this, &FTextureEditorToolkit::HandleTabSpawnerSpawnProperties))
 		.SetDisplayName(LOCTEXT("PropertiesTab", "Details") )
-		.SetGroup(MenuStructure.GetAssetEditorCategory());
+		.SetGroup(WorkspaceMenuCategoryRef)
+		.SetIcon(FSlateIcon(FEditorStyle::GetStyleSetName(), "LevelEditor.Tabs.Details"));
 }
 
 
diff --git a/Engine/Source/Editor/TranslationEditor/Private/TranslationEditor.cpp b/Engine/Source/Editor/TranslationEditor/Private/TranslationEditor.cpp
--- a/Engine/Source/Editor/TranslationEditor/Private/TranslationEditor.cpp
+++ b/Engine/Source/Editor/TranslationEditor/Private/TranslationEditor.cpp
@@ -46,41 +46,42 @@ void FTranslationEditor::Initialize()
 
 void FTranslationEditor::RegisterTabSpawners(const TSharedRef<class FTabManager>& TabManager)
 {
-	FAssetEditorToolkit::RegisterTabSpawners(TabManager);
+	WorkspaceMenuCategory = TabManager->GetLocalWorkspaceMenuRoot()->AddGroup(LOCTEXT("WorkspaceMenu_TranslationEditor", "Translation Editor"));
+	auto WorkspaceMenuCategoryRef = WorkspaceMenuCategory.ToSharedRef();
 
-	const IWorkspaceMenuStructure& MenuStructure = WorkspaceMenu::GetMenuStructure();
+	FAssetEditorToolkit::RegisterTabSpawners(TabManager);
 
 	TabManager->RegisterTabSpawner( UntranslatedTabId, FOnSpawnTab::CreateSP(this, &FTranslationEditor::SpawnTab_Untranslated) )
 		.SetDisplayName( LOCTEXT("UntranslatedTab", "Untranslated") )
-		.SetGroup( MenuStructure.GetAssetEditorCategory() );
+		.SetGroup( WorkspaceMenuCategoryRef );
 
 	TabManager->RegisterTabSpawner( ReviewTabId, FOnSpawnTab::CreateSP(this, &FTranslationEditor::SpawnTab_Review) )
 		.SetDisplayName( LOCTEXT("ReviewTab", "Needs Review") )
-		.SetGroup( MenuStructure.GetAssetEditorCategory() );
+		.SetGroup( WorkspaceMenuCategoryRef );
 
 	TabManager->RegisterTabSpawner( CompletedTabId, FOnSpawnTab::CreateSP(this, &FTranslationEditor::SpawnTab_Completed) )
 		.SetDisplayName( LOCTEXT("CompletedTab", "Completed") )
-		.SetGroup( MenuStructure.GetAssetEditorCategory() );
+		.SetGroup( WorkspaceMenuCategoryRef );
 
 	TabManager->RegisterTabSpawner( PreviewTabId, FOnSpawnTab::CreateSP(this, &FTranslationEditor::SpawnTab_Preview) )
 		.SetDisplayName( LOCTEXT("PreviewTab", "Preview") )
-		.SetGroup( MenuStructure.GetAssetEditorCategory() );
+		.SetGroup( WorkspaceMenuCategoryRef );
 
 	TabManager->RegisterTabSpawner( ContextTabId, FOnSpawnTab::CreateSP(this, &FTranslationEditor::SpawnTab_Context) )
 		.SetDisplayName( LOCTEXT("ContextTab", "Context") )
-		.SetGroup( MenuStructure.GetAssetEditorCategory() );
+		.SetGroup( WorkspaceMenuCategoryRef );
 
 	TabManager->RegisterTabSpawner( HistoryTabId, FOnSpawnTab::CreateSP(this, &FTranslationEditor::SpawnTab_History) )
 		.SetDisplayName( LOCTEXT("HistoryTab", "History") )
-		.SetGroup( MenuStructure.GetAssetEditorCategory() );
+		.SetGroup( WorkspaceMenuCategoryRef );
 
 	TabManager->RegisterTabSpawner( SearchTabId, FOnSpawnTab::CreateSP(this, &FTranslationEditor::SpawnTab_Search) )
 		.SetDisplayName(LOCTEXT("SearchTab", "Search"))
-		.SetGroup(MenuStructure.GetAssetEditorCategory());
+		.SetGroup( WorkspaceMenuCategoryRef );
 
 	TabManager->RegisterTabSpawner( ChangedOnImportTabId, FOnSpawnTab::CreateSP(this, &FTranslationEditor::SpawnTab_ChangedOnImport) )
-		.SetDisplayName(LOCTEXT("ChangedOnImportTab", "ChangedOnImport"))
-		.SetGroup(MenuStructure.GetAssetEditorCategory());
+		.SetDisplayName(LOCTEXT("ChangedOnImportTab", "Changed On Import"))
+		.SetGroup( WorkspaceMenuCategoryRef );
 }
 
 void FTranslationEditor::UnregisterTabSpawners(const TSharedRef<class FTabManager>& TabManager)
diff --git a/Engine/Source/Editor/UMGEditor/Private/BlueprintModes/WidgetBlueprintApplicationMode.cpp b/Engine/Source/Editor/UMGEditor/Private/BlueprintModes/WidgetBlueprintApplicationMode.cpp
--- a/Engine/Source/Editor/UMGEditor/Private/BlueprintModes/WidgetBlueprintApplicationMode.cpp
+++ b/Engine/Source/Editor/UMGEditor/Private/BlueprintModes/WidgetBlueprintApplicationMode.cpp
@@ -14,7 +14,7 @@
 // FWidgetBlueprintApplicationMode
 
 FWidgetBlueprintApplicationMode::FWidgetBlueprintApplicationMode(TSharedPtr<FWidgetBlueprintEditor> InWidgetEditor, FName InModeName)
-	: FBlueprintEditorApplicationMode(InWidgetEditor, InModeName, false, false)
+	: FBlueprintEditorApplicationMode(InWidgetEditor, InModeName, FWidgetBlueprintApplicationModes::GetLocalizedMode, false, false)
 	, MyWidgetBlueprintEditor(InWidgetEditor)
 {
 }
@@ -34,4 +34,4 @@ UWidgetBlueprint* FWidgetBlueprintApplicationMode::GetBlueprint() const
 TSharedPtr<FWidgetBlueprintEditor> FWidgetBlueprintApplicationMode::GetBlueprintEditor() const
 {
 	return MyWidgetBlueprintEditor.Pin();
-}
+}
\ No newline at end of file
diff --git a/Engine/Source/Editor/UMGEditor/Private/BlueprintModes/WidgetBlueprintApplicationModes.cpp b/Engine/Source/Editor/UMGEditor/Private/BlueprintModes/WidgetBlueprintApplicationModes.cpp
--- a/Engine/Source/Editor/UMGEditor/Private/BlueprintModes/WidgetBlueprintApplicationModes.cpp
+++ b/Engine/Source/Editor/UMGEditor/Private/BlueprintModes/WidgetBlueprintApplicationModes.cpp
@@ -3,8 +3,8 @@
 #include "UMGEditorPrivatePCH.h"
 
 // Mode constants
-const FName FWidgetBlueprintApplicationModes::DesignerMode("DesignModeName");
-const FName FWidgetBlueprintApplicationModes::GraphMode("GraphModeName");
+const FName FWidgetBlueprintApplicationModes::DesignerMode("DesignerName");
+const FName FWidgetBlueprintApplicationModes::GraphMode("GraphName");
 
 FText FWidgetBlueprintApplicationModes::GetLocalizedMode(const FName InMode)
 {
diff --git a/Engine/Source/Editor/UMGEditor/Private/BlueprintModes/WidgetDesignerApplicationMode.cpp b/Engine/Source/Editor/UMGEditor/Private/BlueprintModes/WidgetDesignerApplicationMode.cpp
--- a/Engine/Source/Editor/UMGEditor/Private/BlueprintModes/WidgetDesignerApplicationMode.cpp
+++ b/Engine/Source/Editor/UMGEditor/Private/BlueprintModes/WidgetDesignerApplicationMode.cpp
@@ -22,12 +22,17 @@
 #include "DetailsTabSummoner.h"
 #include "AnimationTabSummoner.h"
 
+#define LOCTEXT_NAMESPACE "WidgetDesignerMode"
+
 /////////////////////////////////////////////////////
 // FWidgetDesignerApplicationMode
 
 FWidgetDesignerApplicationMode::FWidgetDesignerApplicationMode(TSharedPtr<FWidgetBlueprintEditor> InWidgetEditor)
 	: FWidgetBlueprintApplicationMode(InWidgetEditor, FWidgetBlueprintApplicationModes::DesignerMode)
 {
+	// Override the default created category here since "Designer Editor" sounds awkward
+	WorkspaceMenuCategory = FWorkspaceItem::NewGroup(LOCTEXT("WorkspaceMenu_WidgetDesigner", "Widget Designer"));
+
 	TabLayout = FTabManager::NewLayout( "WidgetBlueprintEditor_Designer_Layout_v3" )
 	->AddArea
 	(
@@ -104,7 +109,6 @@ FWidgetDesignerApplicationMode::FWidgetDesignerApplicationMode(TSharedPtr<FWidge
 		)
 	);
 
-
 	// Add Tab Spawners
 	//TabFactories.RegisterFactory(MakeShareable(new FSelectionDetailsSummoner(InWidgetEditor)));
 	TabFactories.RegisterFactory(MakeShareable(new FDetailsTabSummoner(InWidgetEditor)));
@@ -141,3 +145,5 @@ void FWidgetDesignerApplicationMode::PostActivateMode()
 {
 	//FWidgetBlueprintApplicationMode::PostActivateMode();
 }
+
+#undef LOCTEXT_NAMESPACE
\ No newline at end of file
diff --git a/Engine/Source/Editor/UMGEditor/Private/TabFactory/AnimationTabSummoner.cpp b/Engine/Source/Editor/UMGEditor/Private/TabFactory/AnimationTabSummoner.cpp
--- a/Engine/Source/Editor/UMGEditor/Private/TabFactory/AnimationTabSummoner.cpp
+++ b/Engine/Source/Editor/UMGEditor/Private/TabFactory/AnimationTabSummoner.cpp
@@ -18,7 +18,7 @@ FAnimationTabSummoner::FAnimationTabSummoner(TSharedPtr<class FWidgetBlueprintEd
 		, BlueprintEditor(InBlueprintEditor)
 {
 	TabLabel = LOCTEXT("AnimationsTabLabel", "Animations");
-	TabIcon = FUMGStyle::Get().GetBrush("UMGEditor.Tabs.Animations");
+	TabIcon = FSlateIcon(FUMGStyle::GetStyleSetName(), "Animations.TabIcon");
 
 	bIsSingleton = true;
 
diff --git a/Engine/Source/Editor/UMGEditor/Private/TabFactory/DesignerTabSummoner.cpp b/Engine/Source/Editor/UMGEditor/Private/TabFactory/DesignerTabSummoner.cpp
--- a/Engine/Source/Editor/UMGEditor/Private/TabFactory/DesignerTabSummoner.cpp
+++ b/Engine/Source/Editor/UMGEditor/Private/TabFactory/DesignerTabSummoner.cpp
@@ -16,7 +16,7 @@ FDesignerTabSummoner::FDesignerTabSummoner(TSharedPtr<class FWidgetBlueprintEdit
 		, BlueprintEditor(InBlueprintEditor)
 {
 	TabLabel = LOCTEXT("DesignerTabLabel", "Designer");
-	TabIcon = FUMGStyle::Get().GetBrush("UMGEditor.Tabs.Designer");
+	TabIcon = FSlateIcon(FUMGStyle::GetStyleSetName(), "Designer.TabIcon");
 
 	bIsSingleton = true;
 
diff --git a/Engine/Source/Editor/UMGEditor/Private/TabFactory/DetailsTabSummoner.cpp b/Engine/Source/Editor/UMGEditor/Private/TabFactory/DetailsTabSummoner.cpp
--- a/Engine/Source/Editor/UMGEditor/Private/TabFactory/DetailsTabSummoner.cpp
+++ b/Engine/Source/Editor/UMGEditor/Private/TabFactory/DetailsTabSummoner.cpp
@@ -14,7 +14,7 @@ FDetailsTabSummoner::FDetailsTabSummoner(TSharedPtr<class FWidgetBlueprintEditor
 		, BlueprintEditor(InBlueprintEditor)
 {
 	TabLabel = LOCTEXT("WidgetDetails_TabLabel", "Details");
-	TabIcon = FEditorStyle::GetBrush("LevelEditor.Tabs.Details");
+	TabIcon = FSlateIcon(FEditorStyle::GetStyleSetName(), "LevelEditor.Tabs.Details");
 
 	bIsSingleton = true;
 
diff --git a/Engine/Source/Editor/UMGEditor/Private/TabFactory/HierarchyTabSummoner.cpp b/Engine/Source/Editor/UMGEditor/Private/TabFactory/HierarchyTabSummoner.cpp
--- a/Engine/Source/Editor/UMGEditor/Private/TabFactory/HierarchyTabSummoner.cpp
+++ b/Engine/Source/Editor/UMGEditor/Private/TabFactory/HierarchyTabSummoner.cpp
@@ -14,7 +14,7 @@ FHierarchyTabSummoner::FHierarchyTabSummoner(TSharedPtr<class FWidgetBlueprintEd
 		, BlueprintEditor(InBlueprintEditor)
 {
 	TabLabel = LOCTEXT("SlateHierarchyTabLabel", "Hierarchy");
-	TabIcon = FEditorStyle::Get().GetBrush("Kismet.Tabs.Palette");
+	TabIcon = FSlateIcon(FEditorStyle::GetStyleSetName(), "Kismet.Tabs.Palette");
 
 	bIsSingleton = true;
 
diff --git a/Engine/Source/Editor/UMGEditor/Private/TabFactory/PaletteTabSummoner.cpp b/Engine/Source/Editor/UMGEditor/Private/TabFactory/PaletteTabSummoner.cpp
--- a/Engine/Source/Editor/UMGEditor/Private/TabFactory/PaletteTabSummoner.cpp
+++ b/Engine/Source/Editor/UMGEditor/Private/TabFactory/PaletteTabSummoner.cpp
@@ -15,7 +15,7 @@ FPaletteTabSummoner::FPaletteTabSummoner(TSharedPtr<class FWidgetBlueprintEditor
 		, BlueprintEditor(InBlueprintEditor)
 {
 	TabLabel = LOCTEXT("WidgetTemplatesTabLabel", "Palette");
-	TabIcon = FUMGStyle::Get().GetBrush("UMGEditor.Tabs.Palette");
+	TabIcon = FSlateIcon(FUMGStyle::GetStyleSetName(), "Palette.TabIcon");
 
 	bIsSingleton = true;
 
diff --git a/Engine/Source/Editor/UMGEditor/Private/TabFactory/SequencerTabSummoner.cpp b/Engine/Source/Editor/UMGEditor/Private/TabFactory/SequencerTabSummoner.cpp
--- a/Engine/Source/Editor/UMGEditor/Private/TabFactory/SequencerTabSummoner.cpp
+++ b/Engine/Source/Editor/UMGEditor/Private/TabFactory/SequencerTabSummoner.cpp
@@ -15,7 +15,7 @@ FSequencerTabSummoner::FSequencerTabSummoner(TSharedPtr<class FWidgetBlueprintEd
 		, BlueprintEditor(InBlueprintEditor)
 {
 	TabLabel = LOCTEXT("SequencerLabel", "Timeline");
-	TabIcon = FUMGStyle::Get().GetBrush("UMGEditor.Tabs.Timeline");
+	TabIcon = FSlateIcon(FUMGStyle::GetStyleSetName(), "Sequencer.TabIcon");
 
 	bIsSingleton = true;
 
diff --git a/Engine/Source/Editor/UndoHistory/Private/UndoHistoryModule.cpp b/Engine/Source/Editor/UndoHistory/Private/UndoHistoryModule.cpp
--- a/Engine/Source/Editor/UndoHistory/Private/UndoHistoryModule.cpp
+++ b/Engine/Source/Editor/UndoHistory/Private/UndoHistoryModule.cpp
@@ -1,59 +1,44 @@
 // Copyright 1998-2013 Epic Games, Inc. All Rights Reserved.
 
 #include "UndoHistoryPrivatePCH.h"
+#include "UndoHistoryModule.h"
 #include "Editor/WorkspaceMenuStructure/Public/WorkspaceMenuStructureModule.h"
 
-
 #define LOCTEXT_NAMESPACE "FUndoHistoryModule"
 
-
-static const FName UndoHistoryTabName("UndoHistory");
-
-
-/**
- * Implements the UndoHistory module.
- */
-class FUndoHistoryModule
-	: public IUndoHistoryModule
+void FUndoHistoryModule::StartupModule()
 {
-public:
-
-	// Begin IModuleInterface interface
+	FGlobalTabmanager::Get()->RegisterTabSpawner(UndoHistoryTabName, FOnSpawnTab::CreateRaw(this, &FUndoHistoryModule::HandleSpawnSettingsTab))
+		.SetDisplayName(NSLOCTEXT("FUndoHistoryModule", "UndoHistoryTabTitle", "Undo History"))
+		.SetTooltipText(NSLOCTEXT("FUndoHistoryModule", "UndoHistoryTooltipText", "Open the Undo History tab."))
+		.SetIcon(FSlateIcon(FEditorStyle::GetStyleSetName(), "UndoHistory.TabIcon"));
+}
 
-	virtual void StartupModule( ) override
-	{
-		FGlobalTabmanager::Get()->RegisterNomadTabSpawner(UndoHistoryTabName, FOnSpawnTab::CreateRaw(this, &FUndoHistoryModule::HandleSpawnSettingsTab))
-			.SetDisplayName(NSLOCTEXT("FUndoHistoryModule", "UndoHistoryTabTitle", "Undo History"))
-			.SetGroup(WorkspaceMenu::GetMenuStructure().GetToolsCategory())
-			.SetTooltipText(NSLOCTEXT("FUndoHistoryModule", "UndoHistoryTooltipText", "Open the Undo History tab."))
-			.SetIcon(FSlateIcon(FEditorStyle::GetStyleSetName(), "UndoHistory.TabIcon"));
-	}
-
-	virtual void ShutdownModule( ) override
-	{
-		FGlobalTabmanager::Get()->UnregisterNomadTabSpawner(UndoHistoryTabName);
-	}
-
-	virtual bool SupportsDynamicReloading( ) override
-	{
-		return true;
-	}
+void FUndoHistoryModule::ShutdownModule()
+{
+	FGlobalTabmanager::Get()->UnregisterTabSpawner(UndoHistoryTabName);
+	//FGlobalTabmanager::Get()->UnregisterNomadTabSpawner(UndoHistoryTabName);
+}
 
-	// End IModuleInterface interface
+bool FUndoHistoryModule::SupportsDynamicReloading()
+{
+	return true;
+}
 
-private:
+//void FUndoHistoryModule::ExecuteOpenUndoHistory() /*const*/
+//{
+//	FGlobalTabmanager::Get()->InvokeTab(UndoHistoryTabName);
+//}
 
-	// Handles creating the project settings tab.
-	TSharedRef<SDockTab> HandleSpawnSettingsTab( const FSpawnTabArgs& SpawnTabArgs )
-	{
-		const TSharedRef<SDockTab> DockTab = SNew(SDockTab)
-			.TabRole(ETabRole::NomadTab);
+TSharedRef<SDockTab> FUndoHistoryModule::HandleSpawnSettingsTab(const FSpawnTabArgs& SpawnTabArgs)
+{
+	const TSharedRef<SDockTab> DockTab = SNew(SDockTab)
+		.TabRole(ETabRole::NomadTab);
 
-		DockTab->SetContent(SNew(SUndoHistory));
+	DockTab->SetContent(SNew(SUndoHistory));
 
-		return DockTab;
-	}
-};
+	return DockTab;
+}
 
 
 IMPLEMENT_MODULE(FUndoHistoryModule, UndoHistory);
diff --git a/Engine/Source/Editor/UndoHistory/Public/UndoHistoryModule.h b/Engine/Source/Editor/UndoHistory/Public/UndoHistoryModule.h
new file mode 100644
--- /dev/null
+++ b/Engine/Source/Editor/UndoHistory/Public/UndoHistoryModule.h
@@ -0,0 +1,33 @@
+// Copyright 1998-2013 Epic Games, Inc. All Rights Reserved.
+
+#pragma once
+
+#include "IUndoHistoryModule.h"
+
+static const FName UndoHistoryTabName("UndoHistory");
+
+/**
+* Implements the UndoHistory module.
+*/
+class FUndoHistoryModule : public IUndoHistoryModule
+{
+public:
+
+	// Begin IModuleInterface interface
+
+	virtual void StartupModule() override;
+	virtual void ShutdownModule() override;
+	virtual bool SupportsDynamicReloading() override;
+
+	// End IModuleInterface interface
+
+	//static /*virtual */void ExecuteOpenUndoHistory() /*const override*/;
+	static void ExecuteOpenUndoHistory() /*const*/
+	{
+		FGlobalTabmanager::Get()->InvokeTab(UndoHistoryTabName);
+	}
+
+private:
+	// Handles creating the project settings tab.
+	TSharedRef<SDockTab> HandleSpawnSettingsTab(const FSpawnTabArgs& SpawnTabArgs);
+};
\ No newline at end of file
diff --git a/Engine/Source/Editor/UnrealEd/Classes/Settings/EditorExperimentalSettings.h b/Engine/Source/Editor/UnrealEd/Classes/Settings/EditorExperimentalSettings.h
--- a/Engine/Source/Editor/UnrealEd/Classes/Settings/EditorExperimentalSettings.h
+++ b/Engine/Source/Editor/UnrealEd/Classes/Settings/EditorExperimentalSettings.h
@@ -29,9 +29,9 @@ public:
 	UPROPERTY(EditAnywhere, config, Category=Tools, meta=(DisplayName="Editor Utility Blueprints (Blutility)"))
 	bool bEnableEditorUtilityBlueprints;
 
-	/** The Game Launcher provides advanced workflows for packaging, deploying and launching your games. */
+	/** The Project Launcher provides advanced workflows for packaging, deploying and launching your projects. */
 	UPROPERTY(EditAnywhere, config, Category=Tools, meta=(DisplayName="Project Launcher"))
-	bool bGameLauncher;
+	bool bProjectLauncher;
 
 	/** The Messaging Debugger provides a visual utility for debugging the messaging system. */
 	UPROPERTY(EditAnywhere, config, Category=Tools, meta=(DisplayName="Messaging Debugger"))
diff --git a/Engine/Source/Editor/UnrealEd/Private/AnimationEditorUtils.cpp b/Engine/Source/Editor/UnrealEd/Private/AnimationEditorUtils.cpp
--- a/Engine/Source/Editor/UnrealEd/Private/AnimationEditorUtils.cpp
+++ b/Engine/Source/Editor/UnrealEd/Private/AnimationEditorUtils.cpp
@@ -265,6 +265,62 @@ namespace AnimationEditorUtils
 		}
 	}
 
+	void CreateNewAnimBlueprint(TArray<TWeakObjectPtr<USkeleton>> Skeletons, FAnimAssetCreated AssetCreated)
+	{
+		const FString DefaultSuffix = TEXT("_AnimBlueprint");
+
+		if (Skeletons.Num() == 1)
+		{
+			auto Object = Skeletons[0].Get();
+
+			if (Object)
+			{
+				// Determine an appropriate name for inline-rename
+				FString Name;
+				FString PackageName;
+				CreateUniqueAssetName(Object->GetOutermost()->GetName(), DefaultSuffix, PackageName, Name);
+
+				UAnimBlueprintFactory* Factory = ConstructObject<UAnimBlueprintFactory>(UAnimBlueprintFactory::StaticClass());
+				Factory->TargetSkeleton = Object;
+
+				FContentBrowserModule& ContentBrowserModule = FModuleManager::LoadModuleChecked<FContentBrowserModule>("ContentBrowser");
+				ContentBrowserModule.Get().CreateNewAsset(Name, FPackageName::GetLongPackagePath(PackageName), UAnimBlueprint::StaticClass(), Factory);
+			}
+		}
+		else
+		{
+			TArray<UObject*> SkeletonsToSync;
+			for (auto ObjIt = Skeletons.CreateConstIterator(); ObjIt; ++ObjIt)
+			{
+				auto Object = (*ObjIt).Get();
+				if (Object)
+				{
+					// Determine an appropriate name
+					FString Name;
+					FString PackageName;
+					CreateUniqueAssetName(Object->GetOutermost()->GetName(), DefaultSuffix, PackageName, Name);
+
+					// Create the anim blueprint factory used to generate the asset
+					UAnimBlueprintFactory* Factory = ConstructObject<UAnimBlueprintFactory>(UAnimBlueprintFactory::StaticClass());
+					Factory->TargetSkeleton = Object;
+
+					FAssetToolsModule& AssetToolsModule = FModuleManager::GetModuleChecked<FAssetToolsModule>("AssetTools");
+					UObject* NewAsset = AssetToolsModule.Get().CreateAsset(Name, FPackageName::GetLongPackagePath(PackageName), UAnimBlueprint::StaticClass(), Factory);
+
+					if (NewAsset)
+					{
+						SkeletonsToSync.Add(NewAsset);
+					}
+				}
+			}
+
+			if (AssetCreated.IsBound())
+			{
+				AssetCreated.Execute(SkeletonsToSync);
+			}
+		}
+	}
+
 	template <typename TFactory, typename T>
 	void ExecuteNewAnimAsset(TArray<TWeakObjectPtr<USkeleton>> Objects, const FString InSuffix, FAnimAssetCreated AssetCreated, bool bInContentBrowser )
 	{
@@ -306,85 +362,89 @@ namespace AnimationEditorUtils
 		}
 	}
 
-	/** Handler for the blend space sub menu */
-	void FillBlendSpaceMenu(FMenuBuilder& MenuBuilder, TArray<TWeakObjectPtr<USkeleton>> Skeletons, FAnimAssetCreated AssetCreated, bool bInContentBrowser)
+	void FillCreateAssetMenu(FMenuBuilder& MenuBuilder, TArray<TWeakObjectPtr<USkeleton>> Skeletons, FAnimAssetCreated AssetCreated, bool bInContentBrowser) 
 	{
-		MenuBuilder.AddMenuEntry(
-			LOCTEXT("SkeletalMesh_New1DBlendspace", "Create 1D BlendSpace"),
-			LOCTEXT("SkeletalMesh_New1DBlendspaceTooltip", "Creates a 1D blendspace using the skeleton of the selected mesh."),
-			FSlateIcon(),
-			FUIAction(
-				FExecuteAction::CreateStatic(&ExecuteNewAnimAsset<UBlendSpaceFactory1D, UBlendSpace1D>, Skeletons, FString("_BlendSpace1D"), AssetCreated, bInContentBrowser),
-				FCanExecuteAction()
-				)
-			);
-
-		MenuBuilder.AddMenuEntry(
-			LOCTEXT("SkeletalMesh_New2DBlendspace", "Create 2D BlendSpace"),
-			LOCTEXT("SkeletalMesh_New2DBlendspaceTooltip", "Creates a 2D blendspace using the skeleton of the selected mesh."),
-			FSlateIcon(),
-			FUIAction(
-				FExecuteAction::CreateStatic(&ExecuteNewAnimAsset<UBlendSpaceFactoryNew, UBlendSpace>, Skeletons, FString("_BlendSpace2D"), AssetCreated, bInContentBrowser),
-				FCanExecuteAction()
-				)
-			);
-	}
+		MenuBuilder.BeginSection("CreateAnimAssets", LOCTEXT("CreateAnimAssetsMenuHeading", "Anim Assets"));
+		{
+			MenuBuilder.AddMenuEntry(
+				LOCTEXT("Skeleton_NewAnimBlueprint", "Anim Blueprint"),
+				LOCTEXT("Skeleton_NewAnimBlueprintTooltip", "Creates an Anim Blueprint using the selected skeleton."),
+				FSlateIcon(FEditorStyle::GetStyleSetName(), "ClassIcon.AnimBlueprint"),
+				FUIAction(
+					FExecuteAction::CreateStatic(&CreateNewAnimBlueprint, Skeletons, AssetCreated),
+					FCanExecuteAction()
+					)
+				);
+
+			MenuBuilder.AddMenuEntry(
+				LOCTEXT("Skeleton_NewAnimComposite", "Anim Composite"),
+				LOCTEXT("Skeleton_NewAnimCompositeTooltip", "Creates an AnimComposite using the selected skeleton."),
+				FSlateIcon(FEditorStyle::GetStyleSetName(), "ClassIcon.AnimComposite"),
+				FUIAction(
+					FExecuteAction::CreateStatic(&ExecuteNewAnimAsset<UAnimCompositeFactory, UAnimComposite>, Skeletons, FString("_Composite"), AssetCreated, bInContentBrowser),
+					FCanExecuteAction()
+					)
+				);
+
+			MenuBuilder.AddMenuEntry(
+				LOCTEXT("Skeleton_NewAnimMontage", "Anim Montage"),
+				LOCTEXT("Skeleton_NewAnimMontageTooltip", "Creates an AnimMontage using the selected skeleton."),
+				FSlateIcon(FEditorStyle::GetStyleSetName(), "ClassIcon.AnimMontage"),
+				FUIAction(
+					FExecuteAction::CreateStatic(&ExecuteNewAnimAsset<UAnimMontageFactory, UAnimMontage>, Skeletons, FString("_Montage"), AssetCreated, bInContentBrowser),
+					FCanExecuteAction()
+					)
+				);
+		}
+		MenuBuilder.EndSection();
 
-	/** Handler for the blend space sub menu */
-	void FillAimOffsetBlendSpaceMenu(FMenuBuilder& MenuBuilder, TArray<TWeakObjectPtr<USkeleton>> Skeletons, FAnimAssetCreated AssetCreated, bool bInContentBrowser)
-	{
-		MenuBuilder.AddMenuEntry(
-			LOCTEXT("SkeletalMesh_New1DAimOffset", "Create 1D AimOffset"),
-			LOCTEXT("SkeletalMesh_New1DAimOffsetTooltip", "Creates a 1D aimoffset blendspace using the selected skeleton."),
-			FSlateIcon(),
-			FUIAction(
-				FExecuteAction::CreateStatic(&ExecuteNewAnimAsset<UAimOffsetBlendSpaceFactory1D, UAimOffsetBlendSpace1D>, Skeletons, FString("_AimOffset1D"), AssetCreated, bInContentBrowser),
-				FCanExecuteAction()
-				)
-			);
-
-		MenuBuilder.AddMenuEntry(
-			LOCTEXT("SkeletalMesh_New2DAimOffset", "Create 2D AimOffset"),
-			LOCTEXT("SkeletalMesh_New2DAimOffsetTooltip", "Creates a 2D aimoffset blendspace using the selected skeleton."),
-			FSlateIcon(),
-			FUIAction(
-				FExecuteAction::CreateStatic(&ExecuteNewAnimAsset<UAimOffsetBlendSpaceFactoryNew, UAimOffsetBlendSpace>, Skeletons, FString("_AimOffset2D"), AssetCreated, bInContentBrowser),
-				FCanExecuteAction()
-				)
-			);
-	}
+		MenuBuilder.BeginSection("CreateBlendSpace", LOCTEXT("CreateBlendSpaceMenuHeading", "Blend Spaces"));
+		{
+			MenuBuilder.AddMenuEntry(
+				LOCTEXT("SkeletalMesh_New1DBlendspace", "1D Blend Space"),
+				LOCTEXT("SkeletalMesh_New1DBlendspaceTooltip", "Creates a 1D Blend Space using the selected skeleton."),
+				FSlateIcon(FEditorStyle::GetStyleSetName(), "ClassIcon.BlendSpace1D"),
+				FUIAction(
+					FExecuteAction::CreateStatic(&ExecuteNewAnimAsset<UBlendSpaceFactory1D, UBlendSpace1D>, Skeletons, FString("_BlendSpace1D"), AssetCreated, bInContentBrowser),
+					FCanExecuteAction()
+					)
+				);
+
+			MenuBuilder.AddMenuEntry(
+				LOCTEXT("SkeletalMesh_New2DBlendspace", "2D Blend Space"),
+				LOCTEXT("SkeletalMesh_New2DBlendspaceTooltip", "Creates a 2D Blend Space using the selected skeleton."),
+				FSlateIcon(FEditorStyle::GetStyleSetName(), "ClassIcon.BlendSpace"),
+				FUIAction(
+					FExecuteAction::CreateStatic(&ExecuteNewAnimAsset<UBlendSpaceFactoryNew, UBlendSpace>, Skeletons, FString("_BlendSpace2D"), AssetCreated, bInContentBrowser),
+					FCanExecuteAction()
+					)
+				);
+		}
+		MenuBuilder.EndSection();
 
-	void FillCreateAssetMenu(FMenuBuilder& MenuBuilder, TArray<TWeakObjectPtr<USkeleton>> Skeletons, FAnimAssetCreated AssetCreated, bool bInContentBrowser) 
-	{
-		MenuBuilder.AddSubMenu(
-			LOCTEXT("SkeletalMesh_NewAimOffset", "Create AimOffset"),
-			LOCTEXT("SkeletalMesh_NewAimOffsetTooltip", "Creates an aimoffset blendspace using the selected skeleton."),
-			FNewMenuDelegate::CreateStatic(&FillAimOffsetBlendSpaceMenu, Skeletons, AssetCreated, bInContentBrowser));
-
-		MenuBuilder.AddSubMenu(
-			LOCTEXT("SkeletalMesh_NewBlendspace", "Create BlendSpace"),
-			LOCTEXT("SkeletalMesh_NewBlendspaceTooltip", "Creates a blendspace using the skeleton of the selected mesh."),
-			FNewMenuDelegate::CreateStatic(&FillBlendSpaceMenu, Skeletons, AssetCreated, bInContentBrowser));
-
-		MenuBuilder.AddMenuEntry(
-			LOCTEXT("Skeleton_NewAnimComposite", "Create AnimComposite"),
-			LOCTEXT("Skeleton_NewAnimCompositeTooltip", "Creates an AnimComposite using the selected skeleton."),
-			FSlateIcon(),
-			FUIAction(
-				FExecuteAction::CreateStatic(&ExecuteNewAnimAsset<UAnimCompositeFactory, UAnimComposite>, Skeletons, FString("_Composite"), AssetCreated, bInContentBrowser),
-				FCanExecuteAction()
-				)
-			);
-
-		MenuBuilder.AddMenuEntry(
-			LOCTEXT("Skeleton_NewAnimMontage", "Create AnimMontage"),
-			LOCTEXT("Skeleton_NewAnimMontageTooltip", "Creates an AnimMontage using the selected skeleton."),
-			FSlateIcon(),
-			FUIAction(
-				FExecuteAction::CreateStatic(&ExecuteNewAnimAsset<UAnimMontageFactory, UAnimMontage>, Skeletons, FString("_Montage"), AssetCreated, bInContentBrowser),
-				FCanExecuteAction()
-				)
-			);
+		MenuBuilder.BeginSection("CreateAimOffset", LOCTEXT("CreateAimOffsetMenuHeading", "Aim Offsets"));
+		{
+			MenuBuilder.AddMenuEntry(
+				LOCTEXT("SkeletalMesh_New1DAimOffset", "1D Aim Offset"),
+				LOCTEXT("SkeletalMesh_New1DAimOffsetTooltip", "Creates a 1D Aim Offset blendspace using the selected skeleton."),
+				FSlateIcon(),
+				FUIAction(
+					FExecuteAction::CreateStatic(&ExecuteNewAnimAsset<UAimOffsetBlendSpaceFactory1D, UAimOffsetBlendSpace1D>, Skeletons, FString("_AimOffset1D"), AssetCreated, bInContentBrowser),
+					FCanExecuteAction()
+					)
+				);
+
+			MenuBuilder.AddMenuEntry(
+				LOCTEXT("SkeletalMesh_New2DAimOffset", "2D Aim Offset"),
+				LOCTEXT("SkeletalMesh_New2DAimOffsetTooltip", "Creates a 2D Aim Offset blendspace using the selected skeleton."),
+				FSlateIcon(),
+				FUIAction(
+					FExecuteAction::CreateStatic(&ExecuteNewAnimAsset<UAimOffsetBlendSpaceFactoryNew, UAimOffsetBlendSpace>, Skeletons, FString("_AimOffset2D"), AssetCreated, bInContentBrowser),
+					FCanExecuteAction()
+					)
+				);
+		}
+		MenuBuilder.EndSection();
 	}
 }
 
diff --git a/Engine/Source/Editor/UnrealEd/Private/Kismet2/DebuggerCommands.cpp b/Engine/Source/Editor/UnrealEd/Private/Kismet2/DebuggerCommands.cpp
--- a/Engine/Source/Editor/UnrealEd/Private/Kismet2/DebuggerCommands.cpp
+++ b/Engine/Source/Editor/UnrealEd/Private/Kismet2/DebuggerCommands.cpp
@@ -228,7 +228,7 @@ void FPlayWorldCommands::RegisterCommands()
 
 	// Launch
 	UI_COMMAND( RepeatLastLaunch, "Launch", "Launches the game on the device as the last session launched from the dropdown next to the Play on Device button on the level editor toolbar", EUserInterfaceActionType::Button, FInputGesture( EKeys::P, EModifierKey::Alt | EModifierKey::Shift ) )
-	UI_COMMAND( OpenDeviceManager, "Device Manager...", "Opens the device manager", EUserInterfaceActionType::Button, FInputGesture() );
+	UI_COMMAND( OpenDeviceManager, "Device Manager...", "View and manage connected devices.", EUserInterfaceActionType::Button, FInputGesture() );
 }
 
 
@@ -764,7 +764,12 @@ TSharedRef< SWidget > FPlayWorldCommands::GenerateLaunchMenuContent( TSharedRef<
 	// options section
 	MenuBuilder.BeginSection("LevelEditorLaunchOptions");
 	{
-		MenuBuilder.AddMenuEntry( FPlayWorldCommands::Get().OpenDeviceManager );
+		MenuBuilder.AddMenuEntry( FPlayWorldCommands::Get().OpenDeviceManager,
+			NAME_None,
+			TAttribute<FText>(),
+			TAttribute<FText>(),
+			FSlateIcon(FEditorStyle::GetStyleSetName(), "DeviceDetails.TabIcon")
+			);
 		
 		ProjectTargetPlatformEditorModule.AddOpenProjectTargetPlatformEditorMenuItem(MenuBuilder);
 	}
diff --git a/Engine/Source/Editor/UnrealEd/Private/Toolkits/AssetEditorToolkit.cpp b/Engine/Source/Editor/UnrealEd/Private/Toolkits/AssetEditorToolkit.cpp
--- a/Engine/Source/Editor/UnrealEd/Private/Toolkits/AssetEditorToolkit.cpp
+++ b/Engine/Source/Editor/UnrealEd/Private/Toolkits/AssetEditorToolkit.cpp
@@ -25,6 +25,7 @@ const FName FAssetEditorToolkit::ToolbarTabId( TEXT( "AssetEditorToolkit_Toolbar
 FAssetEditorToolkit::FAssetEditorToolkit()
 	: bIsToolbarFocusable(false)
 {
+	WorkspaceMenuCategory = FWorkspaceItem::NewGroup(LOCTEXT("WorkspaceMenu_BaseAssetEditor", "Asset Editor"));
 }
 
 
@@ -236,11 +237,14 @@ FAssetEditorToolkit::~FAssetEditorToolkit()
 
 void FAssetEditorToolkit::RegisterTabSpawners(const TSharedRef<class FTabManager>& InTabManager)
 {
-	const IWorkspaceMenuStructure& MenuStructure = WorkspaceMenu::GetMenuStructure();
+	// Use the first child of the local workspace root if there is one, otherwise use the root itself
+	const auto& LocalCategories = InTabManager->GetLocalWorkspaceMenuRoot()->GetChildItems();
+	TSharedRef<FWorkspaceItem> ToolbarSpawnerCategory = LocalCategories.Num() > 0 ? LocalCategories[0] : InTabManager->GetLocalWorkspaceMenuRoot();
 
 	InTabManager->RegisterTabSpawner( ToolbarTabId, FOnSpawnTab::CreateSP(this, &FAssetEditorToolkit::SpawnTab_Toolbar) )
 		.SetDisplayName( LOCTEXT("ToolbarTab", "Toolbar") )
-		.SetGroup( MenuStructure.GetAssetEditorCategory() );
+		.SetGroup( ToolbarSpawnerCategory )
+		.SetIcon( FSlateIcon(FEditorStyle::GetStyleSetName(), "Toolbar.Icon") );
 }
 
 void FAssetEditorToolkit::UnregisterTabSpawners(const TSharedRef<class FTabManager>& InTabManager)
diff --git a/Engine/Source/Editor/UnrealEd/Private/Toolkits/BaseToolkit.cpp b/Engine/Source/Editor/UnrealEd/Private/Toolkits/BaseToolkit.cpp
--- a/Engine/Source/Editor/UnrealEd/Private/Toolkits/BaseToolkit.cpp
+++ b/Engine/Source/Editor/UnrealEd/Private/Toolkits/BaseToolkit.cpp
@@ -5,11 +5,13 @@
 #include "Toolkits/IToolkitHost.h"
 #include "Toolkits/ToolkitManager.h"
 
+#define LOCTEXT_NAMESPACE "BaseToolkit"
 
 FBaseToolkit::FBaseToolkit()
 	: ToolkitMode( EToolkitMode::Standalone ),
 	  ToolkitCommands( new FUICommandList() )
 {
+	WorkspaceMenuCategory = FWorkspaceItem::NewGroup(LOCTEXT("WorkspaceMenu_BaseToolkit", "Toolkit"));
 }
 
 
@@ -196,6 +198,8 @@ bool FBaseToolkit::IsBlueprintEditor() const
 	return false;
 }
 
+#undef LOCTEXT_NAMESPACE
+
 
 void FModeToolkit::Init(const TSharedPtr< class IToolkitHost >& InitToolkitHost)
 {
diff --git a/Engine/Source/Editor/UnrealEd/Private/Toolkits/SimpleAssetEditor.cpp b/Engine/Source/Editor/UnrealEd/Private/Toolkits/SimpleAssetEditor.cpp
--- a/Engine/Source/Editor/UnrealEd/Private/Toolkits/SimpleAssetEditor.cpp
+++ b/Engine/Source/Editor/UnrealEd/Private/Toolkits/SimpleAssetEditor.cpp
@@ -15,13 +15,14 @@ const FName FSimpleAssetEditor::PropertiesTabId( TEXT( "GenericEditor_Properties
 
 void FSimpleAssetEditor::RegisterTabSpawners(const TSharedRef<class FTabManager>& TabManager)
 {
-	FAssetEditorToolkit::RegisterTabSpawners(TabManager);
+	WorkspaceMenuCategory = TabManager->GetLocalWorkspaceMenuRoot()->AddGroup(LOCTEXT("WorkspaceMenu_GenericAssetEditor", "Asset Editor"));
 
-	const IWorkspaceMenuStructure& MenuStructure = WorkspaceMenu::GetMenuStructure();
+	FAssetEditorToolkit::RegisterTabSpawners(TabManager);
 
 	TabManager->RegisterTabSpawner( PropertiesTabId, FOnSpawnTab::CreateSP(this, &FSimpleAssetEditor::SpawnPropertiesTab) )
 		.SetDisplayName( LOCTEXT("PropertiesTab", "Details") )
-		.SetGroup( MenuStructure.GetAssetEditorCategory() );
+		.SetGroup(WorkspaceMenuCategory.ToSharedRef())
+		.SetIcon(FSlateIcon(FEditorStyle::GetStyleSetName(), "LevelEditor.Tabs.Details"));
 }
 
 void FSimpleAssetEditor::UnregisterTabSpawners(const TSharedRef<class FTabManager>& TabManager)
diff --git a/Engine/Source/Editor/UnrealEd/Public/AnimationEditorUtils.h b/Engine/Source/Editor/UnrealEd/Public/AnimationEditorUtils.h
--- a/Engine/Source/Editor/UnrealEd/Public/AnimationEditorUtils.h
+++ b/Engine/Source/Editor/UnrealEd/Public/AnimationEditorUtils.h
@@ -54,7 +54,9 @@ DECLARE_DELEGATE_OneParam(FAnimAssetCreated, TArray<class UObject*>);
 namespace AnimationEditorUtils
 {
 	UNREALED_API void CreateAnimationAssets(const TArray<TWeakObjectPtr<USkeleton>>& Skeletons, TSubclassOf<UAnimationAsset> AssetClass, const FString& InPrefix, FAnimAssetCreated AssetCreated );
+	UNREALED_API void CreateNewAnimBlueprint(TArray<TWeakObjectPtr<USkeleton>> Skeletons, FAnimAssetCreated AssetCreated);
 	UNREALED_API void FillCreateAssetMenu(FMenuBuilder& MenuBuilder, TArray<TWeakObjectPtr<USkeleton>> Skeletons, FAnimAssetCreated AssetCreated, bool bInContentBrowser=true);
+
 } // namespace AnimationEditorUtils
 
 #endif //__AnimationEditorUtils_h__
diff --git a/Engine/Source/Editor/UnrealEd/Public/Toolkits/BaseToolkit.h b/Engine/Source/Editor/UnrealEd/Public/Toolkits/BaseToolkit.h
--- a/Engine/Source/Editor/UnrealEd/Public/Toolkits/BaseToolkit.h
+++ b/Engine/Source/Editor/UnrealEd/Public/Toolkits/BaseToolkit.h
@@ -28,6 +28,7 @@ public:
 	virtual void BringToolkitToFront() override;
 	virtual TSharedPtr<class SWidget> GetInlineContent() const override;
 	virtual bool IsBlueprintEditor() const override;
+	virtual TSharedRef<FWorkspaceItem> GetWorkspaceMenuCategory() const override { return WorkspaceMenuCategory.ToSharedRef(); }
 
 	/** @return	Returns true if this is a world-centric asset editor.  That is, the user is editing the asset inline in a Level Editor app. */
 	bool IsWorldCentricAssetEditor() const;	
@@ -71,6 +72,9 @@ protected:
 
 	/** Map of toolkit tab spots to known tabs (these are weak pointers and may be invalid after tabs are closed.) */
 	TMap< EToolkitTabSpot::Type, TArray< TWeakPtr< SDockableTab > > > ToolkitTabsInSpots;
+
+	/** The workspace menu category of this toolkit */
+	TSharedPtr<FWorkspaceItem> WorkspaceMenuCategory;
 };
 
 
diff --git a/Engine/Source/Editor/UnrealEd/Public/Toolkits/IToolkit.h b/Engine/Source/Editor/UnrealEd/Public/Toolkits/IToolkit.h
--- a/Engine/Source/Editor/UnrealEd/Public/Toolkits/IToolkit.h
+++ b/Engine/Source/Editor/UnrealEd/Public/Toolkits/IToolkit.h
@@ -72,6 +72,9 @@ public:
 	/** Returns the localized name of this toolkit */
 	virtual FText GetToolkitName() const = 0;
 
+	/** Returns the workspace menu category of this toolkit */
+	virtual TSharedRef<FWorkspaceItem> GetWorkspaceMenuCategory() const = 0;
+
 	/** Returns the localize prefix string to use for tab labels in world-centric mode. */
 	virtual FString GetWorldCentricTabPrefix() const = 0;
 
diff --git a/Engine/Source/Editor/WorkspaceMenuStructure/Private/WorkspaceMenuStructureModule.cpp b/Engine/Source/Editor/WorkspaceMenuStructure/Private/WorkspaceMenuStructureModule.cpp
--- a/Engine/Source/Editor/WorkspaceMenuStructure/Private/WorkspaceMenuStructureModule.cpp
+++ b/Engine/Source/Editor/WorkspaceMenuStructure/Private/WorkspaceMenuStructureModule.cpp
@@ -21,11 +21,6 @@ public:
 		return MenuRoot.ToSharedRef();
 	}
 
-	virtual TSharedRef<FWorkspaceItem> GetAssetEditorCategory() const override
-	{
-		return AssetEditorCategory.ToSharedRef();
-	}
-
 	virtual TSharedRef<FWorkspaceItem> GetLevelEditorCategory() const override
 	{
 		return LevelEditorCategory.ToSharedRef();
@@ -51,19 +46,24 @@ public:
 		return ToolsCategory.ToSharedRef();
 	}
 
-	virtual TSharedRef<FWorkspaceItem> GetDeveloperToolsCategory() const override
+	virtual TSharedRef<FWorkspaceItem> GetDeveloperToolsDebugCategory() const override
 	{
-		return DeveloperToolsCategory.ToSharedRef();
+		return DeveloperToolsDebugCategory.ToSharedRef();
 	}
 
-	virtual TSharedRef<FWorkspaceItem> GetEditOptions() const override
+	virtual TSharedRef<FWorkspaceItem> GetDeveloperToolsLogCategory() const override
 	{
-		return EditOptions.ToSharedRef();
+		return DeveloperToolsLogCategory.ToSharedRef();
 	}
 
-	void ResetAssetEditorCategory()
+	virtual TSharedRef<FWorkspaceItem> GetDeveloperToolsMiscCategory() const override
 	{
-		AssetEditorCategory->ClearItems();
+		return DeveloperToolsMiscCategory.ToSharedRef();
+	}
+
+	virtual TSharedRef<FWorkspaceItem> GetEditOptions() const override
+	{
+		return EditOptions.ToSharedRef();
 	}
 
 	void ResetLevelEditorCategory()
@@ -77,18 +77,23 @@ public:
 	void ResetToolsCategory()
 	{
 		ToolsCategory->ClearItems();
-		DeveloperToolsCategory = ToolsCategory->AddGroup(LOCTEXT( "WorkspaceMenu_DeveloperToolsCategory", "Developer Tools" ), FSlateIcon(FEditorStyle::GetStyleSetName(), "DeveloperTools.MenuIcon"), true);
+
+		// Developer tools sub menu
+		DeveloperToolsCategory = ToolsCategory->AddGroup(LOCTEXT("WorkspaceMenu_DeveloperToolsCategory", "Developer Tools"), FSlateIcon(FEditorStyle::GetStyleSetName(), "DeveloperTools.MenuIcon"));
+
+		// Developer tools sections
+		DeveloperToolsDebugCategory = DeveloperToolsCategory->AddGroup(LOCTEXT("WorkspaceMenu_DeveloperToolsDebugCategory", "Debug"), FSlateIcon(), true);
+		DeveloperToolsLogCategory = DeveloperToolsCategory->AddGroup(LOCTEXT("WorkspaceMenu_DeveloperToolsLogCategory", "Log"), FSlateIcon(), true);
+		DeveloperToolsMiscCategory = DeveloperToolsCategory->AddGroup(LOCTEXT("WorkspaceMenu_DeveloperToolsMiscCategory", "Miscellaneous"), FSlateIcon(), true);
 	}
 
 public:
 	FWorkspaceMenuStructure()
 		: MenuRoot ( FWorkspaceItem::NewGroup(LOCTEXT( "WorkspaceMenu_Root", "Menu Root" )) )
-		, AssetEditorCategory ( MenuRoot->AddGroup(LOCTEXT( "WorkspaceMenu_AssetEditorCategory", "Asset Editor Tabs" )) )
-		, LevelEditorCategory ( MenuRoot->AddGroup(LOCTEXT( "WorkspaceMenu_LevelEditorCategory", "Level Editor Tabs" ), FSlateIcon(), true) )
-		, ToolsCategory ( MenuRoot->AddGroup(LOCTEXT( "WorkspaceMenu_ToolsCategory", "Application Windows" ), FSlateIcon(), true) )
-		, EditOptions ( FWorkspaceItem::NewGroup(LOCTEXT( "WorkspaceEdit_Options", "Edit Options" )) )
+		, LevelEditorCategory ( MenuRoot->AddGroup(LOCTEXT( "WorkspaceMenu_LevelEditorCategory", "Level Editor" ), FSlateIcon(), true) )
+		, ToolsCategory ( MenuRoot->AddGroup(LOCTEXT( "WorkspaceMenu_ToolsCategory", "General" ), FSlateIcon(), true) )
+		, EditOptions( FWorkspaceItem::NewGroup(LOCTEXT( "WorkspaceEdit_Options", "Edit Options" )) )
 	{
-		ResetAssetEditorCategory();
 		ResetLevelEditorCategory();
 		ResetToolsCategory();
 	}
@@ -97,13 +102,18 @@ public:
 
 private:
 	TSharedPtr<FWorkspaceItem> MenuRoot;
-	TSharedPtr<FWorkspaceItem> AssetEditorCategory;
+	
 	TSharedPtr<FWorkspaceItem> LevelEditorCategory;
 	TSharedPtr<FWorkspaceItem> LevelEditorViewportsCategory;
 	TSharedPtr<FWorkspaceItem> LevelEditorDetailsCategory;
 	TSharedPtr<FWorkspaceItem> LevelEditorModesCategory;
+
 	TSharedPtr<FWorkspaceItem> ToolsCategory;
 	TSharedPtr<FWorkspaceItem> DeveloperToolsCategory;
+	TSharedPtr<FWorkspaceItem> DeveloperToolsDebugCategory;
+	TSharedPtr<FWorkspaceItem> DeveloperToolsLogCategory;
+	TSharedPtr<FWorkspaceItem> DeveloperToolsMiscCategory;
+	
 	TSharedPtr<FWorkspaceItem> EditOptions;
 };
 
@@ -123,12 +133,6 @@ const IWorkspaceMenuStructure& FWorkspaceMenuStructureModule::GetWorkspaceMenuSt
 	return *WorkspaceMenuStructure;
 }
 
-void FWorkspaceMenuStructureModule::ResetAssetEditorCategory()
-{
-	check(WorkspaceMenuStructure.IsValid());
-	WorkspaceMenuStructure->ResetAssetEditorCategory();
-}
-
 void FWorkspaceMenuStructureModule::ResetLevelEditorCategory()
 {
 	check(WorkspaceMenuStructure.IsValid());
diff --git a/Engine/Source/Editor/WorkspaceMenuStructure/Public/WorkspaceMenuStructure.h b/Engine/Source/Editor/WorkspaceMenuStructure/Public/WorkspaceMenuStructure.h
--- a/Engine/Source/Editor/WorkspaceMenuStructure/Public/WorkspaceMenuStructure.h
+++ b/Engine/Source/Editor/WorkspaceMenuStructure/Public/WorkspaceMenuStructure.h
@@ -27,9 +27,6 @@ public:
 	/** Get the root of the menu structure. Pass this into PopulateTabSpawnerMenu() */
 	virtual TSharedRef<FWorkspaceItem> GetStructureRoot() const = 0;
 
-	/** See diagram above */
-	virtual TSharedRef<FWorkspaceItem> GetAssetEditorCategory() const = 0;
-
 	/** See diagram above */
 	virtual TSharedRef<FWorkspaceItem> GetLevelEditorCategory() const = 0;
 
@@ -46,7 +43,13 @@ public:
 	virtual TSharedRef<FWorkspaceItem> GetToolsCategory() const = 0;
 
 	/** See diagram above */
-	virtual TSharedRef<FWorkspaceItem> GetDeveloperToolsCategory() const = 0;
+	virtual TSharedRef<FWorkspaceItem> GetDeveloperToolsDebugCategory() const = 0;
+
+	/** See diagram above */
+	virtual TSharedRef<FWorkspaceItem> GetDeveloperToolsLogCategory() const = 0;
+
+	/** See diagram above */
+	virtual TSharedRef<FWorkspaceItem> GetDeveloperToolsMiscCategory() const = 0;
 
 	/** Get the root of the edit menu structure. Pass this into PopulateTabSpawnerMenu() */
 	virtual TSharedRef<FWorkspaceItem> GetEditOptions() const = 0;
diff --git a/Engine/Source/Editor/WorkspaceMenuStructure/Public/WorkspaceMenuStructureModule.h b/Engine/Source/Editor/WorkspaceMenuStructure/Public/WorkspaceMenuStructureModule.h
--- a/Engine/Source/Editor/WorkspaceMenuStructure/Public/WorkspaceMenuStructureModule.h
+++ b/Engine/Source/Editor/WorkspaceMenuStructure/Public/WorkspaceMenuStructureModule.h
@@ -23,9 +23,6 @@ public:
 	/** @return The menu structure that is populated by tab spawners */
 	virtual const IWorkspaceMenuStructure& GetWorkspaceMenuStructure() const;
 
-	/** Reset the "Asset Editor" category to its default state */
-	virtual void ResetAssetEditorCategory();
-
 	/** Reset the "Level Editor" category to its default state */
 	virtual void ResetLevelEditorCategory();
 
diff --git a/Engine/Source/Runtime/Slate/Private/Framework/Docking/TabManager.cpp b/Engine/Source/Runtime/Slate/Private/Framework/Docking/TabManager.cpp
--- a/Engine/Source/Runtime/Slate/Private/Framework/Docking/TabManager.cpp
+++ b/Engine/Source/Runtime/Slate/Private/Framework/Docking/TabManager.cpp
@@ -507,6 +507,11 @@ bool FTabManager::IsTabCloseable(const TSharedRef<const SDockTab>& InTab) const
 	return !(MainNonCloseableTab.Pin() == InTab);
 }
 
+TSharedRef<FWorkspaceItem> FTabManager::GetLocalWorkspaceMenuRoot()
+{
+	return LocalWorkspaceMenuRoot.ToSharedRef();
+}
+
 TSharedPtr<FTabManager::FStack> FTabManager::FLayoutNode::AsStack()
 {
 	return TSharedPtr<FTabManager::FStack>();
@@ -706,19 +711,19 @@ void FTabManager::PopulateTabSpawnerMenu_Helper( FMenuBuilder& PopulateMe, FPopu
 		else
 		{
 			// GROUP NODE
-			// If it's not empty, populate it
+			// If it's not empty, create a section and populate it
 			if ( ChildItem->HasChildrenIn(*Args.AllSpawners) )
 			{
 				const FPopulateTabSpawnerMenu_Args Payload( Args.AllSpawners, ChildItem, Args.Level+1 );
 
 				if ( Args.Level % 2 == 0 )
 				{
-					if (!bFirstItemOnLevel)
+					FName SectionName(*ChildItem->GetDisplayName().ToString().Replace(TEXT(" "), TEXT("")));
+					PopulateMe.BeginSection(SectionName, ChildItem->GetDisplayName());
 					{
-						PopulateMe.AddMenuSeparator();
+						PopulateTabSpawnerMenu_Helper(PopulateMe, Payload);
 					}
-						
-					PopulateTabSpawnerMenu_Helper( PopulateMe, Payload );
+					PopulateMe.EndSection();
 				}
 				else
 				{
@@ -740,7 +745,6 @@ void FTabManager::PopulateTabSpawnerMenu_Helper( FMenuBuilder& PopulateMe, FPopu
 
 void FTabManager::MakeSpawnerMenuEntry( FMenuBuilder &PopulateMe, const TSharedPtr<FTabSpawnerEntry> &SpawnerNode ) 
 {
-	//PopulateMe.BeginSection( "TabSection", SpawnerNode->GetDisplayName() );
 	if ( SpawnerNode->MenuType != ETabSpawnerMenuType::Hide )
 	{
 		PopulateMe.AddMenuEntry(
@@ -756,19 +760,25 @@ void FTabManager::MakeSpawnerMenuEntry( FMenuBuilder &PopulateMe, const TSharedP
 			EUserInterfaceActionType::Check
 			);
 	}
-	//PopulateMe.EndSection();
+}
+
+void FTabManager::PopulateLocalTabSpawnerMenu(FMenuBuilder& PopulateMe)
+{
+	PopulateTabSpawnerMenu(PopulateMe, LocalWorkspaceMenuRoot.ToSharedRef());
 }
 
 void FTabManager::PopulateTabSpawnerMenu( FMenuBuilder& PopulateMe, TSharedRef<FWorkspaceItem> MenuStructure )
 {
 	TSharedRef< TArray< TWeakPtr<FTabSpawnerEntry> > > AllSpawners = MakeShareable( new TArray< TWeakPtr<FTabSpawnerEntry> >() );
 	{
+		// Editor-specific tabs
 		for ( FTabSpawner::TIterator SpawnerIterator(TabSpawner); SpawnerIterator; ++SpawnerIterator  )	
 		{
 			const TSharedRef<FTabSpawnerEntry>& SpawnerEntry = SpawnerIterator.Value();
 			AllSpawners->AddUnique(SpawnerEntry);
 		}
 
+		// General Tabs
 		for ( FTabSpawner::TIterator SpawnerIterator(*NomadTabSpawner); SpawnerIterator; ++SpawnerIterator  )	
 		{
 			const TSharedRef<FTabSpawnerEntry>& SpawnerEntry = SpawnerIterator.Value();
@@ -992,6 +1002,7 @@ FTabManager::FTabManager( const TSharedPtr<SDockTab>& InOwnerTab, const TSharedR
 , LastDocumentUID( 0 )
 , bIsSavingVisualState( false )
 {
+	LocalWorkspaceMenuRoot = FWorkspaceItem::NewGroup(LOCTEXT("LocalWorkspaceRoot", "Local Workspace Root"));
 }
 
 
diff --git a/Engine/Source/Runtime/Slate/Public/Framework/Docking/TabManager.h b/Engine/Source/Runtime/Slate/Public/Framework/Docking/TabManager.h
--- a/Engine/Source/Runtime/Slate/Public/Framework/Docking/TabManager.h
+++ b/Engine/Source/Runtime/Slate/Public/Framework/Docking/TabManager.h
@@ -550,6 +550,8 @@ class SLATE_API FTabManager : public TSharedFromThis<FTabManager>
 
 		TSharedPtr<SWidget> RestoreFrom( const TSharedRef<FLayout>& Layout, const TSharedPtr<SWindow>& ParentWindow, const bool bEmbedTitleAreaContent = false );
 
+		void PopulateLocalTabSpawnerMenu( FMenuBuilder& PopulateMe );
+
 		void PopulateTabSpawnerMenu( FMenuBuilder& PopulateMe, TSharedRef<FWorkspaceItem> MenuStructure );
 
 		void PopulateTabSpawnerMenu( FMenuBuilder &PopulateMe, const FName& TabType );
@@ -604,6 +606,9 @@ class SLATE_API FTabManager : public TSharedFromThis<FTabManager>
 		/** @return if the provided tab can be closed. */
 		bool IsTabCloseable(const TSharedRef<const SDockTab>& InTab) const;
 
+		/** @return The local workspace menu root */
+		TSharedRef<FWorkspaceItem> GetLocalWorkspaceMenuRoot();
+
 	protected:
 		void InvokeTabForMenu( FName TabId );
 
@@ -678,6 +683,9 @@ class SLATE_API FTabManager : public TSharedFromThis<FTabManager>
 		TArray< TWeakPtr<SDockingArea> > DockAreas;
 		TArray< TSharedRef<FTabManager::FArea> > CollapsedDockAreas;
 
+		/** The root for the local editor's tab spawner workspace menu */
+		TSharedPtr<FWorkspaceItem> LocalWorkspaceMenuRoot;
+
 		/** A Major tab that contains this TabManager's widgets. */
 		TWeakPtr<SDockTab> OwnerTabPtr;
 
diff --git a/Engine/Source/Runtime/UMG/Private/UMGStyle.cpp b/Engine/Source/Runtime/UMG/Private/UMGStyle.cpp
--- a/Engine/Source/Runtime/UMG/Private/UMGStyle.cpp
+++ b/Engine/Source/Runtime/UMG/Private/UMGStyle.cpp
@@ -83,15 +83,22 @@ TSharedRef< FSlateStyleSet > FUMGStyle::Create()
 
 	Style->Set("Widget.UserWidget", new IMAGE_BRUSH(TEXT("UserWidget"), Icon16x16));
 
-	Style->Set("UMGEditor.Tabs.Animations", new IMAGE_BRUSH(TEXT("Animations_16x"), Icon16x16));
-	Style->Set("UMGEditor.Tabs.Designer", new IMAGE_BRUSH(TEXT("Designer_16x"), Icon16x16));
-	Style->Set("UMGEditor.Tabs.Palette", new IMAGE_BRUSH(TEXT("Palette_16x"), Icon16x16));
-	Style->Set("UMGEditor.Tabs.Timeline", new IMAGE_BRUSH(TEXT("Timeline_16x"), Icon16x16));
-
-	Style->Set("UMGEditor.Icons.Animations", new IMAGE_BRUSH(TEXT("Animations_40x"), Icon40x40));
-	Style->Set("UMGEditor.Icons.Designer", new IMAGE_BRUSH(TEXT("Designer_40x"), Icon40x40));
-	Style->Set("UMGEditor.Icons.Palette", new IMAGE_BRUSH(TEXT("Palette_40x"), Icon40x40));
-	Style->Set("UMGEditor.Icons.Timeline", new IMAGE_BRUSH(TEXT("Timeline_40x"), Icon40x40));
+	Style->Set("Animations.TabIcon", new IMAGE_BRUSH(TEXT("Animations_16x"), Icon16x16));
+	Style->Set("Designer.TabIcon", new IMAGE_BRUSH(TEXT("Designer_16x"), Icon16x16));
+	Style->Set("Palette.TabIcon", new IMAGE_BRUSH(TEXT("Palette_16x"), Icon16x16));
+	Style->Set("Sequencer.TabIcon", new IMAGE_BRUSH(TEXT("Timeline_16x"), Icon16x16));
+
+	Style->Set("Animations.Icon", new IMAGE_BRUSH(TEXT("Animations_40x"), Icon40x40));
+	Style->Set("Animations.Icon.Small", new IMAGE_BRUSH(TEXT("Animations_40x"), Icon20x20));
+
+	Style->Set("Designer.Icon", new IMAGE_BRUSH(TEXT("Designer_40x"), Icon40x40));
+	Style->Set("Designer.Icon.Small", new IMAGE_BRUSH(TEXT("Designer_40x"), Icon20x20));
+
+	Style->Set("Palette.Icon", new IMAGE_BRUSH(TEXT("Palette_40x"), Icon40x40));
+	Style->Set("Palette.Icon.Small", new IMAGE_BRUSH(TEXT("Palette_40x"), Icon20x20));
+
+	Style->Set("Timeline.Icon", new IMAGE_BRUSH(TEXT("Timeline_40x"), Icon40x40));
+	Style->Set("Timeline.Icon.Small", new IMAGE_BRUSH(TEXT("Timeline_40x"), Icon20x20));
 
 
 	return Style;
