New set of window close/minimize/zoom buttons for OS X 10.10

diff --git a/Engine/Source/Editor/EditorStyle/Private/SlateEditorStyle.cpp b/Engine/Source/Editor/EditorStyle/Private/SlateEditorStyle.cpp
--- a/Engine/Source/Editor/EditorStyle/Private/SlateEditorStyle.cpp
+++ b/Engine/Source/Editor/EditorStyle/Private/SlateEditorStyle.cpp
@@ -2234,45 +2234,92 @@ void FSlateEditorStyle::FStyle::SetupWindowStyles()
 		Set( "Window.Title.Flash",				  new IMAGE_BRUSH( "Old/Window/WindowTitle_Flashing", Icon24x24, FLinearColor(1,1,1,1), ESlateBrushTileType::Horizontal  ) );
 
 #if PLATFORM_MAC
-		const FButtonStyle MinimizeButtonStyle = FButtonStyle(Button)
-			.SetNormal ( IMAGE_BRUSH( "Old/Window/WindowButton_Mac_Minimize_Normal", FVector2D(20, 20) ) )
-			.SetHovered( IMAGE_BRUSH( "Old/Window/WindowButton_Mac_Minimize_Hovered", FVector2D(20, 20) ) )
-			.SetPressed( IMAGE_BRUSH( "Old/Window/WindowButton_Mac_Minimize_Pressed", FVector2D(20, 20) ) )
-			.SetDisabled( IMAGE_BRUSH( "Old/Window/WindowButton_Mac_Disabled", FVector2D(20, 20) ) );
-
-		Set( "Window.Buttons.Minimize.Normal",    new IMAGE_BRUSH( "Old/Window/WindowButton_Mac_Minimize_Normal", FVector2D(20, 20) ) );
-		Set( "Window.Buttons.Minimize.Hovered",   new IMAGE_BRUSH( "Old/Window/WindowButton_Mac_Minimize_Hovered", FVector2D(20, 20) ) );
-		Set( "Window.Buttons.Minimize.Pressed",   new IMAGE_BRUSH( "Old/Window/WindowButton_Mac_Minimize_Pressed", FVector2D(20, 20) ) );
-		Set( "Window.Buttons.Minimize.Disabled",  new IMAGE_BRUSH( "Old/Window/WindowButton_Mac_Disabled", FVector2D(20, 20) ) );
-
-		const FButtonStyle MaximizeButtonStyle = FButtonStyle(Button)
-			.SetNormal ( IMAGE_BRUSH( "Old/Window/WindowButton_Mac_Maximize_Normal", FVector2D(20, 20) ) )
-			.SetHovered( IMAGE_BRUSH( "Old/Window/WindowButton_Mac_Maximize_Hovered", FVector2D(20, 20) ) )
-			.SetPressed( IMAGE_BRUSH( "Old/Window/WindowButton_Mac_Maximize_Pressed", FVector2D(20, 20) ) )
-			.SetDisabled( IMAGE_BRUSH( "Old/Window/WindowButton_Mac_Disabled", FVector2D(20, 20) ) );
-
-		Set( "Window.Buttons.Maximize.Normal",    new IMAGE_BRUSH( "Old/Window/WindowButton_Mac_Maximize_Normal", FVector2D(20, 20) ) );
-		Set( "Window.Buttons.Maximize.Hovered",   new IMAGE_BRUSH( "Old/Window/WindowButton_Mac_Maximize_Hovered", FVector2D(20, 20) ) );
-		Set( "Window.Buttons.Maximize.Pressed",   new IMAGE_BRUSH( "Old/Window/WindowButton_Mac_Maximize_Pressed", FVector2D(20, 20) ) );
-		Set( "Window.Buttons.Maximize.Disabled",  new IMAGE_BRUSH( "Old/Window/WindowButton_Mac_Disabled", FVector2D(20, 20) ) );
-
-		const FButtonStyle RestoreButtonStyle = FButtonStyle(Button)
-			.SetNormal ( IMAGE_BRUSH( "Old/Window/WindowButton_Mac_Maximize_Normal", FVector2D(20, 20) ) )
-			.SetHovered( IMAGE_BRUSH( "Old/Window/WindowButton_Mac_Maximize_Hovered", FVector2D(20, 20) ) )
-			.SetPressed( IMAGE_BRUSH( "Old/Window/WindowButton_Mac_Maximize_Pressed", FVector2D(20, 20) ) );
-
-		Set( "Window.Buttons.Restore.Normal",     new IMAGE_BRUSH( "Old/Window/WindowButton_Mac_Maximize_Normal", FVector2D(20, 20) ) );
-		Set( "Window.Buttons.Restore.Hovered",    new IMAGE_BRUSH( "Old/Window/WindowButton_Mac_Maximize_Hovered", FVector2D(20, 20) ) );
-		Set( "Window.Buttons.Restore.Pressed",    new IMAGE_BRUSH( "Old/Window/WindowButton_Mac_Maximize_Pressed", FVector2D(20, 20) ) );
-
-		const FButtonStyle CloseButtonStyle = FButtonStyle(Button)
-			.SetNormal ( IMAGE_BRUSH( "Old/Window/WindowButton_Mac_Close_Normal", FVector2D(20, 20) ) )
-			.SetHovered( IMAGE_BRUSH( "Old/Window/WindowButton_Mac_Close_Hovered", FVector2D(20, 20) ) )
-			.SetPressed( IMAGE_BRUSH( "Old/Window/WindowButton_Mac_Close_Pressed", FVector2D(20, 20) ) );
-
-		Set( "Window.Buttons.Close.Normal",       new IMAGE_BRUSH( "Old/Window/WindowButton_Mac_Close_Normal", FVector2D(20, 20) ) );
-		Set( "Window.Buttons.Close.Hovered",      new IMAGE_BRUSH( "Old/Window/WindowButton_Mac_Close_Hovered", FVector2D(20, 20) ) );
-		Set( "Window.Buttons.Close.Pressed",      new IMAGE_BRUSH( "Old/Window/WindowButton_Mac_Close_Pressed", FVector2D(20, 20) ) );
+		const bool RunningOnMavericks = FPlatformMisc::IsRunningOnMavericks();
+		FButtonStyle MinimizeButtonStyle, MaximizeButtonStyle, RestoreButtonStyle, CloseButtonStyle;
+		if (RunningOnMavericks)
+		{
+			MinimizeButtonStyle = FButtonStyle(Button)
+				.SetNormal ( IMAGE_BRUSH( "Old/Window/WindowButton_Mac_Minimize_Normal", FVector2D(20, 20) ) )
+				.SetHovered( IMAGE_BRUSH( "Old/Window/WindowButton_Mac_Minimize_Hovered", FVector2D(20, 20) ) )
+				.SetPressed( IMAGE_BRUSH( "Old/Window/WindowButton_Mac_Minimize_Pressed", FVector2D(20, 20) ) )
+				.SetDisabled( IMAGE_BRUSH( "Old/Window/WindowButton_Mac_Disabled", FVector2D(20, 20) ) );
+
+			Set( "Window.Buttons.Minimize.Normal",    new IMAGE_BRUSH( "Old/Window/WindowButton_Mac_Minimize_Normal", FVector2D(20, 20) ) );
+			Set( "Window.Buttons.Minimize.Hovered",   new IMAGE_BRUSH( "Old/Window/WindowButton_Mac_Minimize_Hovered", FVector2D(20, 20) ) );
+			Set( "Window.Buttons.Minimize.Pressed",   new IMAGE_BRUSH( "Old/Window/WindowButton_Mac_Minimize_Pressed", FVector2D(20, 20) ) );
+			Set( "Window.Buttons.Minimize.Disabled",  new IMAGE_BRUSH( "Old/Window/WindowButton_Mac_Disabled", FVector2D(20, 20) ) );
+
+			MaximizeButtonStyle = FButtonStyle(Button)
+				.SetNormal ( IMAGE_BRUSH( "Old/Window/WindowButton_Mac_Maximize_Normal", FVector2D(20, 20) ) )
+				.SetHovered( IMAGE_BRUSH( "Old/Window/WindowButton_Mac_Maximize_Hovered", FVector2D(20, 20) ) )
+				.SetPressed( IMAGE_BRUSH( "Old/Window/WindowButton_Mac_Maximize_Pressed", FVector2D(20, 20) ) )
+				.SetDisabled( IMAGE_BRUSH( "Old/Window/WindowButton_Mac_Disabled", FVector2D(20, 20) ) );
+
+			Set( "Window.Buttons.Maximize.Normal",    new IMAGE_BRUSH( "Old/Window/WindowButton_Mac_Maximize_Normal", FVector2D(20, 20) ) );
+			Set( "Window.Buttons.Maximize.Hovered",   new IMAGE_BRUSH( "Old/Window/WindowButton_Mac_Maximize_Hovered", FVector2D(20, 20) ) );
+			Set( "Window.Buttons.Maximize.Pressed",   new IMAGE_BRUSH( "Old/Window/WindowButton_Mac_Maximize_Pressed", FVector2D(20, 20) ) );
+			Set( "Window.Buttons.Maximize.Disabled",  new IMAGE_BRUSH( "Old/Window/WindowButton_Mac_Disabled", FVector2D(20, 20) ) );
+
+			RestoreButtonStyle = FButtonStyle(Button)
+				.SetNormal ( IMAGE_BRUSH( "Old/Window/WindowButton_Mac_Maximize_Normal", FVector2D(20, 20) ) )
+				.SetHovered( IMAGE_BRUSH( "Old/Window/WindowButton_Mac_Maximize_Hovered", FVector2D(20, 20) ) )
+				.SetPressed( IMAGE_BRUSH( "Old/Window/WindowButton_Mac_Maximize_Pressed", FVector2D(20, 20) ) );
+
+			Set( "Window.Buttons.Restore.Normal",     new IMAGE_BRUSH( "Old/Window/WindowButton_Mac_Maximize_Normal", FVector2D(20, 20) ) );
+			Set( "Window.Buttons.Restore.Hovered",    new IMAGE_BRUSH( "Old/Window/WindowButton_Mac_Maximize_Hovered", FVector2D(20, 20) ) );
+			Set( "Window.Buttons.Restore.Pressed",    new IMAGE_BRUSH( "Old/Window/WindowButton_Mac_Maximize_Pressed", FVector2D(20, 20) ) );
+
+			CloseButtonStyle = FButtonStyle(Button)
+				.SetNormal ( IMAGE_BRUSH( "Old/Window/WindowButton_Mac_Close_Normal", FVector2D(20, 20) ) )
+				.SetHovered( IMAGE_BRUSH( "Old/Window/WindowButton_Mac_Close_Hovered", FVector2D(20, 20) ) )
+				.SetPressed( IMAGE_BRUSH( "Old/Window/WindowButton_Mac_Close_Pressed", FVector2D(20, 20) ) );
+
+			Set( "Window.Buttons.Close.Normal",       new IMAGE_BRUSH( "Old/Window/WindowButton_Mac_Close_Normal", FVector2D(20, 20) ) );
+			Set( "Window.Buttons.Close.Hovered",      new IMAGE_BRUSH( "Old/Window/WindowButton_Mac_Close_Hovered", FVector2D(20, 20) ) );
+			Set( "Window.Buttons.Close.Pressed",      new IMAGE_BRUSH( "Old/Window/WindowButton_Mac_Close_Pressed", FVector2D(20, 20) ) );
+		}
+		else
+		{
+			MinimizeButtonStyle = FButtonStyle(Button)
+				.SetNormal ( IMAGE_BRUSH( "Old/Window/WindowButton_Mac_Yosemite_Minimize_Normal", FVector2D(20, 20) ) )
+				.SetHovered( IMAGE_BRUSH( "Old/Window/WindowButton_Mac_Yosemite_Minimize_Hovered", FVector2D(20, 20) ) )
+				.SetPressed( IMAGE_BRUSH( "Old/Window/WindowButton_Mac_Yosemite_Minimize_Pressed", FVector2D(20, 20) ) )
+				.SetDisabled( IMAGE_BRUSH( "Old/Window/WindowButton_Mac_Yosemite_Disabled", FVector2D(20, 20) ) );
+
+			Set( "Window.Buttons.Minimize.Normal",    new IMAGE_BRUSH( "Old/Window/WindowButton_Mac_Yosemite_Minimize_Normal", FVector2D(20, 20) ) );
+			Set( "Window.Buttons.Minimize.Hovered",   new IMAGE_BRUSH( "Old/Window/WindowButton_Mac_Yosemite_Minimize_Hovered", FVector2D(20, 20) ) );
+			Set( "Window.Buttons.Minimize.Pressed",   new IMAGE_BRUSH( "Old/Window/WindowButton_Mac_Yosemite_Minimize_Pressed", FVector2D(20, 20) ) );
+			Set( "Window.Buttons.Minimize.Disabled",  new IMAGE_BRUSH( "Old/Window/WindowButton_Mac_Yosemite_Disabled", FVector2D(20, 20) ) );
+
+			MaximizeButtonStyle = FButtonStyle(Button)
+				.SetNormal ( IMAGE_BRUSH( "Old/Window/WindowButton_Mac_Yosemite_Maximize_Normal", FVector2D(20, 20) ) )
+				.SetHovered( IMAGE_BRUSH( "Old/Window/WindowButton_Mac_Yosemite_Maximize_Hovered", FVector2D(20, 20) ) )
+				.SetPressed( IMAGE_BRUSH( "Old/Window/WindowButton_Mac_Yosemite_Maximize_Pressed", FVector2D(20, 20) ) )
+				.SetDisabled( IMAGE_BRUSH( "Old/Window/WindowButton_Mac_Yosemite_Disabled", FVector2D(20, 20) ) );
+
+			Set( "Window.Buttons.Maximize.Normal",    new IMAGE_BRUSH( "Old/Window/WindowButton_Mac_Yosemite_Maximize_Normal", FVector2D(20, 20) ) );
+			Set( "Window.Buttons.Maximize.Hovered",   new IMAGE_BRUSH( "Old/Window/WindowButton_Mac_Yosemite_Maximize_Hovered", FVector2D(20, 20) ) );
+			Set( "Window.Buttons.Maximize.Pressed",   new IMAGE_BRUSH( "Old/Window/WindowButton_Mac_Yosemite_Maximize_Pressed", FVector2D(20, 20) ) );
+			Set( "Window.Buttons.Maximize.Disabled",  new IMAGE_BRUSH( "Old/Window/WindowButton_Mac_Yosemite_Disabled", FVector2D(20, 20) ) );
+
+			RestoreButtonStyle = FButtonStyle(Button)
+				.SetNormal ( IMAGE_BRUSH( "Old/Window/WindowButton_Mac_Yosemite_Maximize_Normal", FVector2D(20, 20) ) )
+				.SetHovered( IMAGE_BRUSH( "Old/Window/WindowButton_Mac_Yosemite_Maximize_Hovered", FVector2D(20, 20) ) )
+				.SetPressed( IMAGE_BRUSH( "Old/Window/WindowButton_Mac_Yosemite_Maximize_Pressed", FVector2D(20, 20) ) );
+
+			Set( "Window.Buttons.Restore.Normal",     new IMAGE_BRUSH( "Old/Window/WindowButton_Mac_Yosemite_Maximize_Normal", FVector2D(20, 20) ) );
+			Set( "Window.Buttons.Restore.Hovered",    new IMAGE_BRUSH( "Old/Window/WindowButton_Mac_Yosemite_Maximize_Hovered", FVector2D(20, 20) ) );
+			Set( "Window.Buttons.Restore.Pressed",    new IMAGE_BRUSH( "Old/Window/WindowButton_Mac_Yosemite_Maximize_Pressed", FVector2D(20, 20) ) );
+
+			CloseButtonStyle = FButtonStyle(Button)
+				.SetNormal ( IMAGE_BRUSH( "Old/Window/WindowButton_Mac_Yosemite_Close_Normal", FVector2D(20, 20) ) )
+				.SetHovered( IMAGE_BRUSH( "Old/Window/WindowButton_Mac_Yosemite_Close_Hovered", FVector2D(20, 20) ) )
+				.SetPressed( IMAGE_BRUSH( "Old/Window/WindowButton_Mac_Yosemite_Close_Pressed", FVector2D(20, 20) ) );
+
+			Set( "Window.Buttons.Close.Normal",       new IMAGE_BRUSH( "Old/Window/WindowButton_Mac_Yosemite_Close_Normal", FVector2D(20, 20) ) );
+			Set( "Window.Buttons.Close.Hovered",      new IMAGE_BRUSH( "Old/Window/WindowButton_Mac_Yosemite_Close_Hovered", FVector2D(20, 20) ) );
+			Set( "Window.Buttons.Close.Pressed",      new IMAGE_BRUSH( "Old/Window/WindowButton_Mac_Yosemite_Close_Pressed", FVector2D(20, 20) ) );
+		}
 #else
 		const FButtonStyle MinimizeButtonStyle = FButtonStyle(Button)
 			.SetNormal ( IMAGE_BRUSH( "Old/Window/WindowButton_Minimize_Normal", FVector2D(27, 18) ) )
@@ -2314,7 +2361,7 @@ void FSlateEditorStyle::FStyle::SetupWindowStyles()
 		Set( "Window.Buttons.Close.Hovered",      new IMAGE_BRUSH( "Old/Window/WindowButton_Close_Hovered", FVector2D(44, 18) ) );
 		Set( "Window.Buttons.Close.Pressed",      new IMAGE_BRUSH( "Old/Window/WindowButton_Close_Pressed", FVector2D(44, 18) ) );
 #endif
-		
+
 		Set( "Window.Buttons.Minimize", MinimizeButtonStyle );
 		Set( "Window.Buttons.Maximize", MaximizeButtonStyle );
 		Set( "Window.Buttons.Restore", RestoreButtonStyle );
diff --git a/Engine/Source/Runtime/SlateCore/Private/Styling/CoreStyle.cpp b/Engine/Source/Runtime/SlateCore/Private/Styling/CoreStyle.cpp
--- a/Engine/Source/Runtime/SlateCore/Private/Styling/CoreStyle.cpp
+++ b/Engine/Source/Runtime/SlateCore/Private/Styling/CoreStyle.cpp
@@ -1019,29 +1019,58 @@ TSharedRef<ISlateStyle> FCoreStyle::Create( const FName& InStyleSetName )
 	// SWindow defaults...
 	{
 #if PLATFORM_MAC
-		const FButtonStyle MinimizeButtonStyle = FButtonStyle(Button)
-			.SetNormal(IMAGE_BRUSH("Common/Window/WindowButton_Mac_Minimize_Normal", FVector2D(20.0f, 20.0f)))
-			.SetHovered(IMAGE_BRUSH("Common/Window/WindowButton_Mac_Minimize_Hovered", FVector2D(20.0f, 20.0f)))
-			.SetPressed(IMAGE_BRUSH("Common/Window/WindowButton_Mac_Minimize_Pressed", FVector2D(20.0f, 20.0f)))
-			.SetDisabled(IMAGE_BRUSH("Common/Window/WindowButton_Mac_Disabled", FVector2D(20.0f, 20.0f)));
-
-		const FButtonStyle MaximizeButtonStyle = FButtonStyle(Button)
-			.SetNormal (IMAGE_BRUSH("Common/Window/WindowButton_Mac_Maximize_Normal", FVector2D(20.0f, 20.0f)))
-			.SetHovered(IMAGE_BRUSH("Common/Window/WindowButton_Mac_Maximize_Hovered", FVector2D(20.0f, 20.0f)))
-			.SetPressed(IMAGE_BRUSH("Common/Window/WindowButton_Mac_Maximize_Pressed", FVector2D(20.0f, 20.0f)))
-			.SetDisabled(IMAGE_BRUSH("Common/Window/WindowButton_Mac_Disabled", FVector2D(20.0f, 20.0f)));
-
-		const FButtonStyle RestoreButtonStyle = FButtonStyle(Button)
-			.SetNormal (IMAGE_BRUSH("Common/Window/WindowButton_Mac_Maximize_Normal", FVector2D(20.0f, 20.0f)))
-			.SetHovered(IMAGE_BRUSH("Common/Window/WindowButton_Mac_Maximize_Hovered", FVector2D(20.0f, 20.0f)))
-			.SetPressed(IMAGE_BRUSH("Common/Window/WindowButton_Mac_Maximize_Pressed", FVector2D(20.0f, 20.0f)))
-			.SetDisabled(IMAGE_BRUSH("Common/Window/WindowButton_Mac_Disabled", FVector2D(20.0f, 20.0f)));
-
-		const FButtonStyle CloseButtonStyle = FButtonStyle(Button)
-			.SetNormal (IMAGE_BRUSH("Common/Window/WindowButton_Mac_Close_Normal", FVector2D(20.0f, 20.0f)))
-			.SetHovered(IMAGE_BRUSH("Common/Window/WindowButton_Mac_Close_Hovered", FVector2D(20.0f, 20.0f)))
-			.SetPressed(IMAGE_BRUSH("Common/Window/WindowButton_Mac_Close_Pressed", FVector2D(20.0f, 20.0f)))
-			.SetDisabled(IMAGE_BRUSH("Common/Window/WindowButton_Mac_Disabled", FVector2D(20.0f, 20.0f)));
+		const bool RunningOnMavericks = FPlatformMisc::IsRunningOnMavericks();
+		const FButtonStyle MinimizeButtonStyle = RunningOnMavericks ?
+			FButtonStyle(Button)
+				.SetNormal(IMAGE_BRUSH("Common/Window/WindowButton_Mac_Minimize_Normal", FVector2D(20.0f, 20.0f)))
+				.SetHovered(IMAGE_BRUSH("Common/Window/WindowButton_Mac_Minimize_Hovered", FVector2D(20.0f, 20.0f)))
+				.SetPressed(IMAGE_BRUSH("Common/Window/WindowButton_Mac_Minimize_Pressed", FVector2D(20.0f, 20.0f)))
+				.SetDisabled(IMAGE_BRUSH("Common/Window/WindowButton_Mac_Disabled", FVector2D(20.0f, 20.0f)))
+		:
+			FButtonStyle(Button)
+				.SetNormal(IMAGE_BRUSH("Common/Window/WindowButton_Mac_Yosemite_Minimize_Normal", FVector2D(20.0f, 20.0f)))
+				.SetHovered(IMAGE_BRUSH("Common/Window/WindowButton_Mac_Yosemite_Minimize_Hovered", FVector2D(20.0f, 20.0f)))
+				.SetPressed(IMAGE_BRUSH("Common/Window/WindowButton_Mac_Yosemite_Minimize_Pressed", FVector2D(20.0f, 20.0f)))
+				.SetDisabled(IMAGE_BRUSH("Common/Window/WindowButton_Mac_Yosemite_Disabled", FVector2D(20.0f, 20.0f)));
+
+		const FButtonStyle MaximizeButtonStyle = RunningOnMavericks ?
+			FButtonStyle(Button)
+				.SetNormal (IMAGE_BRUSH("Common/Window/WindowButton_Mac_Maximize_Normal", FVector2D(20.0f, 20.0f)))
+				.SetHovered(IMAGE_BRUSH("Common/Window/WindowButton_Mac_Maximize_Hovered", FVector2D(20.0f, 20.0f)))
+				.SetPressed(IMAGE_BRUSH("Common/Window/WindowButton_Mac_Maximize_Pressed", FVector2D(20.0f, 20.0f)))
+				.SetDisabled(IMAGE_BRUSH("Common/Window/WindowButton_Mac_Disabled", FVector2D(20.0f, 20.0f)))
+		:
+			FButtonStyle(Button)
+				.SetNormal (IMAGE_BRUSH("Common/Window/WindowButton_Mac_Yosemite_Maximize_Normal", FVector2D(20.0f, 20.0f)))
+				.SetHovered(IMAGE_BRUSH("Common/Window/WindowButton_Mac_Yosemite_Maximize_Hovered", FVector2D(20.0f, 20.0f)))
+				.SetPressed(IMAGE_BRUSH("Common/Window/WindowButton_Mac_Yosemite_Maximize_Pressed", FVector2D(20.0f, 20.0f)))
+				.SetDisabled(IMAGE_BRUSH("Common/Window/WindowButton_Mac_Yosemite_Disabled", FVector2D(20.0f, 20.0f)));
+
+		const FButtonStyle RestoreButtonStyle = RunningOnMavericks ?
+			FButtonStyle(Button)
+				.SetNormal (IMAGE_BRUSH("Common/Window/WindowButton_Mac_Maximize_Normal", FVector2D(20.0f, 20.0f)))
+				.SetHovered(IMAGE_BRUSH("Common/Window/WindowButton_Mac_Maximize_Hovered", FVector2D(20.0f, 20.0f)))
+				.SetPressed(IMAGE_BRUSH("Common/Window/WindowButton_Mac_Maximize_Pressed", FVector2D(20.0f, 20.0f)))
+				.SetDisabled(IMAGE_BRUSH("Common/Window/WindowButton_Mac_Disabled", FVector2D(20.0f, 20.0f)))
+		:
+			FButtonStyle(Button)
+				.SetNormal (IMAGE_BRUSH("Common/Window/WindowButton_Mac_Yosemite_Maximize_Normal", FVector2D(20.0f, 20.0f)))
+				.SetHovered(IMAGE_BRUSH("Common/Window/WindowButton_Mac_Yosemite_Maximize_Hovered", FVector2D(20.0f, 20.0f)))
+				.SetPressed(IMAGE_BRUSH("Common/Window/WindowButton_Mac_Yosemite_Maximize_Pressed", FVector2D(20.0f, 20.0f)))
+				.SetDisabled(IMAGE_BRUSH("Common/Window/WindowButton_Mac_Yosemite_Disabled", FVector2D(20.0f, 20.0f)));
+
+		const FButtonStyle CloseButtonStyle = RunningOnMavericks ?
+			FButtonStyle(Button)
+				.SetNormal (IMAGE_BRUSH("Common/Window/WindowButton_Mac_Close_Normal", FVector2D(20.0f, 20.0f)))
+				.SetHovered(IMAGE_BRUSH("Common/Window/WindowButton_Mac_Close_Hovered", FVector2D(20.0f, 20.0f)))
+				.SetPressed(IMAGE_BRUSH("Common/Window/WindowButton_Mac_Close_Pressed", FVector2D(20.0f, 20.0f)))
+				.SetDisabled(IMAGE_BRUSH("Common/Window/WindowButton_Mac_Disabled", FVector2D(20.0f, 20.0f)))
+		:
+			FButtonStyle(Button)
+				.SetNormal (IMAGE_BRUSH("Common/Window/WindowButton_Mac_Yosemite_Close_Normal", FVector2D(20.0f, 20.0f)))
+				.SetHovered(IMAGE_BRUSH("Common/Window/WindowButton_Mac_Yosemite_Close_Hovered", FVector2D(20.0f, 20.0f)))
+				.SetPressed(IMAGE_BRUSH("Common/Window/WindowButton_Mac_Yosemite_Close_Pressed", FVector2D(20.0f, 20.0f)))
+				.SetDisabled(IMAGE_BRUSH("Common/Window/WindowButton_Mac_Yosemite_Disabled", FVector2D(20.0f, 20.0f)));
 #else
 		const FButtonStyle MinimizeButtonStyle = FButtonStyle(Button)
 			.SetNormal (IMAGE_BRUSH("Common/Window/WindowButton_Minimize_Normal", FVector2D(27.0f, 18.0f)))
