Fix missing font warnings

diff --git a/Engine/Source/Editor/EditorStyle/Private/SlateEditorStyle.cpp b/Engine/Source/Editor/EditorStyle/Private/SlateEditorStyle.cpp
--- a/Engine/Source/Editor/EditorStyle/Private/SlateEditorStyle.cpp
+++ b/Engine/Source/Editor/EditorStyle/Private/SlateEditorStyle.cpp
@@ -4451,7 +4451,7 @@ void FSlateEditorStyle::FStyle::SetupPersonaStyle()
 		Set( "FullBlueprintEditor.Diff", new IMAGE_BRUSH( "Icons/BlueprintEditorDiff", Icon40x40 ) );
 		Set( "FullBlueprintEditor.Diff.Small", new IMAGE_BRUSH( "Icons/BlueprintEditorDiff", Icon20x20 ) );
 
-		Set( "BlueprintEditor.ActionMenu.ContextDescriptionFont",  FSlateFontInfo(TEXT("Roboto-Regular"), 12) );
+		Set( "BlueprintEditor.ActionMenu.ContextDescriptionFont",  TTF_CORE_FONT("Fonts/Roboto-Regular", 12) );
 
 		Set( "BlueprintEditor.FindInBlueprint", new IMAGE_BRUSH( "Icons/icon_Blueprint_Find_40px", Icon40x40 ) );
 		Set( "BlueprintEditor.FindInBlueprint.Small", new IMAGE_BRUSH( "Icons/icon_Blueprint_Find_40px", Icon20x20 ) );
