UMG - Removing an unused style.

diff --git a/Engine/Source/Editor/EditorStyle/Private/SlateEditorStyle.cpp b/Engine/Source/Editor/EditorStyle/Private/SlateEditorStyle.cpp
--- a/Engine/Source/Editor/EditorStyle/Private/SlateEditorStyle.cpp
+++ b/Engine/Source/Editor/EditorStyle/Private/SlateEditorStyle.cpp
@@ -6004,8 +6004,6 @@ void FSlateEditorStyle::FStyle::SetupUMGEditorStyles()
 	Set("UMGEditor.AnchoredWidget", new BOX_BRUSH("Common/Button", FVector2D(32, 32), 8.0f / 32.0f));
 	Set("UMGEditor.AnchoredWidgetAlignment", new IMAGE_BRUSH("Icons/icon_tab_DeviceManager_16x", Icon8x8));
 	
-	Set("UMGEditor.PalettePreviewWidget", new IMAGE_BRUSH("Icons/UMG/umg_anchors_fill_fill", Icon16x16));
-
 	Set("UMGEditor.PaletteHeader", FTableRowStyle()
 		.SetEvenRowBackgroundBrush( BOX_BRUSH( "PropertyView/DetailCategoryMiddle", FMargin( 4/16.0f, 8.0f/16.0f, 4/16.0f, 4/16.0f ), FLinearColor(0.6f,0.6f,0.6f,1.0f) ) )
 		.SetEvenRowBackgroundHoveredBrush( BOX_BRUSH( "PropertyView/DetailCategoryMiddle", FMargin( 4/16.0f, 8.0f/16.0f, 4/16.0f, 4/16.0f ), FLinearColor(0.3f,0.3f,0.3f,1.0f) ) )
