Backing out change to SlateEditorStyle.cpp from CL#2345039 which resulted in missing icons

diff --git a/Engine/Source/Editor/EditorStyle/Private/SlateEditorStyle.cpp b/Engine/Source/Editor/EditorStyle/Private/SlateEditorStyle.cpp
--- a/Engine/Source/Editor/EditorStyle/Private/SlateEditorStyle.cpp
+++ b/Engine/Source/Editor/EditorStyle/Private/SlateEditorStyle.cpp
@@ -3711,7 +3711,7 @@ void FSlateEditorStyle::FStyle::SetupGraphEditorStyles()
 		Set( "GraphEditor.SpawnActor_16x", new IMAGE_BRUSH("Icons/icon_Blueprint_SpawnActor_16x", Icon16x16));
 		Set( "GraphEditor.PadEvent_16x", new IMAGE_BRUSH("Icons/icon_Blueprint_PadEvent_16x", Icon16x16));
 		Set( "GraphEditor.MouseEvent_16x", new IMAGE_BRUSH("Icons/icon_Blueprint_MouseEvent_16x", Icon16x16));
-		Set( "GraphEditor.KeyEvent_16x", new IMAGE_BRUSH("Icons/icon_Blueprint_KeyEvent_16x", Icon16x16));
+		Set( "GraphEditor.KeyEvent_16x", new IMAGE_BRUSH("Icons/icon_Blueprint_KeyboardEvent_16x", Icon16x16));
 		Set( "GraphEditor.TouchEvent_16x", new IMAGE_BRUSH("Icons/icon_Blueprint_TouchEvent_16x", Icon16x16));
 		Set( "GraphEditor.MakeArray_16x", new IMAGE_BRUSH("Icons/icon_Blueprint_MakeArray_16x", Icon16x16));
 		Set( "GraphEditor.Enum_16x", new IMAGE_BRUSH("Icons/icon_Blueprint_Enum_16x", Icon16x16));
