Fix for styles not being loaded in UFE.

diff --git a/Engine/Source/Editor/EditorStyle/Private/SlateEditorStyle.cpp b/Engine/Source/Editor/EditorStyle/Private/SlateEditorStyle.cpp
--- a/Engine/Source/Editor/EditorStyle/Private/SlateEditorStyle.cpp
+++ b/Engine/Source/Editor/EditorStyle/Private/SlateEditorStyle.cpp
@@ -116,10 +116,14 @@ void FSlateEditorStyle::FStyle::SyncSettings()
 
 void FSlateEditorStyle::FStyle::Initialize()
 {
+	//@Todo slate: splitting game and style atlases is a better solution to avoiding editor textures impacting game atlas pages. Tho this would still be a loading win.
+	// We do WITH_EDITOR and well as !GIsEditor because in UFE !GIsEditor is true, however we need the styles.
+#if WITH_EDITOR
 	if (!GIsEditor)
 	{
 		return;
 	}
+#endif
 
 	SyncSettings();
 
