Editor - Adding a new section for just general reusable icons to the slate style.

diff --git a/Engine/Source/Editor/EditorStyle/Private/SlateEditorStyle.cpp b/Engine/Source/Editor/EditorStyle/Private/SlateEditorStyle.cpp
--- a/Engine/Source/Editor/EditorStyle/Private/SlateEditorStyle.cpp
+++ b/Engine/Source/Editor/EditorStyle/Private/SlateEditorStyle.cpp
@@ -131,6 +131,7 @@ void FSlateEditorStyle::FStyle::Initialize()
 	SetCoreContentRoot(FPaths::EngineContentDir() / TEXT("Slate"));
 
 	SetupGeneralStyles();
+	SetupGeneralIcons();
 	SetupWindowStyles();
 	SetupDockingStyles();
 	SetupTutorialStyles();
@@ -2234,6 +2235,12 @@ void FSlateEditorStyle::FStyle::SetupGeneralStyles()
 #endif
 }
 
+void FSlateEditorStyle::FStyle::SetupGeneralIcons()
+{
+	Set("Plus", new IMAGE_BRUSH("Icons/PlusSymbol_12x", Icon12x12));
+	Set("Cross", new IMAGE_BRUSH("Icons/Cross_12x", Icon12x12));
+}
+
 void FSlateEditorStyle::FStyle::SetupWindowStyles()
 {
 	// Window styling
diff --git a/Engine/Source/Editor/EditorStyle/Private/SlateEditorStyle.h b/Engine/Source/Editor/EditorStyle/Private/SlateEditorStyle.h
--- a/Engine/Source/Editor/EditorStyle/Private/SlateEditorStyle.h
+++ b/Engine/Source/Editor/EditorStyle/Private/SlateEditorStyle.h
@@ -62,6 +62,7 @@ public:
 
 		void Initialize();
 		void SetupGeneralStyles();
+		void SetupGeneralIcons();
 		void SetupWindowStyles();
 		void SetupDockingStyles();
 		void SetupTutorialStyles();
