Implemented some missing Slate Widget styles.

diff --git a/Engine/Source/Editor/EditorStyle/Private/SlateEditorStyle.cpp b/Engine/Source/Editor/EditorStyle/Private/SlateEditorStyle.cpp
--- a/Engine/Source/Editor/EditorStyle/Private/SlateEditorStyle.cpp
+++ b/Engine/Source/Editor/EditorStyle/Private/SlateEditorStyle.cpp
@@ -1920,7 +1920,10 @@ void FSlateEditorStyle::FStyle::SetupGeneralStyles()
 	{
 		const FLinearColor LayoutSelectionColor_Hovered = FLinearColor(0.5f, 0.5f, 0.5f);
 
-		Set( "ViewportLayoutToolbar.Background",	new FSlateNoResource() );
+		Set( "ViewportLayoutToolbar.Background", new FSlateNoResource() );
+		Set( "ViewportLayoutToolbar.Label", FTextBlockStyle() );
+		Set( "ViewportLayoutToolbar.Button", FButtonStyle(NoBorder) );
+		Set( "ViewportLayoutToolbar.Expand", new IMAGE_BRUSH("Icons/toolbar_expand_16x", Icon16x16) );
 
 		/* Create style for "ViewportLayoutToolbar.ToggleButton" ... */
 		const FCheckBoxStyle ViewportLayoutToolbarToggleButtonStyle = FCheckBoxStyle()
@@ -2858,6 +2861,8 @@ void FSlateEditorStyle::FStyle::SetupPropertyEditorStyles()
 			.SetShadowColorAndOpacity( FLinearColor::Black )
 		);
 
+		Set( "PropertyEditor.AssetName.ColorAndOpacity", FLinearColor::White );
+
 		Set( "PropertyEditor.AssetThumbnailLight", new BOX_BRUSH( "ContentBrowser/ThumbnailLight", FMargin( 5.0f / 64.0f ), SelectionColor ) );
 		Set( "PropertyEditor.AssetThumbnailShadow", new BOX_BRUSH( "ContentBrowser/ThumbnailShadow", FMargin( 4.0f / 64.0f ) ) );
 		Set( "PropertyEditor.AssetClass", FTextBlockStyle(NormalText)
