Fixing warnings on build machines. Fonts/Roboto-Italic.ttf doesn't exist in Engine/Content/Slate/Fonts, only Engine/Content/Editor/Slate/Fonts.

diff --git a/Engine/Source/Editor/EditorStyle/Private/SlateEditorStyle.cpp b/Engine/Source/Editor/EditorStyle/Private/SlateEditorStyle.cpp
--- a/Engine/Source/Editor/EditorStyle/Private/SlateEditorStyle.cpp
+++ b/Engine/Source/Editor/EditorStyle/Private/SlateEditorStyle.cpp
@@ -2760,7 +2760,7 @@ void FSlateEditorStyle::FStyle::SetupTutorialStyles()
 				);
 
 			Set("TutorialEditableText.Toolbar.ItalicText", FTextBlockStyle(NormalText)
-				.SetFont(TTF_CORE_FONT("Fonts/Roboto-Italic", 10))
+				.SetFont(TTF_FONT("Fonts/Roboto-Italic", 10))
 				.SetColorAndOpacity(TextColor)
 				);
 
