[INTEGRATE] Change 2412157 by [REDACTED_PERSON][REDACTED_HANDLE] on 2015/01/20 09:34:12

	Remove special case icon code in Class Picker 'common' classes
	Add lots of placeholder icons for class types that were missing them (AIController, GameMode, HUD, PlayerController...)
	Fix class lookup too allow different icon for Actor and Object based classes
	Remove some unused/duplicate icons

diff --git a/Engine/Config/BaseEditor.ini b/Engine/Config/BaseEditor.ini
--- a/Engine/Config/BaseEditor.ini
+++ b/Engine/Config/BaseEditor.ini
@@ -180,11 +180,11 @@ bExpandClassPickerClassList=False
 
 ; Class picker dialog settings
 ; LocText refers to the localization id for the descriptive text to show next to the icon in the picker
-+NewAssetDefaultClasses=(ClassName="/Script/Engine.Actor", LocTextNameID="ActorName", LocTextDescriptionID="ActorDesc", Image ="ClassIcon.Default", AssetClass="/Script/Engine.Blueprint")
-+NewAssetDefaultClasses=(ClassName="/Script/Engine.Pawn", LocTextNameID="PawnName", LocTextDescriptionID="PawnDesc", Image ="ClassIcon.Pawn", AssetClass="/Script/Engine.Blueprint")
-+NewAssetDefaultClasses=(ClassName="/Script/Engine.Character", LocTextNameID="CharacterName", LocTextDescriptionID="CharacterDesc", Image ="ClassIcon.SkeletalMeshComponent", AssetClass="/Script/Engine.Blueprint")
-+NewAssetDefaultClasses=(ClassName="/Script/Engine.PlayerController", LocTextNameID="PlayerControllerName", LocTextDescriptionID="PlayerControllerDesc", Image ="ClassIcon.PlayerStart", AssetClass="/Script/Engine.Blueprint")
-+NewAssetDefaultClasses=(ClassName="/Script/Engine.GameMode", LocTextNameID="GameModeName", LocTextDescriptionID="GameModeDesc", Image ="ClassIcon.Note", AssetClass="/Script/Engine.Blueprint")
++NewAssetDefaultClasses=(ClassName="/Script/Engine.Actor", LocTextNameID="ActorName", LocTextDescriptionID="ActorDesc", AssetClass="/Script/Engine.Blueprint")
++NewAssetDefaultClasses=(ClassName="/Script/Engine.Pawn", LocTextNameID="PawnName", LocTextDescriptionID="PawnDesc", AssetClass="/Script/Engine.Blueprint")
++NewAssetDefaultClasses=(ClassName="/Script/Engine.Character", LocTextNameID="CharacterName", LocTextDescriptionID="CharacterDesc", AssetClass="/Script/Engine.Blueprint")
++NewAssetDefaultClasses=(ClassName="/Script/Engine.PlayerController", LocTextNameID="PlayerControllerName", LocTextDescriptionID="PlayerControllerDesc", AssetClass="/Script/Engine.Blueprint")
++NewAssetDefaultClasses=(ClassName="/Script/Engine.GameMode", LocTextNameID="GameModeName", LocTextDescriptionID="GameModeDesc", AssetClass="/Script/Engine.Blueprint")
 
 [TaskBrowser]
 DefaultServerName=
diff --git a/Engine/Source/Editor/EditorStyle/Private/SlateEditorStyle.cpp b/Engine/Source/Editor/EditorStyle/Private/SlateEditorStyle.cpp
--- a/Engine/Source/Editor/EditorStyle/Private/SlateEditorStyle.cpp
+++ b/Engine/Source/Editor/EditorStyle/Private/SlateEditorStyle.cpp
@@ -4824,23 +4824,14 @@ void FSlateEditorStyle::FStyle::SetupClassIconsAndThumbnails()
 	{
 		Set("ClassIcon.Emitter", new IMAGE_BRUSH("Icons/ActorIcons/Emitter_16x", Icon16x16));
 		Set("ClassIcon.Light", new IMAGE_BRUSH("Icons/ActorIcons/LightActor_16x", Icon16x16));
-		Set("ClassIcon.MatineeActor", new IMAGE_BRUSH("Icons/ActorIcons/MatineeActor_16x", Icon16x16));
 		Set("ClassIcon.Brush", new IMAGE_BRUSH("Icons/ActorIcons/Brush_16x", Icon16x16));
 		Set("ClassIcon.BrushAdditive", new IMAGE_BRUSH("Icons/ActorIcons/Brush_Add_16x", Icon16x16));
 		Set("ClassIcon.BrushSubtractive", new IMAGE_BRUSH("Icons/ActorIcons/Brush_Subtract_16x", Icon16x16));
 		Set("ClassIcon.Volume", new IMAGE_BRUSH("Icons/ActorIcons/Volume_16x", Icon16x16));
-		Set("ClassIcon.SkeletalPhysicsActor", new IMAGE_BRUSH("Icons/ActorIcons/SkeletalPhysicsActor_16x", Icon16x16));
 		Set("ClassIcon.GroupActor", new IMAGE_BRUSH("Icons/ActorIcons/GroupActor_16x", Icon16x16));
 		Set("ClassIcon.VectorFieldVolume", new IMAGE_BRUSH("Icons/ActorIcons/VectorFieldVolume_16x", Icon16x16));
 		Set("ClassIcon.Deleted", new IMAGE_BRUSH("Icons/ActorIcons/DeletedActor_16px", Icon16x16));
 		Set("ClassIcon.DestructibleActor", new IMAGE_BRUSH("Icons/ActorIcons/Destructable_16px", Icon16x16));
-		Set("ClassIcon.DefaultPawn", new IMAGE_BRUSH("Icons/ActorIcons/DefaultPawn_16px", Icon16x16));
-		Set("ClassIcon.Pawn", new IMAGE_BRUSH("Icons/ActorIcons/Pawn_16px", Icon16x16));
-		Set("ClassIcon.Note", new IMAGE_BRUSH("Icons/ActorIcons/Note_16px", Icon16x16));
-		Set("ClassIcon.Character", new IMAGE_BRUSH("Icons/ActorIcons/Character_16px", Icon16x16));
-		Set("ClassIcon.DocumentationActor", new IMAGE_BRUSH("Icons/ActorIcons/DocumentationActor_16x", Icon16x16));
-
-		Set("ClassIcon.BlueprintActorBase", new IMAGE_BRUSH("Icons/AssetIcons/Blueprint_16x", Icon16x16));
 		Set("ClassIcon.StaticMeshActor", new IMAGE_BRUSH("Icons/AssetIcons/StaticMesh_16x", Icon16x16));
 		Set("ClassIcon.SkeletalMeshActor", new IMAGE_BRUSH("Icons/AssetIcons/SkeletalMesh_16x", Icon16x16));
 
@@ -4918,6 +4909,9 @@ void FSlateEditorStyle::FStyle::SetupClassIconsAndThumbnails()
 		// Asset Type Classes
 		const TCHAR* AssetTypes[] = {
 
+			TEXT("Actor"),
+			TEXT("ActorComponent"),
+			TEXT("AIController"),
 			TEXT("AimOffsetBlendSpace"),
 			TEXT("AimOffsetBlendSpace1D"),
 			TEXT("AmbientSound"),
@@ -4942,12 +4936,14 @@ void FSlateEditorStyle::FStyle::SetupClassIconsAndThumbnails()
 			TEXT("CameraActor"),
 			TEXT("CameraAnim"),
 			TEXT("CameraBlockingVolume"),
+			TEXT("Character"),
 			TEXT("Class"),
 			TEXT("CullDistanceVolume"),
 			TEXT("CurveBase"),
 			TEXT("DataAsset"),
 			TEXT("DecalActor"),
 			TEXT("Default"),
+			TEXT("DefaultPawn"),
 			TEXT("DirectionalLight"),
 			TEXT("DirectionalLightMovable"),
 			TEXT("DirectionalLightStatic"),
@@ -4956,6 +4952,9 @@ void FSlateEditorStyle::FStyle::SetupClassIconsAndThumbnails()
 			TEXT("ExponentialHeightFog"),
 			TEXT("Font"),
 			TEXT("ForceFeedbackEffect"),
+			TEXT("GameMode"),
+			TEXT("GameState"),
+			TEXT("HUD"),
 			TEXT("InterpData"),
 			TEXT("KillZVolume"),
 			TEXT("Landscape"),
@@ -4982,11 +4981,13 @@ void FSlateEditorStyle::FStyle::SetupClassIconsAndThumbnails()
 			TEXT("PainCausingVolume"),
 			TEXT("PaperCharacter"),
 			TEXT("ParticleSystem"),
+			TEXT("Pawn"),
 			TEXT("PhysicalMaterial"),
 			TEXT("PhysicsConstraintActor"),
 			TEXT("PhysicsThruster"),
 			TEXT("PhysicsVolume"),
 			TEXT("PlaneReflectionCapture"),
+			TEXT("PlayerController"),
 			TEXT("PlayerStart"),
 			TEXT("PointLight"),
 			TEXT("PostProcessVolume"),
@@ -4997,6 +4998,7 @@ void FSlateEditorStyle::FStyle::SetupClassIconsAndThumbnails()
 			TEXT("ReverbEffect"),
 			TEXT("SceneCapture2D"),
 			TEXT("SceneCaptureCube"),
+			TEXT("SceneComponent"),
 			TEXT("SkyLight"),
 			TEXT("SkeletalMesh"),
 			TEXT("SlateBrushAsset"),
diff --git a/Engine/Source/Editor/UnrealEd/Classes/Preferences/UnrealEdOptions.h b/Engine/Source/Editor/UnrealEd/Classes/Preferences/UnrealEdOptions.h
--- a/Engine/Source/Editor/UnrealEd/Classes/Preferences/UnrealEdOptions.h
+++ b/Engine/Source/Editor/UnrealEd/Classes/Preferences/UnrealEdOptions.h
@@ -60,10 +60,6 @@ struct FClassPickerDefaults
 	UPROPERTY()
 	FString LocTextDescriptionID;
 
-	/** The image to show for the class in the picker */
-	UPROPERTY()
-	FName Image;
-
 	/** The name of the asset type being created */
 	UPROPERTY()
 	FString AssetClass;
diff --git a/Engine/Source/Editor/UnrealEd/Private/ClassIconFinder.cpp b/Engine/Source/Editor/UnrealEd/Private/ClassIconFinder.cpp
--- a/Engine/Source/Editor/UnrealEd/Private/ClassIconFinder.cpp
+++ b/Engine/Source/Editor/UnrealEd/Private/ClassIconFinder.cpp
@@ -136,7 +136,7 @@ FName FClassIconFinder::FindIconNameImpl(const UClass* InClass, const FName& InD
 	{
 		// walk up class hierarchy until we find an icon
 		const UClass* CurrentClass = InClass;
-		while( Brush == NULL && CurrentClass && (CurrentClass != AActor::StaticClass()) )
+		while( Brush == NULL && CurrentClass )
 		{
 			BrushName = *FString::Printf( TEXT( "%s.%s" ), StyleRoot, *CurrentClass->GetName() );
 			Brush = FClassIconFinder::LookupBrush( BrushName );
diff --git a/Engine/Source/Editor/UnrealEd/Private/SClassPickerDialog.cpp b/Engine/Source/Editor/UnrealEd/Private/SClassPickerDialog.cpp
--- a/Engine/Source/Editor/UnrealEd/Private/SClassPickerDialog.cpp
+++ b/Engine/Source/Editor/UnrealEd/Private/SClassPickerDialog.cpp
@@ -6,6 +6,7 @@
 #include "SClassPickerDialog.h"
 #include "EditorClassUtils.h"
 #include "SExpandableArea.h"
+#include "ClassIconFinder.h"
 
 #define LOCTEXT_NAMESPACE "SClassPicker"
 
@@ -153,6 +154,7 @@ TSharedRef<ITableRow> SClassPickerDialog::GenerateListRow(TSharedPtr<FClassPicke
 {
 	FClassPickerDefaults* Obj = InItem.Get();
 	UClass* ItemClass = LoadClass<UObject>(NULL, *Obj->ClassName, NULL, LOAD_None, NULL);
+	const FSlateBrush* ItemBrush = FClassIconFinder::FindIconForClass(ItemClass);
 
 	return 
 	SNew(STableRow< TSharedPtr<FClassPickerDefaults> >, OwnerTable)
@@ -178,7 +180,7 @@ TSharedRef<ITableRow> SClassPickerDialog::GenerateListRow(TSharedPtr<FClassPicke
 					.FillWidth(0.12f)
 					[
 						SNew(SImage)
-						.Image(FEditorStyle::GetBrush(Obj->Image))
+						.Image(ItemBrush)
 					]
 					+SHorizontalBox::Slot()
 					.VAlign(VAlign_Center)
