[INTEGRATE] Change 2415538 by [REDACTED_PERSON][REDACTED_HANDLE] on 2015/01/22 13:50:23

	Add placeholder icons for EditorTutorial, PhysicsAsset and Skeleton

diff --git a/Engine/Source/Editor/EditorStyle/Private/SlateEditorStyle.cpp b/Engine/Source/Editor/EditorStyle/Private/SlateEditorStyle.cpp
--- a/Engine/Source/Editor/EditorStyle/Private/SlateEditorStyle.cpp
+++ b/Engine/Source/Editor/EditorStyle/Private/SlateEditorStyle.cpp
@@ -4969,6 +4969,7 @@ void FSlateEditorStyle::FStyle::SetupClassIconsAndThumbnails()
 			TEXT("DirectionalLightStatic"),
 			TEXT("DirectionalLightStationary"),
 			TEXT("DocumentationActor"),
+			TEXT("EditorTutorial"),
 			TEXT("ExponentialHeightFog"),
 			TEXT("Font"),
 			TEXT("ForceFeedbackEffect"),
@@ -5003,6 +5004,7 @@ void FSlateEditorStyle::FStyle::SetupClassIconsAndThumbnails()
 			TEXT("ParticleSystem"),
 			TEXT("Pawn"),
 			TEXT("PhysicalMaterial"),
+			TEXT("PhysicsAsset"),
 			TEXT("PhysicsConstraintActor"),
 			TEXT("PhysicsThruster"),
 			TEXT("PhysicsVolume"),
@@ -5021,6 +5023,7 @@ void FSlateEditorStyle::FStyle::SetupClassIconsAndThumbnails()
 			TEXT("SceneComponent"),
 			TEXT("SkyLight"),
 			TEXT("SkeletalMesh"),
+			TEXT("Skeleton"),
 			TEXT("SlateBrushAsset"),
 			TEXT("SlateWidgetStyleAsset"),
 			TEXT("SoundAttenuation"),
