[INTEGRATE] Change 2420581 by [REDACTED_PERSON][REDACTED_HANDLE] on 2015/01/27 12:21:03

	Added Slate style entries for the new icons

	UE-7845 - New (large) icons for important classes in Browser

diff --git a/Engine/Source/Editor/EditorStyle/Private/SlateEditorStyle.cpp b/Engine/Source/Editor/EditorStyle/Private/SlateEditorStyle.cpp
--- a/Engine/Source/Editor/EditorStyle/Private/SlateEditorStyle.cpp
+++ b/Engine/Source/Editor/EditorStyle/Private/SlateEditorStyle.cpp
@@ -4935,6 +4935,12 @@ void FSlateEditorStyle::FStyle::SetupClassIconsAndThumbnails()
 			Set("ClassIcon.PaperSpriteAtlas", new IMAGE_BRUSH("Icons/AssetIcons/Paper2DSpriteAtlasGroup_16x", Icon16x16));
 			Set("ClassThumbnail.PaperSpriteAtlas", new IMAGE_BRUSH("Icons/AssetIcons/Paper2DSpriteAtlasGroup_64x", Icon64x64));
 
+			Set("ClassIcon.PaperTileMap", new IMAGE_BRUSH("Icons/AssetIcons/TileMap_16x", Icon16x16));
+			Set("ClassThumbnail.PaperTileMap", new IMAGE_BRUSH("Icons/AssetIcons/TileMap_64x", Icon64x64));
+
+			Set("ClassIcon.PaperTileSet", new IMAGE_BRUSH("Icons/AssetIcons/TileMap_16x", Icon16x16));
+			Set("ClassThumbnail.PaperTileSet", new IMAGE_BRUSH("Icons/AssetIcons/TileMap_64x", Icon64x64));
+
 			// ClassIcon and ClassThumbnail for tile sets and/or terrain materials?  Might not be needed/feasible due to them being a UDataAsset
 			//Set("ClassIcon.PaperTerrainMaterial", new IMAGE_BRUSH("Icons/AssetIcons/Paper2DSpriteAtlasGroup_16x", Icon16x16));
 			//Set("ClassThumbnail.PaperTerrainMaterial", new IMAGE_BRUSH("Icons/AssetIcons/Paper2DSpriteAtlasGroup_64x", Icon64x64));
@@ -4980,15 +4986,19 @@ void FSlateEditorStyle::FStyle::SetupClassIconsAndThumbnails()
 			TEXT("CullDistanceVolume"),
 			TEXT("CurveBase"),
 			TEXT("DataAsset"),
+			TEXT("DataTable"),
 			TEXT("DecalActor"),
 			TEXT("Default"),
 			TEXT("DefaultPawn"),
+			TEXT("DialogueWave"),
+			TEXT("DialogueVoice"),
 			TEXT("DirectionalLight"),
 			TEXT("DirectionalLightMovable"),
 			TEXT("DirectionalLightStatic"),
 			TEXT("DirectionalLightStationary"),
 			TEXT("DocumentationActor"),
 			TEXT("EditorTutorial"),
+			TEXT("EnvQuery"),
 			TEXT("ExponentialHeightFog"),
 			TEXT("Font"),
 			TEXT("ForceFeedbackEffect"),
@@ -5055,6 +5065,7 @@ void FSlateEditorStyle::FStyle::SetupClassIconsAndThumbnails()
 			TEXT("SpotLightStatic"),
 			TEXT("SpotLightStationary"),
 			TEXT("StaticMesh"),
+			TEXT("SubsurfaceProfile"),
 			TEXT("TargetPoint"),
 			TEXT("TextRenderActor"),
 			TEXT("Texture2D"),
