[INTEGRATE] Change 2420649 by [REDACTED_PERSON][REDACTED_HANDLE] on 2015/01/27 12:55:26

	Add more icons (mostly by duplicating existing ones) for component classes

diff --git a/Engine/Source/Editor/EditorStyle/Private/SlateEditorStyle.cpp b/Engine/Source/Editor/EditorStyle/Private/SlateEditorStyle.cpp
--- a/Engine/Source/Editor/EditorStyle/Private/SlateEditorStyle.cpp
+++ b/Engine/Source/Editor/EditorStyle/Private/SlateEditorStyle.cpp
@@ -4921,6 +4921,8 @@ void FSlateEditorStyle::FStyle::SetupClassIconsAndThumbnails()
 			Set("ClassIcon.PaperTerrainComponent", new IMAGE_BRUSH("Icons/ActorIcons/Paper2DTerrainActor_16x", Icon16x16));
 
 			Set("ClassIcon.PaperTileMapActor", new IMAGE_BRUSH("Icons/ActorIcons/Paper2DTileMapActor_16x", Icon16x16));
+			Set("ClassIcon.PaperTileMapComponent", new IMAGE_BRUSH("Icons/ActorIcons/Paper2DTileMapComponent_16x", Icon16x16));
+			Set("ClassThumbnail.PaperTileMapComponent", new IMAGE_BRUSH("Icons/ActorIcons/Paper2DTileMapComponent_64x", Icon64x64));
 
 			Set("ClassIcon.PaperSpriteActor", new IMAGE_BRUSH("Icons/AssetIcons/Paper2DSprite_16x", Icon16x16));
 			Set("ClassIcon.PaperSpriteComponent", new IMAGE_BRUSH("Icons/AssetIcons/Paper2DSprite_16x", Icon16x16));
@@ -4978,6 +4980,7 @@ void FSlateEditorStyle::FStyle::SetupClassIconsAndThumbnails()
 			TEXT("BoxReflectionCapture"),
 			TEXT("ButtonStyleAsset"),
 			TEXT("CableActor"),
+			TEXT("CableComponent"),
 			TEXT("CameraActor"),
 			TEXT("CameraAnim"),
 			TEXT("CameraBlockingVolume"),
@@ -5030,12 +5033,15 @@ void FSlateEditorStyle::FStyle::SetupClassIconsAndThumbnails()
 			TEXT("ObjectLibrary"),
 			TEXT("PainCausingVolume"),
 			TEXT("PaperCharacter"),
+			TEXT("PaperTileMapComponent"),
 			TEXT("ParticleSystem"),
 			TEXT("Pawn"),
 			TEXT("PhysicalMaterial"),
 			TEXT("PhysicsAsset"),
 			TEXT("PhysicsConstraintActor"),
+			TEXT("PhysicsConstraintComponent"),
 			TEXT("PhysicsThruster"),
+			TEXT("PhysicsThrusterComponent"),
 			TEXT("PhysicsVolume"),
 			TEXT("PlaneReflectionCapture"),
 			TEXT("PlayerController"),
@@ -5045,12 +5051,14 @@ void FSlateEditorStyle::FStyle::SetupClassIconsAndThumbnails()
 			TEXT("PrecomputedVisibilityOverrideVolume"),
 			TEXT("PrecomputedVisibilityVolume"),
 			TEXT("RadialForceActor"),
+			TEXT("RadialForceComponent"),
 			TEXT("ReflectionCapture"),
 			TEXT("ReverbEffect"),
 			TEXT("SceneCapture2D"),
 			TEXT("SceneCaptureCube"),
 			TEXT("SceneComponent"),
 			TEXT("SkyLight"),
+			TEXT("SkyLightComponent"),
 			TEXT("SkeletalMesh"),
 			TEXT("Skeleton"),
 			TEXT("SlateBrushAsset"),
