[INTEGRATE] Change 2420716 by [REDACTED_PERSON][REDACTED_HANDLE] on 2015/01/27 13:29:56

	Fix PaperTileMapComponent icon
	Add warning comments to AbilitySystemComponent and GameplayAbilityBlueprint

diff --git a/Engine/Source/Editor/EditorStyle/Private/SlateEditorStyle.cpp b/Engine/Source/Editor/EditorStyle/Private/SlateEditorStyle.cpp
--- a/Engine/Source/Editor/EditorStyle/Private/SlateEditorStyle.cpp
+++ b/Engine/Source/Editor/EditorStyle/Private/SlateEditorStyle.cpp
@@ -4937,6 +4937,9 @@ void FSlateEditorStyle::FStyle::SetupClassIconsAndThumbnails()
 			Set("ClassIcon.PaperSpriteAtlas", new IMAGE_BRUSH("Icons/AssetIcons/Paper2DSpriteAtlasGroup_16x", Icon16x16));
 			Set("ClassThumbnail.PaperSpriteAtlas", new IMAGE_BRUSH("Icons/AssetIcons/Paper2DSpriteAtlasGroup_64x", Icon64x64));
 
+			Set("ClassIcon.PaperCharacter", new IMAGE_BRUSH("Icons/AssetIcons/PaperCharacter_16x", Icon16x16));
+			Set("ClassThumbnail.PaperCharacter", new IMAGE_BRUSH("Icons/AssetIcons/PaperCharacter_64x", Icon64x64));
+
 			Set("ClassIcon.PaperTileMap", new IMAGE_BRUSH("Icons/AssetIcons/TileMap_16x", Icon16x16));
 			Set("ClassThumbnail.PaperTileMap", new IMAGE_BRUSH("Icons/AssetIcons/TileMap_64x", Icon64x64));
 
@@ -5032,8 +5035,6 @@ void FSlateEditorStyle::FStyle::SetupClassIconsAndThumbnails()
 			TEXT("Note"),
 			TEXT("ObjectLibrary"),
 			TEXT("PainCausingVolume"),
-			TEXT("PaperCharacter"),
-			TEXT("PaperTileMapComponent"),
 			TEXT("ParticleSystem"),
 			TEXT("Pawn"),
 			TEXT("PhysicalMaterial"),
diff --git a/Engine/Source/Runtime/GameplayAbilities/Classes/GameplayAbilityBlueprint.h b/Engine/Source/Runtime/GameplayAbilities/Classes/GameplayAbilityBlueprint.h
--- a/Engine/Source/Runtime/GameplayAbilities/Classes/GameplayAbilityBlueprint.h
+++ b/Engine/Source/Runtime/GameplayAbilities/Classes/GameplayAbilityBlueprint.h
@@ -7,6 +7,7 @@
 
 /**
  * A Gameplay Ability Blueprint is essentially a specialized Blueprint whose graphs control a gameplay ability.
+ *	NOTE: This feature is EXPERIMENTAL. Use at your own risk!
  */
 
 UCLASS(BlueprintType)
diff --git a/Engine/Source/Runtime/GameplayAbilities/Public/AbilitySystemComponent.h b/Engine/Source/Runtime/GameplayAbilities/Public/AbilitySystemComponent.h
--- a/Engine/Source/Runtime/GameplayAbilities/Public/AbilitySystemComponent.h
+++ b/Engine/Source/Runtime/GameplayAbilities/Public/AbilitySystemComponent.h
@@ -43,6 +43,7 @@ DECLARE_MULTICAST_DELEGATE_OneParam(FTargetingRejectedConfirmation, int32);
 
 /**
  *	The core ActorComponent for interfacing with the GameplayAbilities System
+ *	NOTE: This feature is EXPERIMENTAL. Use at your own risk!
  */
 UCLASS(ClassGroup=AbilitySystem, hidecategories=(Object,LOD,Lighting,Transform,Sockets,TextureStreaming), editinlinenew, meta=(BlueprintSpawnableComponent))
 class GAMEPLAYABILITIES_API UAbilitySystemComponent : public UActorComponent, public IGameplayTagAssetInterface
