[INTEGRATE] Change 2426456 by [REDACTED_PERSON][REDACTED_HANDLE] on 2015/01/30 17:19:29

	UMG - Changing the asset icon.

diff --git a/Engine/Source/Editor/EditorStyle/Private/SlateEditorStyle.cpp b/Engine/Source/Editor/EditorStyle/Private/SlateEditorStyle.cpp
--- a/Engine/Source/Editor/EditorStyle/Private/SlateEditorStyle.cpp
+++ b/Engine/Source/Editor/EditorStyle/Private/SlateEditorStyle.cpp
@@ -5096,6 +5096,7 @@ void FSlateEditorStyle::FStyle::SetupClassIconsAndThumbnails()
 			TEXT("TouchInterface"),
 			TEXT("UserDefinedEnum"),
 			TEXT("UserDefinedStruct"),
+			TEXT("WidgetBlueprint"),
 			TEXT("WindDirectionalSource"),
 			TEXT("World")
 		};
