[INTEGRATE] Change 2427079 by [REDACTED_PERSON][REDACTED_HANDLE] on 2015/01/31 12:30:41

	Added missing file

diff --git a/Engine/Source/Editor/EditorStyle/Private/SlateEditorStyle.cpp b/Engine/Source/Editor/EditorStyle/Private/SlateEditorStyle.cpp
--- a/Engine/Source/Editor/EditorStyle/Private/SlateEditorStyle.cpp
+++ b/Engine/Source/Editor/EditorStyle/Private/SlateEditorStyle.cpp
@@ -772,36 +772,7 @@ void FSlateEditorStyle::FStyle::SetupGeneralStyles()
 			.SetInactiveHoveredBrush(FSlateNoResource())
 			);
 
-		//Set("DetailsView.TreeView.TableRow", FTableRowStyle()
-		//	.SetEvenRowBackgroundBrush(FSlateNoResource())
-		//	.SetEvenRowBackgroundHoveredBrush(FSlateNoResource())
-		//	.SetOddRowBackgroundBrush(FSlateNoResource())
-		//	.SetOddRowBackgroundHoveredBrush(FSlateNoResource())
-		//	.SetSelectorFocusedBrush(FSlateNoResource())
-		//	.SetActiveBrush(FSlateNoResource())
-		//	.SetActiveHoveredBrush(FSlateNoResource())
-		//	.SetInactiveBrush(FSlateNoResource())
-		//	.SetInactiveHoveredBrush(FSlateNoResource())
-		//	.SetTextColor(DefaultForeground)
-		//	.SetSelectedTextColor(InvertedForeground)
-		//	);
-
-		//Set("DetailsView.CategoryMiddle", new IMAGE_BRUSH("PropertyView/DetailCategoryMiddle", FVector2D(16, 16)));
-		//Set("DetailsView.CategoryMiddle_Hovered", new IMAGE_BRUSH("PropertyView/DetailCategoryMiddle_Hovered", FVector2D(16, 16)));
-		//Set("DetailsView.CategoryMiddle_Highlighted", new BOX_BRUSH("Common/TextBox_Special_Active", FMargin(8.0f / 32.0f)));
-
-		//if ( IsHighlighted() )
-		//{
-		//	return FEditorStyle::GetBrush("DetailsView.CategoryMiddle_Highlighted");
-		//}
-		//else if ( IsHovered() )
-		//{
-		//	return FEditorStyle::GetBrush("DetailsView.CategoryMiddle_Hovered");
-		//}
-		//else
-		//{
-		//	return FEditorStyle::GetBrush("DetailsView.CategoryMiddle");
-		//}
+
 
 		Set( "TreeArrow_Collapsed", new IMAGE_BRUSH( "Common/TreeArrow_Collapsed", Icon10x10, DefaultForeground ) );
 		Set( "TreeArrow_Collapsed_Hovered", new IMAGE_BRUSH( "Common/TreeArrow_Collapsed_Hovered", Icon10x10, DefaultForeground ) );
@@ -5143,7 +5114,11 @@ void FSlateEditorStyle::FStyle::SetupClassIconsAndThumbnails()
 			TEXT("UserDefinedStruct"),
 			TEXT("WidgetBlueprint"),
 			TEXT("WindDirectionalSource"),
-			TEXT("World")
+			TEXT("World"),
+			TEXT("Cube"),
+			TEXT("Sphere"),
+			TEXT("Cylinder"),
+			TEXT("Cone")
 		};
 
 		const TCHAR* Path = TEXT("Icons/AssetIcons/%s_%dx");
