fix for UE-7316 PhAT Simulate button icon is missing with Small Icons enabled

diff --git a/Engine/Source/Editor/EditorStyle/Private/SlateEditorStyle.cpp b/Engine/Source/Editor/EditorStyle/Private/SlateEditorStyle.cpp
--- a/Engine/Source/Editor/EditorStyle/Private/SlateEditorStyle.cpp
+++ b/Engine/Source/Editor/EditorStyle/Private/SlateEditorStyle.cpp
@@ -5843,6 +5843,9 @@ void FSlateEditorStyle::FStyle::SetupToolkitStyles()
 		Set( "PhAT.ConvertToSkeletal", new IMAGE_BRUSH( "Icons/icon_PhAT_Skeletal_40x", Icon40x40 ) );
 		Set( "PhAT.DeleteConstraint", new IMAGE_BRUSH( "Icons/icon_PhAT_DeleteConstraint_40x", Icon40x40 ) );
 
+		Set("PhAT.SimulationNormal.Small", new IMAGE_BRUSH("Icons/icon_PhAT_PlaySim_40x", Icon20x20));
+		Set("PhAT.SimulationNoGravity.Small", new IMAGE_BRUSH("Icons/icon_PhAT_PlaySimNoGravity_40x", Icon20x20));
+		Set("PhAT.ToggleSelectedSimulation.Small", new IMAGE_BRUSH("Icons/icon_PhAT_PlaySimSelected_40x", Icon20x20));
 		Set( "PhAT.ToggleSimulation.Small", new IMAGE_BRUSH( "Icons/icon_PhAT_PlaySim_40x", Icon20x20 ) );
 		Set( "PhAT.Undo.Small", new IMAGE_BRUSH( "Icons/icon_Generic_Undo_40x", Icon20x20 ) );
 		Set( "PhAT.Redo.Small", new IMAGE_BRUSH( "Icons/icon_Generic_Redo_40x", Icon20x20 ) );
