[INTEGRATE] Change 2432183 by [REDACTED_PERSON][REDACTED_HANDLE] on 2015/02/04 14:42:39

	Fixed subdued text in white tool-tips being a bit too dark

diff --git a/Engine/Source/Editor/EditorStyle/Private/SlateEditorStyle.cpp b/Engine/Source/Editor/EditorStyle/Private/SlateEditorStyle.cpp
--- a/Engine/Source/Editor/EditorStyle/Private/SlateEditorStyle.cpp
+++ b/Engine/Source/Editor/EditorStyle/Private/SlateEditorStyle.cpp
@@ -2530,7 +2530,7 @@ void FSlateEditorStyle::FStyle::SetupTutorialStyles()
 
 		const FTextBlockStyle DocumentationTooltipTextSubdued = FTextBlockStyle( NormalText )
 			.SetFont( TTF_CORE_FONT( "Fonts/Roboto-Regular", 8 ) )
-			.SetColorAndOpacity( FLinearColor( 0.2f, 0.2f, 0.2f ) );
+			.SetColorAndOpacity( FLinearColor( 0.1f, 0.1f, 0.1f ) );
 		Set("Documentation.SDocumentationTooltipSubdued", FTextBlockStyle(DocumentationTooltipTextSubdued));
 
 		const FTextBlockStyle DocumentationTooltipHyperlinkText = FTextBlockStyle( NormalText )
@@ -3954,13 +3954,13 @@ void FSlateEditorStyle::FStyle::SetupGraphEditorStyles()
 
 		Set("SCSEditor.ComponentTooltip.Label",
 			FTextBlockStyle(NormalText)
-			.SetColorAndOpacity(FLinearColor(0.2f, 0.2f, 0.2f))
+			.SetColorAndOpacity(FLinearColor(0.075f, 0.075f, 0.075f))
 			.SetShadowOffset(FVector2D(1.0f, 1.0f))
 			.SetShadowColorAndOpacity(FLinearColor(0.8f, 0.8f, 0.8f))
 		);
 		Set("SCSEditor.ComponentTooltip.ImportantLabel",
 			FTextBlockStyle(NormalText)
-			.SetColorAndOpacity(FLinearColor(0.2f, 0.2f, 0.2f))
+			.SetColorAndOpacity(FLinearColor(0.05f, 0.05f, 0.05f))
 			.SetShadowOffset(FVector2D(1.0f, 1.0f))
 			.SetShadowColorAndOpacity(FLinearColor(0.8f, 0.8f, 0.8f))
 		);
@@ -3984,7 +3984,7 @@ void FSlateEditorStyle::FStyle::SetupGraphEditorStyles()
 		Set("SCSEditor.ComponentTooltip.ClassDescription",
 			FTextBlockStyle(NormalText)
 			.SetFont(TTF_CORE_FONT("Fonts/Roboto-Italic", 10))
-			.SetColorAndOpacity(FLinearColor(0.2f, 0.2f, 0.2f))
+			.SetColorAndOpacity(FLinearColor(0.1f, 0.1f, 0.1f))
 			.SetShadowOffset(FVector2D(1.0f, 1.0f))
 			.SetShadowColorAndOpacity(FLinearColor(0.8f, 0.8f, 0.8f))
 		);
