Added a class icon and thumbnail for the 'UInterface' class

UE-9423 - BP interfaces/libraries do not use correct icon in Content Browser

This is currently just a copy of the icon we use for Blueprint interfaces.

diff --git a/Engine/Source/Editor/EditorStyle/Private/SlateEditorStyle.cpp b/Engine/Source/Editor/EditorStyle/Private/SlateEditorStyle.cpp
--- a/Engine/Source/Editor/EditorStyle/Private/SlateEditorStyle.cpp
+++ b/Engine/Source/Editor/EditorStyle/Private/SlateEditorStyle.cpp
@@ -5052,6 +5052,7 @@ void FSlateEditorStyle::FStyle::SetupClassIconsAndThumbnails()
 			TEXT("GameMode"),
 			TEXT("GameState"),
 			TEXT("HUD"),
+			TEXT("Interface"),
 			TEXT("InterpData"),
 			TEXT("KillZVolume"),
 			TEXT("Landscape"),
