[INTEGRATE] Change 2435461 by [REDACTED_PERSON][REDACTED_HANDLE] on 2015/02/06 10:41:49

	Editor - Allowing an even smaller thumbnail size in asset pickers and content browser.  Added a new font size that kicks in for the very small thumbnails.

diff --git a/Engine/Source/Editor/ContentBrowser/Private/AssetViewWidgets.cpp b/Engine/Source/Editor/ContentBrowser/Private/AssetViewWidgets.cpp
--- a/Engine/Source/Editor/ContentBrowser/Private/AssetViewWidgets.cpp
+++ b/Engine/Source/Editor/ContentBrowser/Private/AssetViewWidgets.cpp
@@ -1302,14 +1302,19 @@ FSlateFontInfo SAssetTileItem::GetThumbnailFont() const
 	if ( ThumbSize.IsSet() )
 	{
 		float Size = ThumbSize.Get();
-		if ( Size < 85 )
+		if ( Size < 50 )
 		{
-			static FName SmallFontName("ContentBrowser.AssetTileViewNameFontSmall");
+			const static FName SmallFontName("ContentBrowser.AssetTileViewNameFontVerySmall");
+			return FEditorStyle::GetFontStyle(SmallFontName);
+		}
+		else if ( Size < 85 )
+		{
+			const static FName SmallFontName("ContentBrowser.AssetTileViewNameFontSmall");
 			return FEditorStyle::GetFontStyle(SmallFontName);
 		}
 	}
 
-	static FName RegularFont("ContentBrowser.AssetTileViewNameFont");
+	const static FName RegularFont("ContentBrowser.AssetTileViewNameFont");
 	return FEditorStyle::GetFontStyle(RegularFont);
 }
 
diff --git a/Engine/Source/Editor/EditorStyle/Private/SlateEditorStyle.cpp b/Engine/Source/Editor/EditorStyle/Private/SlateEditorStyle.cpp
--- a/Engine/Source/Editor/EditorStyle/Private/SlateEditorStyle.cpp
+++ b/Engine/Source/Editor/EditorStyle/Private/SlateEditorStyle.cpp
@@ -3877,10 +3877,10 @@ void FSlateEditorStyle::FStyle::SetupGraphEditorStyles()
 
 		Set("SCSEditor.ComponentTooltip.Value", 
 			FTextBlockStyle(NormalText)
- 			.SetFont(TTF_CORE_FONT("Fonts/Roboto-Bold", 10))
+			.SetFont(TTF_CORE_FONT("Fonts/Roboto-Bold", 10))
 			.SetColorAndOpacity(FLinearColor(0.0f, 0.0f, 0.0f))
- 			.SetShadowOffset(FVector2D(1.0f, 1.0f))
- 			.SetShadowColorAndOpacity(FLinearColor(0.8f, 0.8f, 0.8f))
+			.SetShadowOffset(FVector2D(1.0f, 1.0f))
+			.SetShadowColorAndOpacity(FLinearColor(0.8f, 0.8f, 0.8f))
 		);
 		Set("SCSEditor.ComponentTooltip.ImportantValue",
 			FTextBlockStyle(NormalText)
@@ -5107,6 +5107,7 @@ void FSlateEditorStyle::FStyle::SetupContentBrowserStyle()
 		// Tile view
 		Set( "ContentBrowser.AssetTileViewNameFont", TTF_CORE_FONT( "Fonts/Roboto-Regular", 9 ) );
 		Set( "ContentBrowser.AssetTileViewNameFontSmall", TTF_CORE_FONT( "Fonts/Roboto-Light", 8, EFontHinting::Auto ) );
+		Set( "ContentBrowser.AssetTileViewNameFontVerySmall", TTF_CORE_FONT( "Fonts/Roboto-Light", 7, EFontHinting::Auto ) );
 		Set( "ContentBrowser.AssetTileViewNameFontDirty", TTF_CORE_FONT( "Fonts/Roboto-Bold", 10 ) );
 		Set( "ContentBrowser.AssetListView.TableRow", FTableRowStyle()
 			.SetEvenRowBackgroundBrush( FSlateNoResource() )
