[INTEGRATE] Change 2436874 by [REDACTED_PERSON][REDACTED_HANDLE] on 2015/02/07 15:05:06

	Editor: Add standard colors for Blueprint/Native text references to editor style
	- Update hyperlinks and components tree to use the standard color (native links are now greener)
	- Improve selection behavior for these colorized text fields so they are still readable
	- Remove some dead style entries
	[UE-8717]

diff --git a/Engine/Source/Editor/EditorStyle/Private/SlateEditorStyle.cpp b/Engine/Source/Editor/EditorStyle/Private/SlateEditorStyle.cpp
--- a/Engine/Source/Editor/EditorStyle/Private/SlateEditorStyle.cpp
+++ b/Engine/Source/Editor/EditorStyle/Private/SlateEditorStyle.cpp
@@ -69,6 +69,11 @@ FSlateEditorStyle::FStyle::FStyle( const TWeakObjectPtr< UEditorStyleSettings >&
 	, SelectionColor_Inactive( SelectionColor_Inactive_LinearRef )
 	, SelectionColor_Pressed( SelectionColor_Pressed_LinearRef )
 
+	, InheritedFromBlueprintTextColor(FLinearColor(0.25f, 0.5f, 1.0f))
+	, InheritedFromNativeTextColor(FLinearColor(0.375f, 1.0f, 0.375f))
+	, IntroducedInThisInstanceTextColor(FLinearColor::White)
+	, IntroducedInThisBlueprintTextColor(FLinearColor::White)
+
 	, Settings( InSettings )
 {
 
@@ -1215,15 +1220,6 @@ void FSlateEditorStyle::FStyle::SetupGeneralStyles()
 		Set( "SceneOutliner.TableViewRow", FTableRowStyle(NormalTableRowStyle)
 			.SetInactiveBrush( IMAGE_BRUSH( "Common/Selection", Icon8x8, SelectionColor_Subdued ) )
 		);
-
-		Set( "SceneOutliner.EditBlueprintHyperlinkStyle",
-			FTextBlockStyle( NormalText )
-			.SetFont( TTF_CORE_FONT( "Fonts/Roboto-Regular", 10 ) )
-			.SetColorAndOpacity( FSlateColor::UseForeground() ) );
-		Set( "SceneOutliner.GoToCodeHyperlinkStyle",
-			FTextBlockStyle( NormalText )
-			.SetFont( TTF_CORE_FONT( "Fonts/Roboto-Regular", 10 ) )
-			.SetColorAndOpacity( FSlateColor::UseSubduedForeground() ) );
 	}
 
 	// Socket chooser
@@ -2266,6 +2262,25 @@ void FSlateEditorStyle::FStyle::SetupGeneralStyles()
 			.SetInactiveHoveredBrush(IMAGE_BRUSH("Common/Selection", Icon8x8, SelectionColor_Inactive))
 			);
 	}
+
+
+	// Common styles for blueprint/code references
+	{
+		Set("Common.EditBlueprintHyperlinkStyle", FTextBlockStyle(NormalText)
+			.SetFont(TTF_CORE_FONT("Fonts/Roboto-Regular", 10))
+			.SetColorAndOpacity(InheritedFromBlueprintTextColor)
+		);
+
+		Set("Common.GoToNativeCodeHyperlinkStyle", FTextBlockStyle(NormalText)
+			.SetFont(TTF_CORE_FONT("Fonts/Roboto-Regular", 10))
+			.SetColorAndOpacity(InheritedFromNativeTextColor)
+		);
+
+		Set("Common.InheritedFromBlueprintTextColor", InheritedFromBlueprintTextColor);
+		Set("Common.InheritedFromNativeTextColor", InheritedFromNativeTextColor);
+		Set("Common.IntroducedInThisInstanceTextColor", IntroducedInThisInstanceTextColor);
+		Set("Common.IntroducedInThisBlueprintTextColor", IntroducedInThisBlueprintTextColor);
+	}
 #endif
 }
 
@@ -2999,9 +3014,6 @@ void FSlateEditorStyle::FStyle::SetupPropertyEditorStyles()
 		Set( "DetailsView.AdvancedDropdownBorder.Open", new IMAGE_BRUSH( "Common/ScrollBoxShadowTop", FVector2D(64,8) ) );
 		Set( "DetailsView.CategoryFontStyle", TTF_CORE_FONT( "Fonts/Roboto-Bold", 10 ) );
 
-		Set( "DetailsView.EditBlueprintHyperlinkStyle", FTextBlockStyle(NormalText) .SetFont( TTF_CORE_FONT( "Fonts/Roboto-Regular", 10 ) ).SetColorAndOpacity( FLinearColor( 0.25f, 0.5f, 1.0f ) ) );
-		Set( "DetailsView.GoToCodeHyperlinkStyle", FTextBlockStyle(NormalText) .SetFont( TTF_CORE_FONT( "Fonts/Roboto-Regular", 10 ) ).SetColorAndOpacity( FLinearColor( 0.7f, 0.7f, 0.7f ) ) );
-
 		Set( "DetailsView.TreeView.TableRow", FTableRowStyle()
 			.SetEvenRowBackgroundBrush( FSlateNoResource() )
 			.SetEvenRowBackgroundHoveredBrush( FSlateNoResource() )
@@ -3016,7 +3028,7 @@ void FSlateEditorStyle::FStyle::SetupPropertyEditorStyles()
 			.SetSelectedTextColor( InvertedForeground )
 			);
 	}
-	}
+}
 
 void FSlateEditorStyle::FStyle::SetupProfilerStyle()
 {
diff --git a/Engine/Source/Editor/EditorStyle/Private/SlateEditorStyle.h b/Engine/Source/Editor/EditorStyle/Private/SlateEditorStyle.h
--- a/Engine/Source/Editor/EditorStyle/Private/SlateEditorStyle.h
+++ b/Engine/Source/Editor/EditorStyle/Private/SlateEditorStyle.h
@@ -124,6 +124,12 @@ public:
 		const FSlateColor SelectionColor_Inactive;
 		const FSlateColor SelectionColor_Pressed;
 
+		// These are common colors used thruout the editor in mutliple style elements
+		const FSlateColor InheritedFromBlueprintTextColor;
+		const FSlateColor InheritedFromNativeTextColor;
+		const FSlateColor IntroducedInThisInstanceTextColor;
+		const FSlateColor IntroducedInThisBlueprintTextColor;
+
 		FTextBlockStyle NormalText;
 		FEditableTextBoxStyle NormalEditableTextBoxStyle;
 		FTableRowStyle NormalTableRowStyle;
@@ -141,7 +147,7 @@ public:
 		FCoreUObjectDelegates::OnObjectPropertyChanged.AddSP(NewStyle, &FSlateEditorStyle::FStyle::SettingsChanged);
 
 		return NewStyle;
-}
+	}
 
 	static TSharedPtr< FSlateEditorStyle::FStyle > StyleInstance;
 	static TWeakObjectPtr< UEditorStyleSettings > Settings;
diff --git a/Engine/Source/Editor/GameProjectGeneration/Private/SNewClassDialog.cpp b/Engine/Source/Editor/GameProjectGeneration/Private/SNewClassDialog.cpp
--- a/Engine/Source/Editor/GameProjectGeneration/Private/SNewClassDialog.cpp
+++ b/Engine/Source/Editor/GameProjectGeneration/Private/SNewClassDialog.cpp
@@ -429,7 +429,7 @@ void SNewClassDialog::Construct( const FArguments& InArgs )
 							[
 								SNew(SHyperlink)
 								.Style(FCoreStyle::Get(), "Hyperlink")
-								.TextStyle(FEditorStyle::Get(), "DetailsView.GoToCodeHyperlinkStyle")
+								.TextStyle(FEditorStyle::Get(), "Common.GoToNativeCodeHyperlinkStyle")
 								.OnNavigate(this, &SNewClassDialog::OnEditCodeClicked)
 								.Text(this, &SNewClassDialog::GetSelectedParentClassFilename)
 								.ToolTipText(FText::Format(LOCTEXT("GoToCode_ToolTip", "Click to open this source file in {0}"), FSourceCodeNavigation::GetSuggestedSourceCodeIDE()))
diff --git a/Engine/Source/Editor/Kismet/Private/BlueprintEditor.cpp b/Engine/Source/Editor/Kismet/Private/BlueprintEditor.cpp
--- a/Engine/Source/Editor/Kismet/Private/BlueprintEditor.cpp
+++ b/Engine/Source/Editor/Kismet/Private/BlueprintEditor.cpp
@@ -1776,7 +1776,7 @@ void FBlueprintEditor::PostRegenerateMenusAndToolbars()
 			[
 				SNew(SHyperlink)
 				.Style(FEditorStyle::Get(), "EditBPHyperlink")
-				.TextStyle(FEditorStyle::Get(), "DetailsView.EditBlueprintHyperlinkStyle")
+				.TextStyle(FEditorStyle::Get(), "Common.EditBlueprintHyperlinkStyle")
 				.IsEnabled(this, &FBlueprintEditor::IsNativeParentClassCodeLinkEnabled)
 				.Visibility(this, &FBlueprintEditor::GetNativeParentClassButtonsVisibility)
 				.OnNavigate(this, &FBlueprintEditor::OnEditParentClassNativeCodeClicked)
diff --git a/Engine/Source/Editor/Kismet/Private/BlueprintEditorTabFactories.cpp b/Engine/Source/Editor/Kismet/Private/BlueprintEditorTabFactories.cpp
--- a/Engine/Source/Editor/Kismet/Private/BlueprintEditorTabFactories.cpp
+++ b/Engine/Source/Editor/Kismet/Private/BlueprintEditorTabFactories.cpp
@@ -189,7 +189,7 @@ TSharedRef<SWidget> FDefaultsEditorSummoner::CreateOptionalDataOnlyMessage() con
 					[
 						SNew(SHyperlink)
 						.Style(FEditorStyle::Get(), "EditBPHyperlink")
-						.TextStyle(FEditorStyle::Get(), "DetailsView.EditBlueprintHyperlinkStyle")
+						.TextStyle(FEditorStyle::Get(), "Common.EditBlueprintHyperlinkStyle")
 						.OnNavigate(this, &FDefaultsEditorSummoner::OnChangeBlueprintToNotDataOnly)
 						.Text(LOCTEXT("FullEditor", "Open Full Blueprint Editor"))
 						.ToolTipText(LOCTEXT("FullEditorToolTip", "This opens the blueprint in the full editor."))
diff --git a/Engine/Source/Editor/Kismet/Private/SSCSEditor.cpp b/Engine/Source/Editor/Kismet/Private/SSCSEditor.cpp
--- a/Engine/Source/Editor/Kismet/Private/SSCSEditor.cpp
+++ b/Engine/Source/Editor/Kismet/Private/SSCSEditor.cpp
@@ -1712,30 +1712,26 @@ FSlateColor SSCS_RowWidget::GetColorTintForIcon() const
 
 FSlateColor SSCS_RowWidget::GetColorTintForText() const
 {
-	const FLinearColor InheritedBlueprintComponentColor(0.6f, 0.6f, 0.8f);
-	const FLinearColor InstancedInheritedBlueprintComponentColor(InheritedBlueprintComponentColor);
-	const FLinearColor InheritedNativeComponentColor(0.7f, 0.9f, 0.7f);
-	const FLinearColor IntroducedHereColor(FLinearColor::White);
+	const FSlateColor SelectedTextColor(FLinearColor::Black);
 
 	FSCSEditorTreeNodePtrType NodePtr = GetNode();
 	if (NodePtr->IsInherited())
 	{
 		if (NodePtr->IsNative())
 		{
-			return InheritedNativeComponentColor;
-		}
-		else if (NodePtr->IsInstanced())
-		{
-			return InstancedInheritedBlueprintComponentColor;
+			static FSlateColor NativeColor = FEditorStyle::GetSlateColor("Common.InheritedFromNativeTextColor");
+			return IsSelected() ? SelectedTextColor : NativeColor;
 		}
 		else
 		{
-			return InheritedBlueprintComponentColor;
+			static FSlateColor BlueprintColor = FEditorStyle::GetSlateColor("Common.InheritedFromBlueprintTextColor");
+			return IsSelected() ? SelectedTextColor : BlueprintColor;
 		}
 	}
 	else
 	{
-		return IntroducedHereColor;
+		static FSlateColor ThisInstanceColor = FEditorStyle::GetSlateColor("Common.IntroducedInThisInstanceTextColor");
+		return IsSelected() ? SelectedTextColor : ThisInstanceColor;
 	}
 }
 
diff --git a/Engine/Source/Editor/SceneOutliner/Private/SceneOutlinerActorInfoColumn.cpp b/Engine/Source/Editor/SceneOutliner/Private/SceneOutlinerActorInfoColumn.cpp
--- a/Engine/Source/Editor/SceneOutliner/Private/SceneOutlinerActorInfoColumn.cpp
+++ b/Engine/Source/Editor/SceneOutliner/Private/SceneOutlinerActorInfoColumn.cpp
@@ -193,11 +193,12 @@ const TSharedRef< SWidget > FActorInfoColumn::ConstructRowWidget( FTreeItemRef T
 		.VAlign(VAlign_Center)
 		.AutoWidth()
 		[
+			// Make sure that the hyperlink shows as black (by multiplying black * desired color) when selected so it is readable against the orange background even if blue/green/etc... normally
 			SNew(SBorder)
 			.BorderImage(FEditorStyle::GetBrush("NoBorder"))
 			.ColorAndOpacity_Static([](TWeakPtr<const STableRow<FTreeItemPtr>> WeakRow)->FLinearColor{
 				auto Row = WeakRow.Pin();
-				return Row.IsValid() && Row->IsSelectedExclusively() ? FLinearColor::Black : FLinearColor::White;
+				return Row.IsValid() && Row->IsSelected() ? FLinearColor::Black : FLinearColor::White;
 			}, TWeakPtr<const STableRow<FTreeItemPtr>>(StaticCastSharedRef<const STableRow<FTreeItemPtr>>(Row.AsShared())))
 			[
 				Hyperlink.ToSharedRef()
diff --git a/Engine/Source/Editor/UnrealEd/Private/EditorClassUtils.cpp b/Engine/Source/Editor/UnrealEd/Private/EditorClassUtils.cpp
--- a/Engine/Source/Editor/UnrealEd/Private/EditorClassUtils.cpp
+++ b/Engine/Source/Editor/UnrealEd/Private/EditorClassUtils.cpp
@@ -109,7 +109,7 @@ TSharedRef<SWidget> FEditorClassUtils::GetSourceLinkFormatted(const UClass* Clas
 
 		SourceHyperlink = SNew(SHyperlink)
 			.Style(FEditorStyle::Get(), "EditBPHyperlink")
-			.TextStyle(FEditorStyle::Get(), "DetailsView.EditBlueprintHyperlinkStyle")
+			.TextStyle(FEditorStyle::Get(), "Common.EditBlueprintHyperlinkStyle")
 			.OnNavigate_Static(&Local::OnEditBlueprintClicked, BlueprintPtr, ObjectWeakPtr)
 			.Text(FText::Format(BlueprintFormat, FText::FromString(Blueprint->GetName())))
 			.ToolTipText(NSLOCTEXT("SourceHyperlink", "EditBlueprint_ToolTip", "Click to edit the blueprint"));
@@ -130,7 +130,7 @@ TSharedRef<SWidget> FEditorClassUtils::GetSourceLinkFormatted(const UClass* Clas
 
 			SourceHyperlink = SNew(SHyperlink)
 				.Style(FCoreStyle::Get(), "Hyperlink")
-				.TextStyle(FEditorStyle::Get(), "DetailsView.GoToCodeHyperlinkStyle")
+				.TextStyle(FEditorStyle::Get(), "Common.GoToNativeCodeHyperlinkStyle")
 				.OnNavigate_Static(&Local::OnEditCodeClicked, ClassHeaderPath)
 				.Text(FText::Format(CodeFormat, FText::FromString(FPaths::GetCleanFilename( *ClassHeaderPath ) ) ) )
 				.ToolTipText(FText::Format(NSLOCTEXT("SourceHyperlink", "GoToCode_ToolTip", "Click to open this source file in {0}"), FSourceCodeNavigation::GetSuggestedSourceCodeIDE()));
