[INTEGRATE] Change 2436956 by [REDACTED_PERSON][REDACTED_HANDLE] on 2015/02/07 16:50:09

	Editor: Change native/BP hyperlink underline colors to be driven by the common style colors as well

diff --git a/Engine/Source/Editor/EditorStyle/Private/SlateEditorStyle.cpp b/Engine/Source/Editor/EditorStyle/Private/SlateEditorStyle.cpp
--- a/Engine/Source/Editor/EditorStyle/Private/SlateEditorStyle.cpp
+++ b/Engine/Source/Editor/EditorStyle/Private/SlateEditorStyle.cpp
@@ -622,18 +622,6 @@ void FSlateEditorStyle::FStyle::SetupGeneralStyles()
 			.SetPadding(FMargin(0.0f));
 		Set("DarkHyperlink", DarkHyperlink);
 
-		// Edit BP Hyperlink
-		FButtonStyle EditBPHyperlinkButton = FButtonStyle()
-			.SetNormal(BORDER_BRUSH("Old/HyperlinkDotted", FMargin(0, 0, 0, 3 / 16.0f), FLinearColor(0.25f, 0.5f, 1.0f)))
-			.SetPressed(FSlateNoResource())
-			.SetHovered(BORDER_BRUSH("Old/HyperlinkUnderline", FMargin(0, 0, 0, 3 / 16.0f), FLinearColor(0.25f, 0.5f, 1.0f)));
-		FHyperlinkStyle EditBPHyperlink = FHyperlinkStyle()
-			.SetUnderlineStyle(EditBPHyperlinkButton)
-			.SetTextStyle(NormalText)
-			.SetPadding(FMargin(0.0f));
-
-		Set("EditBPHyperlink", EditBPHyperlink);
-
 		// Visible on hover hyper link
 		FButtonStyle HoverOnlyHyperlinkButton = FButtonStyle()
 			.SetNormal(FSlateNoResource() )
@@ -2266,20 +2254,49 @@ void FSlateEditorStyle::FStyle::SetupGeneralStyles()
 
 	// Common styles for blueprint/code references
 	{
-		Set("Common.EditBlueprintHyperlinkStyle", FTextBlockStyle(NormalText)
-			.SetFont(TTF_CORE_FONT("Fonts/Roboto-Regular", 10))
-			.SetColorAndOpacity(InheritedFromBlueprintTextColor)
-		);
+		Set("Common.IntroducedInThisInstanceTextColor", IntroducedInThisInstanceTextColor);
+		Set("Common.IntroducedInThisBlueprintTextColor", IntroducedInThisBlueprintTextColor);
+
+		// Inherited from blueprint
+		Set("Common.InheritedFromBlueprintTextColor", InheritedFromBlueprintTextColor);
 
-		Set("Common.GoToNativeCodeHyperlinkStyle", FTextBlockStyle(NormalText)
+		FTextBlockStyle InheritedFromBlueprintTextStyle = FTextBlockStyle(NormalText)
 			.SetFont(TTF_CORE_FONT("Fonts/Roboto-Regular", 10))
-			.SetColorAndOpacity(InheritedFromNativeTextColor)
-		);
+			.SetColorAndOpacity(InheritedFromBlueprintTextColor);
 
-		Set("Common.InheritedFromBlueprintTextColor", InheritedFromBlueprintTextColor);
+		Set("Common.InheritedFromBlueprintTextStyle", InheritedFromBlueprintTextStyle);
+
+		// Go to blueprint hyperlink
+		FButtonStyle EditBPHyperlinkButton = FButtonStyle()
+			.SetNormal(BORDER_BRUSH("Old/HyperlinkDotted", FMargin(0, 0, 0, 3 / 16.0f), InheritedFromBlueprintTextColor))
+			.SetPressed(FSlateNoResource())
+			.SetHovered(BORDER_BRUSH("Old/HyperlinkUnderline", FMargin(0, 0, 0, 3 / 16.0f), InheritedFromBlueprintTextColor));
+		FHyperlinkStyle EditBPHyperlinkStyle = FHyperlinkStyle()
+			.SetUnderlineStyle(EditBPHyperlinkButton)
+			.SetTextStyle(InheritedFromBlueprintTextStyle)
+			.SetPadding(FMargin(0.0f));
+
+		Set("Common.GotoBlueprintHyperlink", EditBPHyperlinkStyle);
+
+		// Inherited from native
 		Set("Common.InheritedFromNativeTextColor", InheritedFromNativeTextColor);
-		Set("Common.IntroducedInThisInstanceTextColor", IntroducedInThisInstanceTextColor);
-		Set("Common.IntroducedInThisBlueprintTextColor", IntroducedInThisBlueprintTextColor);
+
+		FTextBlockStyle InheritedFromNativeTextStyle = FTextBlockStyle(NormalText)
+			.SetFont(TTF_CORE_FONT("Fonts/Roboto-Regular", 10))
+			.SetColorAndOpacity(InheritedFromNativeTextColor);
+		Set("Common.InheritedFromNativeTextStyle", InheritedFromNativeTextStyle);
+
+		// Go to native class hyperlink
+		FButtonStyle EditNativeHyperlinkButton = FButtonStyle()
+			.SetNormal(BORDER_BRUSH("Old/HyperlinkDotted", FMargin(0, 0, 0, 3 / 16.0f), InheritedFromNativeTextColor))
+			.SetPressed(FSlateNoResource())
+			.SetHovered(BORDER_BRUSH("Old/HyperlinkUnderline", FMargin(0, 0, 0, 3 / 16.0f), InheritedFromNativeTextColor));
+		FHyperlinkStyle EditNativeHyperlinkStyle = FHyperlinkStyle()
+			.SetUnderlineStyle(EditNativeHyperlinkButton)
+			.SetTextStyle(InheritedFromNativeTextStyle)
+			.SetPadding(FMargin(0.0f));
+
+		Set("Common.GotoNativeCodeHyperlink", EditNativeHyperlinkStyle);
 	}
 #endif
 }
diff --git a/Engine/Source/Editor/GameProjectGeneration/Private/SNewClassDialog.cpp b/Engine/Source/Editor/GameProjectGeneration/Private/SNewClassDialog.cpp
--- a/Engine/Source/Editor/GameProjectGeneration/Private/SNewClassDialog.cpp
+++ b/Engine/Source/Editor/GameProjectGeneration/Private/SNewClassDialog.cpp
@@ -428,8 +428,7 @@ void SNewClassDialog::Construct( const FArguments& InArgs )
 							.Padding(0.0f, 0.0f, 0.0f, 0.0f)
 							[
 								SNew(SHyperlink)
-								.Style(FCoreStyle::Get(), "Hyperlink")
-								.TextStyle(FEditorStyle::Get(), "Common.GoToNativeCodeHyperlinkStyle")
+								.Style(FEditorStyle::Get(), "Common.GotoNativeCodeHyperlink")
 								.OnNavigate(this, &SNewClassDialog::OnEditCodeClicked)
 								.Text(this, &SNewClassDialog::GetSelectedParentClassFilename)
 								.ToolTipText(FText::Format(LOCTEXT("GoToCode_ToolTip", "Click to open this source file in {0}"), FSourceCodeNavigation::GetSuggestedSourceCodeIDE()))
diff --git a/Engine/Source/Editor/Kismet/Private/BlueprintEditor.cpp b/Engine/Source/Editor/Kismet/Private/BlueprintEditor.cpp
--- a/Engine/Source/Editor/Kismet/Private/BlueprintEditor.cpp
+++ b/Engine/Source/Editor/Kismet/Private/BlueprintEditor.cpp
@@ -1731,8 +1731,9 @@ void FBlueprintEditor::PostRegenerateMenusAndToolbars()
 			.VAlign(VAlign_Center)
 			[
 				SNew(STextBlock)
-				.ShadowOffset( FVector2D::UnitVector )
+				.ShadowOffset(FVector2D::UnitVector)
 				.Text(this, &FBlueprintEditor::GetParentClassNameText)
+				.TextStyle(FEditorStyle::Get(), "Common.InheritedFromBlueprintTextStyle")
 				.ToolTipText(LOCTEXT("ParentClassToolTip", "The class that the current Blueprint is based on. The parent provides the base definition, which the current Blueprint extends."))
 				.Visibility(this, &FBlueprintEditor::GetParentClassNameVisibility)
 			]
@@ -1775,8 +1776,7 @@ void FBlueprintEditor::PostRegenerateMenusAndToolbars()
 			.VAlign(VAlign_Center)
 			[
 				SNew(SHyperlink)
-				.Style(FEditorStyle::Get(), "EditBPHyperlink")
-				.TextStyle(FEditorStyle::Get(), "Common.EditBlueprintHyperlinkStyle")
+				.Style(FEditorStyle::Get(), "Common.GotoNativeCodeHyperlink")
 				.IsEnabled(this, &FBlueprintEditor::IsNativeParentClassCodeLinkEnabled)
 				.Visibility(this, &FBlueprintEditor::GetNativeParentClassButtonsVisibility)
 				.OnNavigate(this, &FBlueprintEditor::OnEditParentClassNativeCodeClicked)
diff --git a/Engine/Source/Editor/Kismet/Private/BlueprintEditorTabFactories.cpp b/Engine/Source/Editor/Kismet/Private/BlueprintEditorTabFactories.cpp
--- a/Engine/Source/Editor/Kismet/Private/BlueprintEditorTabFactories.cpp
+++ b/Engine/Source/Editor/Kismet/Private/BlueprintEditorTabFactories.cpp
@@ -188,8 +188,7 @@ TSharedRef<SWidget> FDefaultsEditorSummoner::CreateOptionalDataOnlyMessage() con
 					+ SWrapBox::Slot()
 					[
 						SNew(SHyperlink)
-						.Style(FEditorStyle::Get(), "EditBPHyperlink")
-						.TextStyle(FEditorStyle::Get(), "Common.EditBlueprintHyperlinkStyle")
+						.Style(FEditorStyle::Get(), "Common.GotoBlueprintHyperlink")
 						.OnNavigate(this, &FDefaultsEditorSummoner::OnChangeBlueprintToNotDataOnly)
 						.Text(LOCTEXT("FullEditor", "Open Full Blueprint Editor"))
 						.ToolTipText(LOCTEXT("FullEditorToolTip", "This opens the blueprint in the full editor."))
diff --git a/Engine/Source/Editor/UnrealEd/Private/EditorClassUtils.cpp b/Engine/Source/Editor/UnrealEd/Private/EditorClassUtils.cpp
--- a/Engine/Source/Editor/UnrealEd/Private/EditorClassUtils.cpp
+++ b/Engine/Source/Editor/UnrealEd/Private/EditorClassUtils.cpp
@@ -74,8 +74,8 @@ TSharedRef<SWidget> FEditorClassUtils::GetDocumentationLinkWidget(const UClass*
 
 TSharedRef<SWidget> FEditorClassUtils::GetSourceLink(const UClass* Class, const TWeakObjectPtr<UObject> ObjectWeakPtr)
 {
-	FText BlueprintFormat = NSLOCTEXT("SourceHyperlink", "EditBlueprint", "Edit {0}");
-	FText CodeFormat = NSLOCTEXT("SourceHyperlink", "GoToCode", "{0}");
+	const FText BlueprintFormat = NSLOCTEXT("SourceHyperlink", "EditBlueprint", "Edit {0}");
+	const FText CodeFormat = NSLOCTEXT("SourceHyperlink", "GoToCode", "Open {0}");
 
 	return GetSourceLinkFormatted(Class, ObjectWeakPtr, BlueprintFormat, CodeFormat);
 }
@@ -108,8 +108,7 @@ TSharedRef<SWidget> FEditorClassUtils::GetSourceLinkFormatted(const UClass* Clas
 		TWeakObjectPtr<UBlueprint> BlueprintPtr = Blueprint;
 
 		SourceHyperlink = SNew(SHyperlink)
-			.Style(FEditorStyle::Get(), "EditBPHyperlink")
-			.TextStyle(FEditorStyle::Get(), "Common.EditBlueprintHyperlinkStyle")
+			.Style(FEditorStyle::Get(), "Common.GotoBlueprintHyperlink")
 			.OnNavigate_Static(&Local::OnEditBlueprintClicked, BlueprintPtr, ObjectWeakPtr)
 			.Text(FText::Format(BlueprintFormat, FText::FromString(Blueprint->GetName())))
 			.ToolTipText(NSLOCTEXT("SourceHyperlink", "EditBlueprint_ToolTip", "Click to edit the blueprint"));
@@ -129,8 +128,7 @@ TSharedRef<SWidget> FEditorClassUtils::GetSourceLinkFormatted(const UClass* Clas
 			};
 
 			SourceHyperlink = SNew(SHyperlink)
-				.Style(FCoreStyle::Get(), "Hyperlink")
-				.TextStyle(FEditorStyle::Get(), "Common.GoToNativeCodeHyperlinkStyle")
+				.Style(FEditorStyle::Get(), "Common.GotoNativeCodeHyperlink")
 				.OnNavigate_Static(&Local::OnEditCodeClicked, ClassHeaderPath)
 				.Text(FText::Format(CodeFormat, FText::FromString(FPaths::GetCleanFilename( *ClassHeaderPath ) ) ) )
 				.ToolTipText(FText::Format(NSLOCTEXT("SourceHyperlink", "GoToCode_ToolTip", "Click to open this source file in {0}"), FSourceCodeNavigation::GetSuggestedSourceCodeIDE()));
