[INTEGRATE] Change 2438561 by [REDACTED_PERSON][REDACTED_HANDLE] on 2015/02/09 14:49:13

	Editor - Adding Font Awesome to the set of editor fonts to allow us to use some the embeded glyphs where appropriate for all kinds of little symbols around the editor that are currently rastered images that can't match the same font options, but these can.  Additionally introducing several flat button styles to begin being used throughout the editor.

diff --git a/Engine/Source/Editor/EditorStyle/Private/SlateEditorStyle.cpp b/Engine/Source/Editor/EditorStyle/Private/SlateEditorStyle.cpp
--- a/Engine/Source/Editor/EditorStyle/Private/SlateEditorStyle.cpp
+++ b/Engine/Source/Editor/EditorStyle/Private/SlateEditorStyle.cpp
@@ -206,6 +206,13 @@ void FSlateEditorStyle::FStyle::SetupGeneralStyles()
 			.SetShadowColorAndOpacity(FLinearColor(0.0f, 0.0f, 0.0f)));
 
 		Set("NormalText", NormalText);
+
+		Set("LargeText", FTextBlockStyle(NormalText)
+			.SetFont(TTF_CORE_FONT("Fonts/Roboto-Bold", 11))
+			.SetColorAndOpacity(FLinearColor(1.0f, 1.0f, 1.0f))
+			.SetHighlightColor(FLinearColor(1.0f, 1.0f, 1.0f))
+			.SetShadowOffset(FVector2D(1, 1))
+			.SetShadowColorAndOpacity(FLinearColor(0, 0, 0, 0.9f)));
 	}
 
 	// Rendering resources that never change
@@ -333,6 +340,41 @@ void FSlateEditorStyle::FStyle::SetupGeneralStyles()
 			.SetPressed(BOX_BRUSH("Common/RoundedSelection_16x", 4.0f / 16.0f, SelectionColor_Pressed))
 			);
 
+		Set("FlatButton.Default", FButtonStyle(Button)
+			.SetNormal(BOX_BRUSH("Common/FlatButton", 2.0f / 8.0f, FLinearColor(1, 1, 1, 0.1f)))
+			.SetHovered(BOX_BRUSH("Common/FlatButton", 2.0f / 8.0f, SelectionColor))
+			.SetPressed(BOX_BRUSH("Common/FlatButton", 2.0f / 8.0f, SelectionColor_Pressed))
+			);
+
+		struct ButtonColor
+		{
+		public:
+			FName Name; FVector Color;
+			ButtonColor(const FName& Name, const FVector& Color) : Name(Name), Color(Color) {}
+		};
+
+		TArray< ButtonColor > FlatButtons;
+		FlatButtons.Add(ButtonColor("FlatButton.Primary", FVector(0.02899, 0.19752, 0.48195)));
+		FlatButtons.Add(ButtonColor("FlatButton.Success", FVector(0.10616, 0.48777, 0.10616)));
+		FlatButtons.Add(ButtonColor("FlatButton.Info", FVector(0.10363, 0.53564, 0.7372)));
+		FlatButtons.Add(ButtonColor("FlatButton.Warning", FVector(0.87514, 0.42591, 0.07383)));
+		FlatButtons.Add(ButtonColor("FlatButton.Danger", FVector(0.70117, 0.08464, 0.07593)));
+
+		for ( const ButtonColor& Entry : FlatButtons )
+		{
+			Set(Entry.Name, FButtonStyle(Button)
+				.SetNormal(BOX_BRUSH("Common/FlatButton", 2.0f / 8.0f, FLinearColor(Entry.Color.X, Entry.Color.Y, Entry.Color.Z, 0.6f)))
+				.SetHovered(BOX_BRUSH("Common/FlatButton", 2.0f / 8.0f, FLinearColor(Entry.Color.X, Entry.Color.Y, Entry.Color.Z, 1.f)))
+				.SetPressed(BOX_BRUSH("Common/FlatButton", 2.0f / 8.0f, FLinearColor(Entry.Color.X, Entry.Color.Y, Entry.Color.Z, 0.8f)))
+				);
+		}
+
+		Set("FontAwesome.8", TTF_FONT("Fonts/FontAwesome", 8));
+		Set("FontAwesome.9", TTF_FONT("Fonts/FontAwesome", 9));
+		Set("FontAwesome.10", TTF_FONT("Fonts/FontAwesome", 10));
+		Set("FontAwesome.11", TTF_FONT("Fonts/FontAwesome", 11));
+		Set("FontAwesome.12", TTF_FONT("Fonts/FontAwesome", 12));
+
 		/* Create a checkbox style for "ToggleButton" ... */
 		const FCheckBoxStyle ToggleButtonStyle = FCheckBoxStyle()
 			.SetCheckBoxType(ESlateCheckBoxType::ToggleButton)
@@ -5218,9 +5260,6 @@ void FSlateEditorStyle::FStyle::SetupContentBrowserStyle()
 		Set( "ContentBrowser.PathActions.NewAsset", new IMAGE_BRUSH( "Icons/icon_file_new_16px", Icon16x16 ) );
 		Set( "ContentBrowser.PathActions.SetColor", new IMAGE_BRUSH( "Icons/icon_Cascade_Color_40x", Icon16x16 ) );
 
-		FComboButtonStyle NewAssetComboButton = GetWidgetStyle<FComboButtonStyle>( "ToolbarComboButton" );
-		Set( "ContentBrowser.NewAsset.Style", NewAssetComboButton );
-
 		FComboButtonStyle FiltersComboButton = GetWidgetStyle<FComboButtonStyle>( "ToolbarComboButton" );
 		Set( "ContentBrowser.Filters.Style", FiltersComboButton );
 
