[INTEGRATE] Change 2439340 by [REDACTED_PERSON][REDACTED_HANDLE] on 2015/02/09 23:28:38

	Editor - Tweaking the flat button styles.

diff --git a/Engine/Source/Editor/EditorStyle/Private/SlateEditorStyle.cpp b/Engine/Source/Editor/EditorStyle/Private/SlateEditorStyle.cpp
--- a/Engine/Source/Editor/EditorStyle/Private/SlateEditorStyle.cpp
+++ b/Engine/Source/Editor/EditorStyle/Private/SlateEditorStyle.cpp
@@ -340,12 +340,20 @@ void FSlateEditorStyle::FStyle::SetupGeneralStyles()
 			.SetPressed(BOX_BRUSH("Common/RoundedSelection_16x", 4.0f / 16.0f, SelectionColor_Pressed))
 			);
 
-		Set("FlatButton.Default", FButtonStyle(Button)
+		Set("FlatButton.Dark", FButtonStyle(Button)
 			.SetNormal(BOX_BRUSH("Common/FlatButton", 2.0f / 8.0f, FLinearColor(0.125f, 0.125f, 0.125f, 0.8f)))
 			.SetHovered(BOX_BRUSH("Common/FlatButton", 2.0f / 8.0f, SelectionColor))
 			.SetPressed(BOX_BRUSH("Common/FlatButton", 2.0f / 8.0f, SelectionColor_Pressed))
 			);
 
+		Set("FlatButton.Light", FButtonStyle(Button)
+			.SetNormal(BOX_BRUSH("Common/FlatButton", 2.0f / 8.0f, FLinearColor(0.72267, 0.72267, 0.72267, 1)))
+			.SetHovered(BOX_BRUSH("Common/FlatButton", 2.0f / 8.0f, FLinearColor(0.85, 0.85, 0.85, 1)))
+			.SetPressed(BOX_BRUSH("Common/FlatButton", 2.0f / 8.0f, FLinearColor(0.58597, 0.58597, 0.58597, 1)))
+			);
+
+		Set("FlatButton.Default", GetWidgetStyle<FButtonStyle>("FlatButton.Dark"));
+
 		Set("FlatButton.DefaultTextStyle", FTextBlockStyle(NormalText)
 			.SetFont(TTF_CORE_FONT("Fonts/Roboto-Bold", 10))
 			.SetColorAndOpacity(FLinearColor(1.0f, 1.0f, 1.0f))
@@ -353,27 +361,36 @@ void FSlateEditorStyle::FStyle::SetupGeneralStyles()
 			.SetShadowOffset(FVector2D(1, 1))
 			.SetShadowColorAndOpacity(FLinearColor(0, 0, 0, 0.9f)));
 
-
 		struct ButtonColor
 		{
 		public:
-			FName Name; FVector Color;
-			ButtonColor(const FName& Name, const FVector& Color) : Name(Name), Color(Color) {}
+			FName Name;
+			FLinearColor Normal;
+			FLinearColor Hovered;
+			FLinearColor Pressed;
+
+			ButtonColor(const FName& Name, const FLinearColor& Normal) : Name(Name), Normal(Normal)
+			{
+				Hovered = Normal * 0.8f;
+				Hovered.A = Normal.A;
+				Pressed = Normal * 0.6f;
+				Pressed.A = Normal.A;
+			}
 		};
 
 		TArray< ButtonColor > FlatButtons;
-		FlatButtons.Add(ButtonColor("FlatButton.Primary", FVector(0.02899, 0.19752, 0.48195)));
-		FlatButtons.Add(ButtonColor("FlatButton.Success", FVector(0.10616, 0.48777, 0.10616)));
-		FlatButtons.Add(ButtonColor("FlatButton.Info", FVector(0.10363, 0.53564, 0.7372)));
-		FlatButtons.Add(ButtonColor("FlatButton.Warning", FVector(0.87514, 0.42591, 0.07383)));
-		FlatButtons.Add(ButtonColor("FlatButton.Danger", FVector(0.70117, 0.08464, 0.07593)));
+		FlatButtons.Add(ButtonColor("FlatButton.Primary", FLinearColor(0.02899, 0.19752, 0.48195)));
+		FlatButtons.Add(ButtonColor("FlatButton.Success", FLinearColor(0.10616, 0.48777, 0.10616)));
+		FlatButtons.Add(ButtonColor("FlatButton.Info", FLinearColor(0.10363, 0.53564, 0.7372)));
+		FlatButtons.Add(ButtonColor("FlatButton.Warning", FLinearColor(0.87514, 0.42591, 0.07383)));
+		FlatButtons.Add(ButtonColor("FlatButton.Danger", FLinearColor(0.70117, 0.08464, 0.07593)));
 
 		for ( const ButtonColor& Entry : FlatButtons )
 		{
 			Set(Entry.Name, FButtonStyle(Button)
-				.SetNormal(BOX_BRUSH("Common/FlatButton", 2.0f / 8.0f, FLinearColor(Entry.Color.X, Entry.Color.Y, Entry.Color.Z, 0.8f)))
-				.SetHovered(BOX_BRUSH("Common/FlatButton", 2.0f / 8.0f, FLinearColor(Entry.Color.X, Entry.Color.Y, Entry.Color.Z, 1.0f)))
-				.SetPressed(BOX_BRUSH("Common/FlatButton", 2.0f / 8.0f, FLinearColor(Entry.Color.X, Entry.Color.Y, Entry.Color.Z, 0.6f)))
+				.SetNormal(BOX_BRUSH("Common/FlatButton", 2.0f / 8.0f, Entry.Normal))
+				.SetHovered(BOX_BRUSH("Common/FlatButton", 2.0f / 8.0f, Entry.Hovered))
+				.SetPressed(BOX_BRUSH("Common/FlatButton", 2.0f / 8.0f, Entry.Pressed))
 				);
 		}
 
