[INTEGRATE] Change 2442157 by bruce.nesbit@bnesbit_release_port_edit on 2015/02/11 15:08:41

	Tweaked tutorial back button size and color

diff --git a/Engine/Source/Editor/EditorStyle/Private/SlateEditorStyle.cpp b/Engine/Source/Editor/EditorStyle/Private/SlateEditorStyle.cpp
--- a/Engine/Source/Editor/EditorStyle/Private/SlateEditorStyle.cpp
+++ b/Engine/Source/Editor/EditorStyle/Private/SlateEditorStyle.cpp
@@ -2622,6 +2622,8 @@ void FSlateEditorStyle::FStyle::SetupTutorialStyles()
 		const FLinearColor TutorialSelectionColor = FLinearColor(0.19f, 0.33f, 0.72f);
 		const FLinearColor TutorialNavigationButtonColor = FLinearColor(0.0f, 0.59f, 0.14f, 1.0f);
 		const FLinearColor TutorialNavigationButtonHoverColor = FLinearColor(0.2f, 0.79f, 0.34f, 1.0f);
+		const FLinearColor TutorialNavigationBackButtonColor = FLinearColor(0.4f, 0.4f, 0.4f, 0.75f);
+		const FLinearColor TutorialNavigationBackButtonHoverColor = FLinearColor(0.2f, 0.2f, 0.2f, 0.75f);
 
 		const FTextBlockStyle TutorialText = FTextBlockStyle(DocumentationText)
 			.SetColorAndOpacity(FLinearColor::Black)
@@ -2690,6 +2692,14 @@ void FSlateEditorStyle::FStyle::SetupTutorialStyles()
 			.SetNormalPadding( FMargin(0,0,0,1))
 			.SetPressedPadding( FMargin(0,1,0,0)));
 
+		Set("Tutorials.Content.NavigationBackButton", FButtonStyle()
+			.SetNormal(BOX_BRUSH("Common/ButtonHoverHint", FMargin(4 / 16.0f), TutorialNavigationBackButtonColor))
+			.SetHovered(BOX_BRUSH("Common/ButtonHoverHint", FMargin(4 / 16.0f), TutorialNavigationBackButtonHoverColor))
+			.SetPressed(BOX_BRUSH("Common/ButtonHoverHint", FMargin(4 / 16.0f), TutorialNavigationBackButtonHoverColor))
+			.SetDisabled(BOX_BRUSH("Common/ButtonHoverHint", FMargin(4 / 16.0f), TutorialNavigationBackButtonColor))
+			.SetNormalPadding(FMargin(0, 0, 0, 1))
+			.SetPressedPadding(FMargin(0, 1, 0, 0)));
+
 		Set( "Tutorials.Content.NavigationText", FTextBlockStyle(TutorialText));
 
 		Set( "Tutorials.Content.Color", FLinearColor(1.0f,1.0f,1.0f,0.9f) );
diff --git a/Engine/Source/Editor/IntroTutorials/Private/STutorialContent.cpp b/Engine/Source/Editor/IntroTutorials/Private/STutorialContent.cpp
--- a/Engine/Source/Editor/IntroTutorials/Private/STutorialContent.cpp
+++ b/Engine/Source/Editor/IntroTutorials/Private/STutorialContent.cpp
@@ -162,17 +162,7 @@ void STutorialContent::Construct(const FArguments& InArgs, UEditorTutorial* InTu
 						SNew(SBorder)
 						.BorderImage(this, &STutorialContent::GetBackButtonBorder)
 						[
-							SNew(SHorizontalBox)
-							+ SHorizontalBox::Slot()
-							.AutoWidth()
-							.VAlign(VAlign_Center)
-							.Padding(4.0f, 0.0f, 0.0f, 0.0f)
-							[
-								SNew(STextBlock)
-								.Text(this, &STutorialContent::GetBackButtonLabel)
-								.TextStyle(FEditorStyle::Get(), "Tutorials.Content.NavigationText")
-								.ColorAndOpacity(FLinearColor::White)
-							]
+ 							SNew(SHorizontalBox)
 							+ SHorizontalBox::Slot()
 								.AutoWidth()
 								[
@@ -825,7 +815,7 @@ FText STutorialContent::GetBackButtonLabel() const
 
 const FSlateBrush* STutorialContent::GetBackButtonBorder() const
 {
-	return BackButton->IsHovered() ? &FEditorStyle::Get().GetWidgetStyle<FButtonStyle>("Tutorials.Content.NavigationButton").Hovered : &FEditorStyle::Get().GetWidgetStyle<FButtonStyle>("Tutorials.Content.NavigationButton").Normal;
+	return BackButton->IsHovered() ? &FEditorStyle::Get().GetWidgetStyle<FButtonStyle>("Tutorials.Content.NavigationBackButton").Hovered : &FEditorStyle::Get().GetWidgetStyle<FButtonStyle>("Tutorials.Content.NavigationBackButton").Normal;
 }
 
 FReply STutorialContent::OnMouseButtonDown( const FGeometry& MyGeometry, const FPointerEvent& MouseEvent )
