[INTEGRATE] Change 2444776 by [REDACTED_PERSON][REDACTED_HANDLE] on 2015/02/13 10:50:41

	Fix remaining green hyperlink text issues

	https://jira.ol.epicgames.net/browse/UE-9801

diff --git a/Engine/Source/Editor/EditorStyle/Private/SlateEditorStyle.cpp b/Engine/Source/Editor/EditorStyle/Private/SlateEditorStyle.cpp
--- a/Engine/Source/Editor/EditorStyle/Private/SlateEditorStyle.cpp
+++ b/Engine/Source/Editor/EditorStyle/Private/SlateEditorStyle.cpp
@@ -70,9 +70,6 @@ FSlateEditorStyle::FStyle::FStyle( const TWeakObjectPtr< UEditorStyleSettings >&
 	, SelectionColor_Pressed( SelectionColor_Pressed_LinearRef )
 
 	, InheritedFromBlueprintTextColor(FLinearColor(0.25f, 0.5f, 1.0f))
-	, InheritedFromNativeTextColor(FLinearColor(0.375f, 1.0f, 0.375f))
-	, IntroducedInThisInstanceTextColor(FLinearColor::White)
-	, IntroducedInThisBlueprintTextColor(FLinearColor::White)
 
 	, Settings( InSettings )
 {
@@ -2347,9 +2344,6 @@ void FSlateEditorStyle::FStyle::SetupGeneralStyles()
 
 	// Common styles for blueprint/code references
 	{
-		Set("Common.IntroducedInThisInstanceTextColor", IntroducedInThisInstanceTextColor);
-		Set("Common.IntroducedInThisBlueprintTextColor", IntroducedInThisBlueprintTextColor);
-
 		// Inherited from blueprint
 		Set("Common.InheritedFromBlueprintTextColor", InheritedFromBlueprintTextColor);
 
@@ -2371,19 +2365,17 @@ void FSlateEditorStyle::FStyle::SetupGeneralStyles()
 
 		Set("Common.GotoBlueprintHyperlink", EditBPHyperlinkStyle);
 
-		// Inherited from native
-		Set("Common.InheritedFromNativeTextColor", InheritedFromNativeTextColor);
 
 		FTextBlockStyle InheritedFromNativeTextStyle = FTextBlockStyle(NormalText)
-			.SetFont(TTF_CORE_FONT("Fonts/Roboto-Regular", 10))
-			.SetColorAndOpacity(InheritedFromNativeTextColor);
+			.SetFont(TTF_CORE_FONT("Fonts/Roboto-Regular", 10));
+
 		Set("Common.InheritedFromNativeTextStyle", InheritedFromNativeTextStyle);
 
 		// Go to native class hyperlink
 		FButtonStyle EditNativeHyperlinkButton = FButtonStyle()
-			.SetNormal(BORDER_BRUSH("Old/HyperlinkDotted", FMargin(0, 0, 0, 3 / 16.0f), InheritedFromNativeTextColor))
+			.SetNormal(BORDER_BRUSH("Old/HyperlinkDotted", FMargin(0, 0, 0, 3 / 16.0f)))
 			.SetPressed(FSlateNoResource())
-			.SetHovered(BORDER_BRUSH("Old/HyperlinkUnderline", FMargin(0, 0, 0, 3 / 16.0f), InheritedFromNativeTextColor));
+			.SetHovered(BORDER_BRUSH("Old/HyperlinkUnderline", FMargin(0, 0, 0, 3 / 16.0f)));
 		FHyperlinkStyle EditNativeHyperlinkStyle = FHyperlinkStyle()
 			.SetUnderlineStyle(EditNativeHyperlinkButton)
 			.SetTextStyle(InheritedFromNativeTextStyle)
diff --git a/Engine/Source/Editor/EditorStyle/Private/SlateEditorStyle.h b/Engine/Source/Editor/EditorStyle/Private/SlateEditorStyle.h
--- a/Engine/Source/Editor/EditorStyle/Private/SlateEditorStyle.h
+++ b/Engine/Source/Editor/EditorStyle/Private/SlateEditorStyle.h
@@ -126,9 +126,6 @@ public:
 
 		// These are common colors used thruout the editor in mutliple style elements
 		const FSlateColor InheritedFromBlueprintTextColor;
-		const FSlateColor InheritedFromNativeTextColor;
-		const FSlateColor IntroducedInThisInstanceTextColor;
-		const FSlateColor IntroducedInThisBlueprintTextColor;
 
 		FTextBlockStyle NormalText;
 		FEditableTextBoxStyle NormalEditableTextBoxStyle;
