GitHub #887 : Fix tab_background filename case for case-sensitive file systems

diff --git a/Engine/Source/Editor/EditorStyle/Private/SlateEditorStyle.cpp b/Engine/Source/Editor/EditorStyle/Private/SlateEditorStyle.cpp
--- a/Engine/Source/Editor/EditorStyle/Private/SlateEditorStyle.cpp
+++ b/Engine/Source/Editor/EditorStyle/Private/SlateEditorStyle.cpp
@@ -1761,7 +1761,7 @@ void FSlateEditorStyle::FStyle::SetupGeneralStyles()
 			.SetCheckedImage( BOX_BRUSH( "/GameProjectDialog/Tab_Active", 4 / 16.0f ) )
 			);
 
-		Set( "GameProjectDialog.TabBackground", new BOX_BRUSH( "/GameProjectDialog/Tab_Background", 4 / 16.0f ) );
+		Set( "GameProjectDialog.TabBackground", new BOX_BRUSH( "/GameProjectDialog/tab_background", 4 / 16.0f ) );
 
 		Set( "GameProjectDialog.FolderIconClosed", new IMAGE_BRUSH( "Icons/FolderClosed", FVector2D(18, 16) ) );
 		Set( "GameProjectDialog.FolderIconOpen", new IMAGE_BRUSH( "Icons/FolderOpen", FVector2D(18, 16) ) );
