Changed resume icon in behavior tree debugger

It is now consistent with other resume icons.

UE-9758 - Behavior Tree "Resume" button does not match other resume buttons in editor

diff --git a/Engine/Source/Editor/EditorStyle/Private/SlateEditorStyle.cpp b/Engine/Source/Editor/EditorStyle/Private/SlateEditorStyle.cpp
--- a/Engine/Source/Editor/EditorStyle/Private/SlateEditorStyle.cpp
+++ b/Engine/Source/Editor/EditorStyle/Private/SlateEditorStyle.cpp
@@ -3550,8 +3550,8 @@ void FSlateEditorStyle::FStyle::SetupGraphEditorStyles()
 
 			Set( "BTEditor.Debugger.PausePlaySession", new IMAGE_BRUSH( "Icons/icon_pause_40x", Icon40x40 ) );
 			Set( "BTEditor.Debugger.PausePlaySession.Small", new IMAGE_BRUSH( "Icons/icon_pause_40x", Icon20x20 ) );
-			Set( "BTEditor.Debugger.ResumePlaySession", new IMAGE_BRUSH( "Icons/icon_play_40x", Icon40x40 ) );
-			Set( "BTEditor.Debugger.ResumePlaySession.Small", new IMAGE_BRUSH( "Icons/icon_play_40x", Icon20x20 ) );
+			Set( "BTEditor.Debugger.ResumePlaySession", new IMAGE_BRUSH( "Icons/icon_simulate_40x", Icon40x40 ) );
+			Set( "BTEditor.Debugger.ResumePlaySession.Small", new IMAGE_BRUSH( "Icons/icon_simulate_40x", Icon20x20 ) );
 			Set( "BTEditor.Debugger.StopPlaySession", new IMAGE_BRUSH( "Icons/icon_stop_40x", Icon40x40 ) );
 			Set( "BTEditor.Debugger.StopPlaySession.Small", new IMAGE_BRUSH( "Icons/icon_stop_40x", Icon20x20 ) );
 
