Paper2D: Added placeholder icons and thumbnails for various Paper2D types [UE-3658]

diff --git a/Engine/Source/Editor/EditorStyle/Private/SlateEditorStyle.cpp b/Engine/Source/Editor/EditorStyle/Private/SlateEditorStyle.cpp
--- a/Engine/Source/Editor/EditorStyle/Private/SlateEditorStyle.cpp
+++ b/Engine/Source/Editor/EditorStyle/Private/SlateEditorStyle.cpp
@@ -5068,38 +5068,54 @@ void FSlateEditorStyle::FStyle::SetupClassIconsAndThumbnails()
 
 		//@TODO: PAPER2D: Defined here until it is possible to define these in a plugin
 		{
-			Set("ClassIcon.PaperTerrainActor", new IMAGE_BRUSH("Icons/ActorIcons/Paper2DTerrainActor_16x", Icon16x16));
-			Set("ClassIcon.PaperTerrainComponent", new IMAGE_BRUSH("Icons/ActorIcons/Paper2DTerrainActor_16x", Icon16x16));
-
-			Set("ClassIcon.PaperTileMapActor", new IMAGE_BRUSH("Icons/ActorIcons/Paper2DTileMapActor_16x", Icon16x16));
-			Set("ClassIcon.PaperTileMapComponent", new IMAGE_BRUSH("Icons/ActorIcons/Paper2DTileMapComponent_16x", Icon16x16));
-			Set("ClassThumbnail.PaperTileMapComponent", new IMAGE_BRUSH("Icons/ActorIcons/Paper2DTileMapComponent_64x", Icon64x64));
-
-			Set("ClassIcon.PaperSpriteActor", new IMAGE_BRUSH("Icons/AssetIcons/Paper2DSprite_16x", Icon16x16));
-			Set("ClassIcon.PaperSpriteComponent", new IMAGE_BRUSH("Icons/AssetIcons/Paper2DSprite_16x", Icon16x16));
-			Set("ClassIcon.PaperSprite", new IMAGE_BRUSH("Icons/AssetIcons/Paper2DSprite_16x", Icon16x16));
-			Set("ClassThumbnail.PaperSprite", new IMAGE_BRUSH("Icons/AssetIcons/Paper2DSprite_64x", Icon64x64));
-
-			Set("ClassIcon.PaperFlipbookActor", new IMAGE_BRUSH("Icons/AssetIcons/Paper2DSpriteFlipbook16x", Icon16x16));
-			Set("ClassIcon.PaperFlipbookComponent", new IMAGE_BRUSH("Icons/AssetIcons/Paper2DSpriteFlipbook16x", Icon16x16));
-			Set("ClassIcon.PaperFlipbook", new IMAGE_BRUSH("Icons/AssetIcons/Paper2DSpriteFlipbook16x", Icon16x16));
-			Set("ClassThumbnail.PaperFlipbook", new IMAGE_BRUSH("Icons/AssetIcons/Paper2DSpriteFlipbook64x", Icon64x64));
-
+			// Sprites (asset, component, actor)
+			Set("ClassIcon.PaperSprite", new IMAGE_BRUSH("Icons/AssetIcons/PaperSprite_16x", Icon16x16));
+			Set("ClassThumbnail.PaperSprite", new IMAGE_BRUSH("Icons/AssetIcons/PaperSprite_64x", Icon64x64));
+			Set("ClassIcon.PaperSpriteComponent", new IMAGE_BRUSH("Icons/AssetIcons/PaperSpriteComponent_16x", Icon16x16));
+			Set("ClassThumbnail.PaperSpriteComponent", new IMAGE_BRUSH("Icons/AssetIcons/PaperSpriteComponent_64x", Icon64x64));
+			Set("ClassIcon.PaperSpriteActor", new IMAGE_BRUSH("Icons/AssetIcons/PaperSpriteActor_16x", Icon16x16));
+			Set("ClassThumbnail.PaperSpriteActor", new IMAGE_BRUSH("Icons/AssetIcons/PaperSpriteActor_64x", Icon64x64));
+
+			// Flipbooks (asset, component, actor)
+			Set("ClassIcon.PaperFlipbook", new IMAGE_BRUSH("Icons/AssetIcons/PaperFlipbook_16x", Icon16x16));
+			Set("ClassThumbnail.PaperFlipbook", new IMAGE_BRUSH("Icons/AssetIcons/PaperFlipbook_64x", Icon64x64));
+			Set("ClassIcon.PaperFlipbookComponent", new IMAGE_BRUSH("Icons/AssetIcons/PaperFlipbookComponent_16x", Icon16x16));
+			Set("ClassThumbnail.PaperFlipbookComponent", new IMAGE_BRUSH("Icons/AssetIcons/PaperFlipbookComponent_64x", Icon64x64));
+			Set("ClassIcon.PaperFlipbookActor", new IMAGE_BRUSH("Icons/AssetIcons/PaperFlipbookActor_16x", Icon16x16));
+			Set("ClassThumbnail.PaperFlipbookActor", new IMAGE_BRUSH("Icons/AssetIcons/PaperFlipbookActor_64x", Icon64x64));
+
+			// Tile maps (asset, component, actor)
+			Set("ClassIcon.PaperTileMap", new IMAGE_BRUSH("Icons/AssetIcons/PaperTileMap_16x", Icon16x16));
+			Set("ClassThumbnail.PaperTileMap", new IMAGE_BRUSH("Icons/AssetIcons/PaperTileMap_64x", Icon64x64));
+			Set("ClassIcon.PaperTileMapComponent", new IMAGE_BRUSH("Icons/AssetIcons/PaperTileMapComponent_16x", Icon16x16));
+			Set("ClassThumbnail.PaperTileMapComponent", new IMAGE_BRUSH("Icons/AssetIcons/PaperTileMapComponent_64x", Icon64x64));
+			Set("ClassIcon.PaperTileMapActor", new IMAGE_BRUSH("Icons/AssetIcons/PaperTileMapActor_16x", Icon16x16));
+			Set("ClassThumbnail.PaperTileMapActor", new IMAGE_BRUSH("Icons/AssetIcons/PaperTileMapActor_64x", Icon64x64));
+
+			// UPaperSpriteAtlas assets
+			//@TODO: Paper2D: These icons don't match the naming scheme
 			Set("ClassIcon.PaperSpriteAtlas", new IMAGE_BRUSH("Icons/AssetIcons/Paper2DSpriteAtlasGroup_16x", Icon16x16));
 			Set("ClassThumbnail.PaperSpriteAtlas", new IMAGE_BRUSH("Icons/AssetIcons/Paper2DSpriteAtlasGroup_64x", Icon64x64));
 
+			//@TODO: UPaperSpriteSheet icons?
+
+			// APaperCharacter icons
 			Set("ClassIcon.PaperCharacter", new IMAGE_BRUSH("Icons/AssetIcons/PaperCharacter_16x", Icon16x16));
 			Set("ClassThumbnail.PaperCharacter", new IMAGE_BRUSH("Icons/AssetIcons/PaperCharacter_64x", Icon64x64));
 
-			Set("ClassIcon.PaperTileMap", new IMAGE_BRUSH("Icons/AssetIcons/TileMap_16x", Icon16x16));
-			Set("ClassThumbnail.PaperTileMap", new IMAGE_BRUSH("Icons/AssetIcons/TileMap_64x", Icon64x64));
+			// UPaperTileSet icons
+			Set("ClassIcon.PaperTileSet", new IMAGE_BRUSH("Icons/AssetIcons/PaperTileSet_16x", Icon16x16));
+			Set("ClassThumbnail.PaperTileSet", new IMAGE_BRUSH("Icons/AssetIcons/PaperTileSet_64x", Icon64x64));
 
-			Set("ClassIcon.PaperTileSet", new IMAGE_BRUSH("Icons/AssetIcons/TileMap_16x", Icon16x16));
-			Set("ClassThumbnail.PaperTileSet", new IMAGE_BRUSH("Icons/AssetIcons/TileMap_64x", Icon64x64));
+			// UPaperTerrainMaterial icons
+			Set("ClassIcon.PaperTerrainMaterial", new IMAGE_BRUSH("Icons/AssetIcons/PaperTerrainMaterial_16x", Icon16x16));
+			Set("ClassThumbnail.PaperTerrainMaterial", new IMAGE_BRUSH("Icons/AssetIcons/PaperTerrainMaterial_64x", Icon64x64));
 
-			// ClassIcon and ClassThumbnail for tile sets and/or terrain materials?  Might not be needed/feasible due to them being a UDataAsset
-			//Set("ClassIcon.PaperTerrainMaterial", new IMAGE_BRUSH("Icons/AssetIcons/Paper2DSpriteAtlasGroup_16x", Icon16x16));
-			//Set("ClassThumbnail.PaperTerrainMaterial", new IMAGE_BRUSH("Icons/AssetIcons/Paper2DSpriteAtlasGroup_64x", Icon64x64));
+			// Terrain splines (component, actor)
+			Set("ClassIcon.PaperTerrainComponent", new IMAGE_BRUSH("Icons/AssetIcons/PaperTerrainComponent_16x", Icon16x16));
+			Set("ClassThumbnail.PaperTerrainComponent", new IMAGE_BRUSH("Icons/AssetIcons/PaperTerrainComponent_64x", Icon64x64));
+			Set("ClassIcon.PaperTerrainActor", new IMAGE_BRUSH("Icons/AssetIcons/PaperTerrainActor_16x", Icon16x16));
+			Set("ClassThumbnail.PaperTerrainActor", new IMAGE_BRUSH("Icons/AssetIcons/PaperTerrainActor_16x", Icon64x64));
 		}
 
 		// Factory classes
