Merging CL 2496217

//depot/UE4-UT/...

to //depot/UE4/...

diff --git a/Engine/Config/BaseEditorUserSettings.ini b/Engine/Config/BaseEditorUserSettings.ini
--- a/Engine/Config/BaseEditorUserSettings.ini
+++ b/Engine/Config/BaseEditorUserSettings.ini
@@ -81,6 +81,7 @@ bShowFriendlyNames=True
 ;If true, the Editor Preferences and Project Settings menu entries in the main menu will be expanded with sub-menus for each settings section
 bExpandConfigurationMenus=False
 bShowProjectMenus=True
+bShowLaunchMenus=True
 
 [/Script/Levels.LevelBrowserSettings]
 ; True if the actor count is displayed in the slate level browser
diff --git a/Engine/Source/Editor/EditorStyle/Public/Classes/EditorStyleSettings.h b/Engine/Source/Editor/EditorStyle/Public/Classes/EditorStyleSettings.h
--- a/Engine/Source/Editor/EditorStyle/Public/Classes/EditorStyleSettings.h
+++ b/Engine/Source/Editor/EditorStyle/Public/Classes/EditorStyleSettings.h
@@ -82,6 +82,10 @@ public:
 	UPROPERTY(config)
 	uint32 bShowProjectMenus : 1;
 
+	/** When enabled, the Launch menu items will be shown. */
+	UPROPERTY(config)
+	uint32 bShowLaunchMenus : 1;
+
 	/** The display mode for timestamps in the output log */
 	UPROPERTY(EditAnywhere, config, AdvancedDisplay, Category=UserInterface)
 	TEnumAsByte<ELogTimes::Type> LogTimestampMode;
diff --git a/Engine/Source/Editor/UnrealEd/Private/Kismet2/DebuggerCommands.cpp b/Engine/Source/Editor/UnrealEd/Private/Kismet2/DebuggerCommands.cpp
--- a/Engine/Source/Editor/UnrealEd/Private/Kismet2/DebuggerCommands.cpp
+++ b/Engine/Source/Editor/UnrealEd/Private/Kismet2/DebuggerCommands.cpp
@@ -438,7 +438,7 @@ void FPlayWorldCommands::BuildToolbar( FToolBarBuilder& ToolbarBuilder, bool bIn
 		true
 	);
 
-	if (bIncludeLaunchButtonAndOptions)
+	if (bIncludeLaunchButtonAndOptions && GetDefault<UEditorStyleSettings>()->bShowLaunchMenus)
 	{
 		// Launch
 		ToolbarBuilder.AddToolBarButton( 
