[UE-11752] - Icons for proc foliage volume & blocking volume

diff --git a/Engine/Source/Editor/EditorStyle/Private/SlateEditorStyle.cpp b/Engine/Source/Editor/EditorStyle/Private/SlateEditorStyle.cpp
--- a/Engine/Source/Editor/EditorStyle/Private/SlateEditorStyle.cpp
+++ b/Engine/Source/Editor/EditorStyle/Private/SlateEditorStyle.cpp
@@ -5234,6 +5234,8 @@ void FSlateEditorStyle::FStyle::SetupClassIconsAndThumbnails()
 			TEXT("PostProcessVolume"),
 			TEXT("PrecomputedVisibilityOverrideVolume"),
 			TEXT("PrecomputedVisibilityVolume"),
+			TEXT("ProceduralFoliageVolume"),
+			TEXT("ProceduralFoliageBlockingVolume"),
 			TEXT("RadialForceActor"),
 			TEXT("RadialForceComponent"),
 			TEXT("ReflectionCapture"),
