Hooked up new icons for Localization Dashbard UI.

diff --git a/Engine/Source/Editor/EditorStyle/Private/SlateEditorStyle.cpp b/Engine/Source/Editor/EditorStyle/Private/SlateEditorStyle.cpp
--- a/Engine/Source/Editor/EditorStyle/Private/SlateEditorStyle.cpp
+++ b/Engine/Source/Editor/EditorStyle/Private/SlateEditorStyle.cpp
@@ -6526,28 +6526,28 @@ void FSlateEditorStyle::FStyle::SetupTranslationEditorStyles()
 
 void FSlateEditorStyle::FStyle::SetupLocalizationDashboardStyles()
 {
-	Set("LocalizationDashboard.GatherAllTargets", new IMAGE_BRUSH("Icons/icon_Refresh_16x", Icon40x40));
-	Set("LocalizationDashboard.ImportForAllTargetsCultures", new IMAGE_BRUSH("Icons/icon_Samples_Download_40x", Icon40x40));
-	Set("LocalizationDashboard.ExportForAllTargetsCultures", new IMAGE_BRUSH("Icons/icon_Import_40x", Icon40x40));
-	Set("LocalizationDashboard.CompileAllTargets", new IMAGE_BRUSH("Icons/icon_package_16x", Icon40x40));
-
-	Set("LocalizationDashboard.GatherTarget", new IMAGE_BRUSH("Icons/icon_Refresh_16x", Icon16x16));
-	Set("LocalizationDashboard.RefreshWordCount", new IMAGE_BRUSH("Icons/icon_Refresh_16x", Icon16x16));
-	Set("LocalizationDashboard.ImportForAllCultures", new IMAGE_BRUSH("Icons/icon_Samples_Download_20x", Icon16x16));
-	Set("LocalizationDashboard.ExportForAllCultures", new IMAGE_BRUSH("Icons/icon_Import_16x", Icon16x16));
-	Set("LocalizationDashboard.CompileTarget", new IMAGE_BRUSH("Icons/icon_package_16x", Icon16x16));
+	Set("LocalizationDashboard.GatherAllTargets", new IMAGE_BRUSH("Icons/Icon_Localisation_Gather_All_40x", Icon40x40));
+	Set("LocalizationDashboard.ImportForAllTargetsCultures", new IMAGE_BRUSH("Icons/Icon_Localisation_Import_All_40x", Icon40x40));
+	Set("LocalizationDashboard.ExportForAllTargetsCultures", new IMAGE_BRUSH("Icons/Icon_Localisation_Export_All_40x", Icon40x40));
+	Set("LocalizationDashboard.CompileAllTargets", new IMAGE_BRUSH("Icons/Icon_Localisation_Compile_Translations_40x", Icon40x40));
+
+	Set("LocalizationDashboard.GatherTarget", new IMAGE_BRUSH("Icons/Icon_Localisation_Gather_All_16x", Icon16x16));
+	Set("LocalizationDashboard.RefreshWordCount", new IMAGE_BRUSH("Icons/Icon_Localisation_Refresh_Word_Counts_16x", Icon16x16));
+	Set("LocalizationDashboard.ImportForAllCultures", new IMAGE_BRUSH("Icons/Icon_Localisation_Import_All_16x", Icon16x16));
+	Set("LocalizationDashboard.ExportForAllCultures", new IMAGE_BRUSH("Icons/Icon_Localisation_Export_All_16x", Icon16x16));
+	Set("LocalizationDashboard.CompileTarget", new IMAGE_BRUSH("Icons/Icon_Localisation_Compile_Translations_16x", Icon16x16));
 	Set("LocalizationDashboard.DeleteTarget", new IMAGE_BRUSH("Icons/Cross_12x", Icon12x12 ) );
 
-	Set("LocalizationTargetEditor.Gather", new IMAGE_BRUSH("Icons/icon_Refresh_16x", Icon40x40));
-	Set("LocalizationTargetEditor.ImportForAllCultures", new IMAGE_BRUSH("Icons/icon_Samples_Download_40x", Icon40x40));
-	Set("LocalizationTargetEditor.ExportForAllCultures", new IMAGE_BRUSH("Icons/icon_Import_40x", Icon40x40));
-	Set("LocalizationTargetEditor.RefreshWordCounts", new IMAGE_BRUSH("Icons/icon_Refresh_16x", Icon40x40));
-	Set("LocalizationTargetEditor.Compile", new IMAGE_BRUSH( "Icons/icon_package_16x", Icon40x40));
+	Set("LocalizationTargetEditor.Gather", new IMAGE_BRUSH("Icons/Icon_Localisation_Gather_All_16x", Icon40x40));
+	Set("LocalizationTargetEditor.ImportForAllCultures", new IMAGE_BRUSH("Icons/Icon_Localisation_Import_All_16x", Icon40x40));
+	Set("LocalizationTargetEditor.ExportForAllCultures", new IMAGE_BRUSH("Icons/Icon_Localisation_Export_All_16x", Icon40x40));
+	Set("LocalizationTargetEditor.RefreshWordCounts", new IMAGE_BRUSH("Icons/Icon_Localisation_Refresh_Word_Counts_16x", Icon40x40));
+	Set("LocalizationTargetEditor.Compile", new IMAGE_BRUSH( "Icons/Icon_Localisation_Compile_Translations_16x", Icon40x40));
 
 	Set("LocalizationTargetEditor.NativeCulture", new IMAGE_BRUSH( "Icons/Star_16x", Icon16x16 ) );
 	Set("LocalizationTargetEditor.EditTranslations", new IMAGE_BRUSH("Icons/icon_file_open_16px", Icon16x16));
-	Set("LocalizationTargetEditor.ImportForCulture", new IMAGE_BRUSH("Icons/icon_Samples_Download_20x", Icon16x16));
-	Set("LocalizationTargetEditor.ExportForCulture", new IMAGE_BRUSH("Icons/icon_Import_16x", Icon16x16));
+	Set("LocalizationTargetEditor.ImportForCulture", new IMAGE_BRUSH("Icons/Icon_Localisation_Import_All_16x", Icon16x16));
+	Set("LocalizationTargetEditor.ExportForCulture", new IMAGE_BRUSH("Icons/Icon_Localisation_Export_All_16x", Icon16x16));
 	Set("LocalizationTargetEditor.DeleteCulture", new IMAGE_BRUSH("Icons/Cross_12x", Icon12x12 ) );
 
 	Set( "LocalizationDashboard.CommandletLog.Text", FTextBlockStyle(NormalText)
