Additional clean-up of obsolte levels module

diff --git a/Engine/Source/Editor/EditorStyle/Private/SlateEditorStyle.cpp b/Engine/Source/Editor/EditorStyle/Private/SlateEditorStyle.cpp
--- a/Engine/Source/Editor/EditorStyle/Private/SlateEditorStyle.cpp
+++ b/Engine/Source/Editor/EditorStyle/Private/SlateEditorStyle.cpp
@@ -1225,19 +1225,6 @@ void FSlateEditorStyle::FStyle::SetupGeneralStyles()
 		Set( "Level.ScriptHighlightIcon16x", new IMAGE_BRUSH( "Icons/icon_levels_Blueprint_hi_16px", Icon16x16 ) );
 		Set( "Level.EmptyIcon16x", new IMAGE_BRUSH( "Icons/Empty_16x", Icon16x16 ) );
 		Set( "Level.ColorIcon40x", new IMAGE_BRUSH( "Icons/icon_levels_back_16px", Icon16x16 ) );
-
-		Set( "Level.ArrowUp", new IMAGE_BRUSH( "Old/ArrowUp", Icon12x12 ) );
-		Set( "Level.ArrowDown", new IMAGE_BRUSH( "Old/ArrowDown", Icon12x12 ) );
-	}
-
-	// Level Browser
-	{
-		Set( "LevelBrowserButton.Normal", new FSlateNoResource() ); // Note: Intentionally transparent background
-		Set( "LevelBrowserButton.Pressed", new BOX_BRUSH( "Common/RoundedSelection_16x", 4.0f/16.0f, SelectionColor_Pressed ) );
-		Set( "LevelBrowserButton.Hovered", new BOX_BRUSH( "Common/RoundedSelection_16x", 4.0f/16.0f, SelectionColor ) );
-
-		Set( "LevelBrowser.LabelFont", TTF_CORE_FONT( "Fonts/Roboto-Regular", 9 ) );
-		Set( "LevelBrowser.LabelFontBold", TTF_CORE_FONT( "Fonts/Roboto-Bold", 10 ) );
 	}
 
 	// World Browser
@@ -1264,6 +1251,9 @@ void FSlateEditorStyle::FStyle::SetupGeneralStyles()
 				.SetColorAndOpacity( FLinearColor(0.9, 0.9f, 0.9f, 0.5f) )
 				.SetShadowOffset( FVector2D::ZeroVector )
 			);
+
+		Set( "WorldBrowser.LabelFont", TTF_CORE_FONT( "Fonts/Roboto-Regular", 9 ) );
+		Set( "WorldBrowser.LabelFontBold", TTF_CORE_FONT( "Fonts/Roboto-Bold", 10 ) );
 	}
 
 
diff --git a/Engine/Source/Editor/LevelEditor/Private/SLevelEditor.cpp b/Engine/Source/Editor/LevelEditor/Private/SLevelEditor.cpp
--- a/Engine/Source/Editor/LevelEditor/Private/SLevelEditor.cpp
+++ b/Engine/Source/Editor/LevelEditor/Private/SLevelEditor.cpp
@@ -667,7 +667,7 @@ TSharedRef<SDockTab> SLevelEditor::SpawnLevelEditorTab( const FSpawnTabArgs& Arg
 					]
 				];
 	}
-	else if( TabIdentifier == WorldBrowserHierarchyTab || TabIdentifier == TEXT("LevelEditorLevelBrowser") )
+	else if( TabIdentifier == WorldBrowserHierarchyTab )
 	{
 		FWorldBrowserModule& WorldBrowserModule = FModuleManager::LoadModuleChecked<FWorldBrowserModule>( "WorldBrowser" );
 		return SNew( SDockTab )
diff --git a/Engine/Source/Editor/WorldBrowser/Private/SWorldHierarchyItem.cpp b/Engine/Source/Editor/WorldBrowser/Private/SWorldHierarchyItem.cpp
--- a/Engine/Source/Editor/WorldBrowser/Private/SWorldHierarchyItem.cpp
+++ b/Engine/Source/Editor/WorldBrowser/Private/SWorldHierarchyItem.cpp
@@ -378,11 +378,11 @@ FSlateFontInfo SWorldHierarchyItem::GetLevelDisplayNameFont() const
 {
 	if (LevelModel->IsCurrent())
 	{
-		return FEditorStyle::GetFontStyle("LevelBrowser.LabelFontBold");
+		return FEditorStyle::GetFontStyle("WorldBrowser.LabelFontBold");
 	}
 	else
 	{
-		return FEditorStyle::GetFontStyle("LevelBrowser.LabelFont");
+		return FEditorStyle::GetFontStyle("WorldBrowser.LabelFont");
 	}
 }
 
