Moved Plugins menu item under Settings/Configuration

Added WorldProperties icons back in as we want to show them in the settings sub menu now to match plugins.

UE-16205

diff --git a/Engine/Source/Editor/EditorStyle/Private/SlateEditorStyle.cpp b/Engine/Source/Editor/EditorStyle/Private/SlateEditorStyle.cpp
--- a/Engine/Source/Editor/EditorStyle/Private/SlateEditorStyle.cpp
+++ b/Engine/Source/Editor/EditorStyle/Private/SlateEditorStyle.cpp
@@ -4301,8 +4301,8 @@ void FSlateEditorStyle::FStyle::SetupLevelEditorStyle()
 		Set( "LevelEditor.BspMode.Selected", new IMAGE_BRUSH( "Icons/icon_Mode_GeoEdit-a_40px", Icon40x40 ) );
 		Set( "LevelEditor.BspMode.Selected.Small", new IMAGE_BRUSH( "Icons/icon_Mode_GeoEdit-a_40px", Icon20x20 ) );
 
-		//Set( "LevelEditor.WorldProperties", new IMAGE_BRUSH( "Icons/icon_worldscript_40x", Icon40x40 ) );
-		//Set( "LevelEditor.WorldProperties.Small", new IMAGE_BRUSH( "Icons/icon_worldscript_40x", Icon20x20 ) );
+		Set( "LevelEditor.WorldProperties", new IMAGE_BRUSH( "Icons/icon_worldscript_40x", Icon40x40 ) );
+		Set( "LevelEditor.WorldProperties.Small", new IMAGE_BRUSH( "Icons/icon_worldscript_40x", Icon20x20 ) );
 		Set( "LevelEditor.WorldProperties.Tab", new IMAGE_BRUSH( "Icons/icon_worldscript_40x", Icon16x16 ) );
 		Set( "LevelEditor.OpenContentBrowser", new IMAGE_BRUSH( "Icons/icon_ContentBrowser_40x", Icon40x40 ) );
 		Set( "LevelEditor.OpenContentBrowser.Small", new IMAGE_BRUSH( "Icons/icon_ContentBrowser_40x", Icon20x20 ) );
diff --git a/Engine/Source/Editor/LevelEditor/Private/LevelEditorToolBar.cpp b/Engine/Source/Editor/LevelEditor/Private/LevelEditorToolBar.cpp
--- a/Engine/Source/Editor/LevelEditor/Private/LevelEditorToolBar.cpp
+++ b/Engine/Source/Editor/LevelEditor/Private/LevelEditorToolBar.cpp
@@ -36,6 +36,8 @@
 #include "EngineUtils.h"
 #include "GameMapsSettings.h"
 #include "ScopedTransaction.h"
+#include "Features/IModularFeatures.h"
+#include "Features/EditorFeatures.h"
 
 namespace LevelEditorActionHelpers
 {
@@ -1675,9 +1677,14 @@ TSharedRef< SWidget > FLevelEditorToolBar::GenerateQuickSettingsMenu( TSharedRef
 		MenuBuilder.AddMenuEntry(
 			LOCTEXT("ProjectSettingsMenuLabel", "Project Settings..."),
 			LOCTEXT("ProjectSettingsMenuToolTip", "Change the settings of the currently loaded project"),
-			FSlateIcon(),
+			FSlateIcon(FEditorStyle::GetStyleSetName(), "ProjectSettings.TabIcon"),
 			FUIAction(FExecuteAction::CreateStatic(&Local::OpenSettings, FName("Project"), FName("Project"), FName("General")))
 			);
+
+		if (IModularFeatures::Get().IsModularFeatureAvailable(EditorFeatures::PluginsEditor))
+		{
+			FGlobalTabmanager::Get()->PopulateTabSpawnerMenu(MenuBuilder, "PluginsEditor");
+		}
 	}
 	MenuBuilder.EndSection();
 
