Made the boolean pin in Blueprint's less subtle.

Increased the size and made it more white.

diff --git a/Engine/Source/Editor/EditorStyle/Private/SlateEditorStyle.cpp b/Engine/Source/Editor/EditorStyle/Private/SlateEditorStyle.cpp
--- a/Engine/Source/Editor/EditorStyle/Private/SlateEditorStyle.cpp
+++ b/Engine/Source/Editor/EditorStyle/Private/SlateEditorStyle.cpp
@@ -4039,15 +4039,15 @@ void FSlateEditorStyle::FStyle::SetupGraphEditorStyles()
 			{
 				/* Set images for various SCheckBox states of style Graph.Checkbox ... */
 				const FCheckBoxStyle BasicGraphCheckBoxStyle = FCheckBoxStyle()
-					.SetUncheckedImage( IMAGE_BRUSH( "/Graph/CommonWidgets/CheckBox", Icon16x16 ) )
-					.SetUncheckedHoveredImage( IMAGE_BRUSH( "/Graph/CommonWidgets/CheckBox_Hovered", Icon16x16 ) )
-					.SetUncheckedPressedImage( IMAGE_BRUSH( "/Graph/CommonWidgets/CheckBox_Hovered", Icon16x16 ) )
-					.SetCheckedImage( IMAGE_BRUSH( "/Graph/CommonWidgets/CheckBox_Checked", Icon16x16 ) )
-					.SetCheckedHoveredImage( IMAGE_BRUSH( "/Graph/CommonWidgets/CheckBox_Checked_Hovered", Icon16x16 ) )
-					.SetCheckedPressedImage( IMAGE_BRUSH( "/Graph/CommonWidgets/CheckBox_Checked", Icon16x16, FLinearColor( 0.5f, 0.5f, 0.5f ) ) )
-					.SetUndeterminedImage( IMAGE_BRUSH( "/Graph/CommonWidgets/CheckBox_Undetermined", Icon16x16 ) )
-					.SetUndeterminedHoveredImage( IMAGE_BRUSH( "/Graph/CommonWidgets/CheckBox_Undetermined_Hovered", Icon16x16 ) )
-					.SetUndeterminedPressedImage( IMAGE_BRUSH( "/Graph/CommonWidgets/CheckBox_Undetermined_Hovered", Icon16x16, FLinearColor( 0.5f, 0.5f, 0.5f ) ) );
+					.SetUncheckedImage( IMAGE_BRUSH( "/Graph/CommonWidgets/CheckBox", Icon20x20 ) )
+					.SetUncheckedHoveredImage( IMAGE_BRUSH( "/Graph/CommonWidgets/CheckBox_Hovered", Icon20x20 ) )
+					.SetUncheckedPressedImage( IMAGE_BRUSH( "/Graph/CommonWidgets/CheckBox_Hovered", Icon20x20 ) )
+					.SetCheckedImage( IMAGE_BRUSH( "/Graph/CommonWidgets/CheckBox_Checked", Icon20x20 ) )
+					.SetCheckedHoveredImage( IMAGE_BRUSH( "/Graph/CommonWidgets/CheckBox_Checked_Hovered", Icon20x20 ) )
+					.SetCheckedPressedImage( IMAGE_BRUSH( "/Graph/CommonWidgets/CheckBox_Checked", Icon20x20, FLinearColor( 0.5f, 0.5f, 0.5f ) ) )
+					.SetUndeterminedImage( IMAGE_BRUSH( "/Graph/CommonWidgets/CheckBox_Undetermined", Icon20x20 ) )
+					.SetUndeterminedHoveredImage( IMAGE_BRUSH( "/Graph/CommonWidgets/CheckBox_Undetermined_Hovered", Icon20x20 ) )
+					.SetUndeterminedPressedImage( IMAGE_BRUSH( "/Graph/CommonWidgets/CheckBox_Undetermined_Hovered", Icon20x20, FLinearColor( 0.5f, 0.5f, 0.5f ) ) );
 
 				/* ... and add the new style */
 				Set( "Graph.Checkbox", BasicGraphCheckBoxStyle );
