Removed usage of disabled button style which has invalid art work

diff --git a/Engine/Source/Editor/EditorStyle/Private/SlateEditorStyle.cpp b/Engine/Source/Editor/EditorStyle/Private/SlateEditorStyle.cpp
--- a/Engine/Source/Editor/EditorStyle/Private/SlateEditorStyle.cpp
+++ b/Engine/Source/Editor/EditorStyle/Private/SlateEditorStyle.cpp
@@ -342,7 +342,6 @@ void FSlateEditorStyle::FStyle::SetupGeneralStyles()
 		.SetNormal( BOX_BRUSH( "Common/Button", FVector2D(32,32), 8.0f/32.0f ) )
 		.SetHovered( BOX_BRUSH( "Common/Button_Hovered", FVector2D(32,32), 8.0f/32.0f ) )
 		.SetPressed( BOX_BRUSH( "Common/Button_Pressed", FVector2D(32,32), 8.0f/32.0f ) )
-		.SetDisabled( BOX_BRUSH( "Common/Button_Disabled", 8.0f/32.0f ) )
 		.SetNormalPadding( FMargin( 2,2,2,2 ) )
 		.SetPressedPadding( FMargin( 2,3,2,1 ) );
 	Set( "Button", Button );
@@ -2501,8 +2500,7 @@ void FSlateEditorStyle::FStyle::SetupWindowStyles()
 		const FButtonStyle MinimizeButtonStyle = FButtonStyle(Button)
 			.SetNormal ( IMAGE_BRUSH( "Old/Window/WindowButton_Minimize_Normal", FVector2D(27, 18) ) )
 			.SetHovered( IMAGE_BRUSH( "Old/Window/WindowButton_Minimize_Hovered", FVector2D(27, 18) ) )
-			.SetPressed( IMAGE_BRUSH( "Old/Window/WindowButton_Minimize_Pressed", FVector2D(27, 18) ) )
-			.SetDisabled( IMAGE_BRUSH( "Old/Window/WindowButton_Minimize_Disabled", FVector2D(27, 18) ) );
+			.SetPressed( IMAGE_BRUSH( "Old/Window/WindowButton_Minimize_Pressed", FVector2D(27, 18) ) );
 
 		Set( "Window.Buttons.Minimize.Normal",    new IMAGE_BRUSH( "Old/Window/WindowButton_Minimize_Normal", FVector2D(27, 18) ) );
 		Set( "Window.Buttons.Minimize.Hovered",   new IMAGE_BRUSH( "Old/Window/WindowButton_Minimize_Hovered", FVector2D(27, 18) ) );
@@ -2512,8 +2510,7 @@ void FSlateEditorStyle::FStyle::SetupWindowStyles()
 		const FButtonStyle MaximizeButtonStyle = FButtonStyle(Button)
 			.SetNormal ( IMAGE_BRUSH( "Old/Window/WindowButton_Maximize_Normal", FVector2D(23, 18) ) )
 			.SetHovered( IMAGE_BRUSH( "Old/Window/WindowButton_Maximize_Hovered", FVector2D(23, 18) ) )
-			.SetPressed( IMAGE_BRUSH( "Old/Window/WindowButton_Maximize_Pressed", FVector2D(23, 18) ) )
-			.SetDisabled( IMAGE_BRUSH( "Old/Window/WindowButton_Maximize_Disabled", FVector2D(23, 18) ) );
+			.SetPressed( IMAGE_BRUSH( "Old/Window/WindowButton_Maximize_Pressed", FVector2D(23, 18) ) );
 
 		Set( "Window.Buttons.Maximize.Normal",    new IMAGE_BRUSH( "Old/Window/WindowButton_Maximize_Normal", FVector2D(23, 18) ) );
 		Set( "Window.Buttons.Maximize.Hovered",   new IMAGE_BRUSH( "Old/Window/WindowButton_Maximize_Hovered", FVector2D(23, 18) ) );
@@ -2762,7 +2759,6 @@ void FSlateEditorStyle::FStyle::SetupTutorialStyles()
 			.SetNormal( BOX_BRUSH( "Common/ButtonHoverHint", FMargin(4/16.0f), FLinearColor(0.05f,0.05f,0.05f,1) ) )
 			.SetHovered( BOX_BRUSH( "Common/ButtonHoverHint", FMargin(4/16.0f), FLinearColor(0.07f,0.07f,0.07f,1) ) )
 			.SetPressed( BOX_BRUSH( "Common/ButtonHoverHint", FMargin(4/16.0f), FLinearColor(0.08f,0.08f,0.08f,1) ) )
-			.SetDisabled( BOX_BRUSH( "Common/ButtonHoverHint", FMargin(4/16.0f), FLinearColor(0.1f,0.1f,0.1f,1) ) )
 			.SetNormalPadding( FMargin(0,0,0,1))
 			.SetPressedPadding( FMargin(0,1,0,0)));
 
@@ -2770,7 +2766,6 @@ void FSlateEditorStyle::FStyle::SetupTutorialStyles()
 			.SetNormal( BOX_BRUSH( "Common/ButtonHoverHint", FMargin(4/16.0f), FLinearColor(1.0f,1.0f,1.0f,0.0f) ) )
 			.SetHovered( BOX_BRUSH( "Common/ButtonHoverHint", FMargin(4/16.0f), FLinearColor(1.0f,1.0f,1.0f,0.05f) ) )
 			.SetPressed( BOX_BRUSH( "Common/ButtonHoverHint", FMargin(4/16.0f), FLinearColor(1.0f,1.0f,1.0f,0.05f) ) )
-			.SetDisabled( BOX_BRUSH( "Common/ButtonHoverHint", FMargin(4/16.0f), FLinearColor(1.0f,1.0f,1.0f,0.05f) ) )
 			.SetNormalPadding( FMargin(0,0,0,1))
 			.SetPressedPadding( FMargin(0,1,0,0)));
 
@@ -2778,7 +2773,6 @@ void FSlateEditorStyle::FStyle::SetupTutorialStyles()
 			.SetNormal( BOX_BRUSH( "Common/ButtonHoverHint", FMargin(4/16.0f), FLinearColor(0,0,0,0) ) )
 			.SetHovered( BOX_BRUSH( "Common/ButtonHoverHint", FMargin(4/16.0f), FLinearColor(1,1,1,1) ) )
 			.SetPressed( BOX_BRUSH( "Common/ButtonHoverHint", FMargin(4/16.0f), FLinearColor(1,1,1,1) ) )
-			.SetDisabled( BOX_BRUSH( "Common/ButtonHoverHint", FMargin(4/16.0f), FLinearColor(1,1,1,1) ) )
 			.SetNormalPadding( FMargin(0,0,0,1))
 			.SetPressedPadding( FMargin(0,1,0,0)));
 
@@ -2786,7 +2780,6 @@ void FSlateEditorStyle::FStyle::SetupTutorialStyles()
 			.SetNormal( FSlateNoResource() )
 			.SetHovered( FSlateNoResource() )
 			.SetPressed( FSlateNoResource() )
-			.SetDisabled( FSlateNoResource() )
 			.SetNormalPadding( FMargin(0,0,0,1))
 			.SetPressedPadding( FMargin(0,1,0,0)));
 
@@ -2794,7 +2787,6 @@ void FSlateEditorStyle::FStyle::SetupTutorialStyles()
 			.SetNormal( BOX_BRUSH( "Common/ButtonHoverHint", FMargin(4/16.0f), TutorialNavigationButtonColor ) )
 			.SetHovered( BOX_BRUSH( "Common/ButtonHoverHint", FMargin(4/16.0f), TutorialNavigationButtonHoverColor ) )
 			.SetPressed( BOX_BRUSH( "Common/ButtonHoverHint", FMargin(4/16.0f), TutorialNavigationButtonHoverColor ) )
-			.SetDisabled( BOX_BRUSH( "Common/ButtonHoverHint", FMargin(4/16.0f), TutorialNavigationButtonColor ) )
 			.SetNormalPadding( FMargin(0,0,0,1))
 			.SetPressedPadding( FMargin(0,1,0,0)));
 
@@ -2802,7 +2794,6 @@ void FSlateEditorStyle::FStyle::SetupTutorialStyles()
 			.SetNormal(BOX_BRUSH("Common/ButtonHoverHint", FMargin(4 / 16.0f), TutorialNavigationBackButtonColor))
 			.SetHovered(BOX_BRUSH("Common/ButtonHoverHint", FMargin(4 / 16.0f), TutorialNavigationBackButtonHoverColor))
 			.SetPressed(BOX_BRUSH("Common/ButtonHoverHint", FMargin(4 / 16.0f), TutorialNavigationBackButtonHoverColor))
-			.SetDisabled(BOX_BRUSH("Common/ButtonHoverHint", FMargin(4 / 16.0f), TutorialNavigationBackButtonColor))
 			.SetNormalPadding(FMargin(0, 0, 0, 1))
 			.SetPressedPadding(FMargin(0, 1, 0, 0)));
 
@@ -2829,7 +2820,6 @@ void FSlateEditorStyle::FStyle::SetupTutorialStyles()
 			.SetNormal( BOX_BRUSH( "Common/ButtonHoverHint", FMargin(4/16.0f), FLinearColor(0,0,0,0) ) )
 			.SetHovered( BOX_BRUSH( "Common/ButtonHoverHint", FMargin(4/16.0f), FLinearColor(0,0,0,0) ) )
 			.SetPressed( BOX_BRUSH( "Common/ButtonHoverHint", FMargin(4/16.0f), FLinearColor(0,0,0,0) ) )
-			.SetDisabled( BOX_BRUSH( "Common/ButtonHoverHint", FMargin(4/16.0f), FLinearColor(0,0,0,0) ) )
 			.SetNormalPadding( FMargin(0,0,0,1))
 			.SetPressedPadding( FMargin(0,1,0,0)));
 
@@ -2990,7 +2980,6 @@ void FSlateEditorStyle::FStyle::SetupTutorialStyles()
 				.SetNormal(BOX_BRUSH("Tutorials/FlatColorSquare", FVector2D(1.0f, 1.0f), FMargin(1), NormalColor))
 				.SetHovered(BOX_BRUSH("Tutorials/FlatColorSquare", FVector2D(1.0f, 1.0f), FMargin(1), HoverColor))
 				.SetPressed(BOX_BRUSH("Tutorials/FlatColorSquare", FVector2D(1.0f, 1.0f), FMargin(1), SelectedColor))
-				.SetDisabled(BOX_BRUSH("Tutorials/FlatColorSquare", FVector2D(1.0f, 1.0f), FMargin(1), DisabledColor))
 				.SetNormalPadding(FMargin(2, 2, 2, 2))
 				.SetPressedPadding(FMargin(2, 3, 2, 1));
 			Set("TutorialEditableText.Toolbar.Button", TutorialButton);
@@ -3007,7 +2996,6 @@ void FSlateEditorStyle::FStyle::SetupTutorialStyles()
 					.SetNormal(BOX_BRUSH("Tutorials/FlatColorSquare", FVector2D(1.0f, 1.0f), FMargin(1), FLinearColor::White))
 					.SetHovered(BOX_BRUSH("Tutorials/FlatColorSquare", FVector2D(1.0f, 1.0f), FMargin(1), FLinearColor::White))
 					.SetPressed(BOX_BRUSH("Tutorials/FlatColorSquare", FVector2D(1.0f, 1.0f), FMargin(1), FLinearColor::White))
-					.SetDisabled(BOX_BRUSH("Tutorials/FlatColorSquare", FVector2D(1.0f, 1.0f), FMargin(1), DisabledColor))
 					.SetNormalPadding(FMargin(2, 2, 2, 2))
 					.SetPressedPadding(FMargin(2, 3, 2, 1));
 
diff --git a/Engine/Source/Runtime/AppFramework/Private/Framework/Testing/TestStyle.cpp b/Engine/Source/Runtime/AppFramework/Private/Framework/Testing/TestStyle.cpp
--- a/Engine/Source/Runtime/AppFramework/Private/Framework/Testing/TestStyle.cpp
+++ b/Engine/Source/Runtime/AppFramework/Private/Framework/Testing/TestStyle.cpp
@@ -322,7 +322,6 @@ TSharedRef< ISlateStyle > FTestStyle::Create()
 					.SetNormal( BOX_BRUSH("Testing/FlatColorSquare", FVector2D(1.0f, 1.0f), FMargin(1), NormalColor ) )
 					.SetHovered( BOX_BRUSH("Testing/FlatColorSquare", FVector2D(1.0f, 1.0f), FMargin(1), HoverColor ) )
 					.SetPressed( BOX_BRUSH("Testing/FlatColorSquare",  FVector2D(1.0f, 1.0f), FMargin(1), SelectedColor ) )
-					.SetDisabled( BOX_BRUSH("Testing/FlatColorSquare",  FVector2D(1.0f, 1.0f), FMargin(1), DisabledColor ) )
 					.SetNormalPadding( FMargin( 2,2,2,2 ) )
 					.SetPressedPadding( FMargin( 2,3,2,1 ) );
 				Style->Set( "RichText.Toolbar.Button", Button );
@@ -339,7 +338,6 @@ TSharedRef< ISlateStyle > FTestStyle::Create()
 						.SetNormal( BOX_BRUSH("Testing/FlatColorSquare", FVector2D(1.0f, 1.0f), FMargin(1), FLinearColor::White ) )
 						.SetHovered( BOX_BRUSH("Testing/FlatColorSquare", FVector2D(1.0f, 1.0f), FMargin(1), FLinearColor::White ) )
 						.SetPressed( BOX_BRUSH("Testing/FlatColorSquare",  FVector2D(1.0f, 1.0f), FMargin(1), FLinearColor::White ) )
-						.SetDisabled( BOX_BRUSH("Testing/FlatColorSquare",  FVector2D(1.0f, 1.0f), FMargin(1), DisabledColor ) )
 						.SetNormalPadding( FMargin( 2,2,2,2 ) )
 						.SetPressedPadding( FMargin( 2,3,2,1 ) );
 
diff --git a/Engine/Source/Runtime/SlateCore/Private/Styling/CoreStyle.cpp b/Engine/Source/Runtime/SlateCore/Private/Styling/CoreStyle.cpp
--- a/Engine/Source/Runtime/SlateCore/Private/Styling/CoreStyle.cpp
+++ b/Engine/Source/Runtime/SlateCore/Private/Styling/CoreStyle.cpp
@@ -296,7 +296,6 @@ TSharedRef<ISlateStyle> FCoreStyle::Create( const FName& InStyleSetName )
 		.SetNormal( BOX_BRUSH( "Common/Button", FVector2D(32,32), 8.0f/32.0f ) )
 		.SetHovered( BOX_BRUSH( "Common/Button_Hovered", FVector2D(32,32), 8.0f/32.0f ) )
 		.SetPressed( BOX_BRUSH( "Common/Button_Pressed", FVector2D(32,32), 8.0f/32.0f ) )
-		.SetDisabled( BOX_BRUSH( "Common/Button_Disabled", 8.0f/32.0f ) )
 		.SetNormalPadding( FMargin( 2,2,2,2 ) )
 		.SetPressedPadding( FMargin( 2,3,2,1 ) );
 	{
@@ -1062,14 +1061,12 @@ TSharedRef<ISlateStyle> FCoreStyle::Create( const FName& InStyleSetName )
 		const FButtonStyle MinimizeButtonStyle = FButtonStyle(Button)
 			.SetNormal (IMAGE_BRUSH("Common/Window/WindowButton_Minimize_Normal", FVector2D(27.0f, 18.0f)))
 			.SetHovered(IMAGE_BRUSH("Common/Window/WindowButton_Minimize_Hovered", FVector2D(27.0f, 18.0f)))
-			.SetPressed(IMAGE_BRUSH("Common/Window/WindowButton_Minimize_Pressed", FVector2D(27.0f, 18.0f)))
-			.SetDisabled(IMAGE_BRUSH("Common/Window/WindowButton_Minimize_Disabled", FVector2D(27.0f, 18.0f)));
+			.SetPressed(IMAGE_BRUSH("Common/Window/WindowButton_Minimize_Pressed", FVector2D(27.0f, 18.0f)));
 
 		const FButtonStyle MaximizeButtonStyle = FButtonStyle(Button)
 			.SetNormal (IMAGE_BRUSH("Common/Window/WindowButton_Maximize_Normal", FVector2D(23.0f, 18.0f)))
 			.SetHovered(IMAGE_BRUSH("Common/Window/WindowButton_Maximize_Hovered", FVector2D(23.0f, 18.0f)))
-			.SetPressed(IMAGE_BRUSH("Common/Window/WindowButton_Maximize_Pressed", FVector2D(23.0f, 18.0f)))
-			.SetDisabled(IMAGE_BRUSH("Common/Window/WindowButton_Maximize_Disabled", FVector2D(23.0f, 18.0f)));
+			.SetPressed(IMAGE_BRUSH("Common/Window/WindowButton_Maximize_Pressed", FVector2D(23.0f, 18.0f)));
 
 		const FButtonStyle RestoreButtonStyle = FButtonStyle(Button)
 			.SetNormal (IMAGE_BRUSH("Common/Window/WindowButton_Restore_Normal", FVector2D(23.0f, 18)))
