[UEDOC-1711] Fixed color on "previous" button in in-editor tutorials.

diff --git a/Engine/Source/Editor/EditorStyle/Private/SlateEditorStyle.cpp b/Engine/Source/Editor/EditorStyle/Private/SlateEditorStyle.cpp
--- a/Engine/Source/Editor/EditorStyle/Private/SlateEditorStyle.cpp
+++ b/Engine/Source/Editor/EditorStyle/Private/SlateEditorStyle.cpp
@@ -2727,8 +2727,8 @@ void FSlateEditorStyle::FStyle::SetupTutorialStyles()
 		const FLinearColor TutorialSelectionColor = FLinearColor(0.19f, 0.33f, 0.72f);
 		const FLinearColor TutorialNavigationButtonColor = FLinearColor(0.0f, 0.59f, 0.14f, 1.0f);
 		const FLinearColor TutorialNavigationButtonHoverColor = FLinearColor(0.2f, 0.79f, 0.34f, 1.0f);
-		const FLinearColor TutorialNavigationBackButtonColor = FLinearColor(0.4f, 0.4f, 0.4f, 0.75f);
-		const FLinearColor TutorialNavigationBackButtonHoverColor = FLinearColor(0.2f, 0.2f, 0.2f, 0.75f);
+		const FLinearColor TutorialNavigationBackButtonColor = TutorialNavigationButtonColor;
+		const FLinearColor TutorialNavigationBackButtonHoverColor = TutorialNavigationButtonHoverColor;
 
 		const FTextBlockStyle TutorialText = FTextBlockStyle(DocumentationText)
 			.SetColorAndOpacity(FLinearColor::Black)
