Editor - Adding a generated header containing all the glyphs contained in the font awesome glyph set for use in the editor.  Updating the SContentBrowser to use the new shortcuts instead of hardcoding it.

diff --git a/Engine/Source/Editor/ContentBrowser/Private/SContentBrowser.cpp b/Engine/Source/Editor/ContentBrowser/Private/SContentBrowser.cpp
--- a/Engine/Source/Editor/ContentBrowser/Private/SContentBrowser.cpp
+++ b/Engine/Source/Editor/ContentBrowser/Private/SContentBrowser.cpp
@@ -130,7 +130,7 @@ void SContentBrowser::Construct( const FArguments& InArgs, const FName& InInstan
 									SNew(STextBlock)
 									.TextStyle(FEditorStyle::Get(), "ContentBrowser.TopBar.Font")
 									.Font(FEditorStyle::Get().GetFontStyle("FontAwesome.11"))
-									.Text(FText::FromString(FString(TEXT("\xf15b"))) /*fa-file*/)
+									.Text(FEditorFontGlyphs::File)
 								]
 
 								// New Text
@@ -153,7 +153,7 @@ void SContentBrowser::Construct( const FArguments& InArgs, const FName& InInstan
 									SNew(STextBlock)
 									.TextStyle(FEditorStyle::Get(), "ContentBrowser.TopBar.Font")
 									.Font(FEditorStyle::Get().GetFontStyle("FontAwesome.10"))
-									.Text(FText::FromString(FString(TEXT("\xf0d7"))) /*fa-caret-down*/)
+									.Text(FEditorFontGlyphs::Caret_Down)
 								]
 							]
 						]
@@ -183,7 +183,7 @@ void SContentBrowser::Construct( const FArguments& InArgs, const FName& InInstan
 									SNew(STextBlock)
 									.TextStyle(FEditorStyle::Get(), "ContentBrowser.TopBar.Font")
 									.Font(FEditorStyle::Get().GetFontStyle("FontAwesome.11"))
-									.Text(FText::FromString(FString(TEXT("\xf019"))) /*fa-download*/)
+									.Text(FEditorFontGlyphs::Download)
 								]
 
 								// Import Text
@@ -222,7 +222,7 @@ void SContentBrowser::Construct( const FArguments& InArgs, const FName& InInstan
 									SNew(STextBlock)
 									.TextStyle(FEditorStyle::Get(), "ContentBrowser.TopBar.Font")
 									.Font(FEditorStyle::Get().GetFontStyle("FontAwesome.11"))
-									.Text(FText::FromString(FString(TEXT("\xf0c7"))) /*fa-floppy-o*/)
+									.Text(FEditorFontGlyphs::Floppy_O)
 								]
 
 								// Save All Text
diff --git a/Engine/Source/Editor/EditorStyle/Private/EditorFontGlyphs.cpp b/Engine/Source/Editor/EditorStyle/Private/EditorFontGlyphs.cpp
new file mode 100644
--- /dev/null
+++ b/Engine/Source/Editor/EditorStyle/Private/EditorFontGlyphs.cpp
@@ -0,0 +1,523 @@
+// Copyright 1998-2015 Epic Games, Inc. All Rights Reserved.
+
+#include "EditorStylePrivatePCH.h"
+
+FText FEditorFontGlyphs::Adjust( FText::FromString( FString( TEXT("\xf042") ) ) );
+FText FEditorFontGlyphs::Adn( FText::FromString( FString( TEXT("\xf170") ) ) );
+FText FEditorFontGlyphs::Align_Center( FText::FromString( FString( TEXT("\xf037") ) ) );
+FText FEditorFontGlyphs::Align_Justify( FText::FromString( FString( TEXT("\xf039") ) ) );
+FText FEditorFontGlyphs::Align_Left( FText::FromString( FString( TEXT("\xf036") ) ) );
+FText FEditorFontGlyphs::Align_Right( FText::FromString( FString( TEXT("\xf038") ) ) );
+FText FEditorFontGlyphs::Ambulance( FText::FromString( FString( TEXT("\xf0f9") ) ) );
+FText FEditorFontGlyphs::Anchor( FText::FromString( FString( TEXT("\xf13d") ) ) );
+FText FEditorFontGlyphs::Android( FText::FromString( FString( TEXT("\xf17b") ) ) );
+FText FEditorFontGlyphs::Angellist( FText::FromString( FString( TEXT("\xf209") ) ) );
+FText FEditorFontGlyphs::Angle_Double_Down( FText::FromString( FString( TEXT("\xf103") ) ) );
+FText FEditorFontGlyphs::Angle_Double_Left( FText::FromString( FString( TEXT("\xf100") ) ) );
+FText FEditorFontGlyphs::Angle_Double_Right( FText::FromString( FString( TEXT("\xf101") ) ) );
+FText FEditorFontGlyphs::Angle_Double_Up( FText::FromString( FString( TEXT("\xf102") ) ) );
+FText FEditorFontGlyphs::Angle_Down( FText::FromString( FString( TEXT("\xf107") ) ) );
+FText FEditorFontGlyphs::Angle_Left( FText::FromString( FString( TEXT("\xf104") ) ) );
+FText FEditorFontGlyphs::Angle_Right( FText::FromString( FString( TEXT("\xf105") ) ) );
+FText FEditorFontGlyphs::Angle_Up( FText::FromString( FString( TEXT("\xf106") ) ) );
+FText FEditorFontGlyphs::Apple( FText::FromString( FString( TEXT("\xf179") ) ) );
+FText FEditorFontGlyphs::Archive( FText::FromString( FString( TEXT("\xf187") ) ) );
+FText FEditorFontGlyphs::Area_Chart( FText::FromString( FString( TEXT("\xf1fe") ) ) );
+FText FEditorFontGlyphs::Arrow_Circle_Down( FText::FromString( FString( TEXT("\xf0ab") ) ) );
+FText FEditorFontGlyphs::Arrow_Circle_Left( FText::FromString( FString( TEXT("\xf0a8") ) ) );
+FText FEditorFontGlyphs::Arrow_Circle_O_Down( FText::FromString( FString( TEXT("\xf01a") ) ) );
+FText FEditorFontGlyphs::Arrow_Circle_O_Left( FText::FromString( FString( TEXT("\xf190") ) ) );
+FText FEditorFontGlyphs::Arrow_Circle_O_Right( FText::FromString( FString( TEXT("\xf18e") ) ) );
+FText FEditorFontGlyphs::Arrow_Circle_O_Up( FText::FromString( FString( TEXT("\xf01b") ) ) );
+FText FEditorFontGlyphs::Arrow_Circle_Right( FText::FromString( FString( TEXT("\xf0a9") ) ) );
+FText FEditorFontGlyphs::Arrow_Circle_Up( FText::FromString( FString( TEXT("\xf0aa") ) ) );
+FText FEditorFontGlyphs::Arrow_Down( FText::FromString( FString( TEXT("\xf063") ) ) );
+FText FEditorFontGlyphs::Arrow_Left( FText::FromString( FString( TEXT("\xf060") ) ) );
+FText FEditorFontGlyphs::Arrow_Right( FText::FromString( FString( TEXT("\xf061") ) ) );
+FText FEditorFontGlyphs::Arrows( FText::FromString( FString( TEXT("\xf047") ) ) );
+FText FEditorFontGlyphs::Arrows_Alt( FText::FromString( FString( TEXT("\xf0b2") ) ) );
+FText FEditorFontGlyphs::Arrows_H( FText::FromString( FString( TEXT("\xf07e") ) ) );
+FText FEditorFontGlyphs::Arrows_V( FText::FromString( FString( TEXT("\xf07d") ) ) );
+FText FEditorFontGlyphs::Arrow_Up( FText::FromString( FString( TEXT("\xf062") ) ) );
+FText FEditorFontGlyphs::Asterisk( FText::FromString( FString( TEXT("\xf069") ) ) );
+FText FEditorFontGlyphs::At( FText::FromString( FString( TEXT("\xf1fa") ) ) );
+FText FEditorFontGlyphs::Backward( FText::FromString( FString( TEXT("\xf04a") ) ) );
+FText FEditorFontGlyphs::Ban( FText::FromString( FString( TEXT("\xf05e") ) ) );
+FText FEditorFontGlyphs::Bar_Chart( FText::FromString( FString( TEXT("\xf080") ) ) );
+FText FEditorFontGlyphs::Barcode( FText::FromString( FString( TEXT("\xf02a") ) ) );
+FText FEditorFontGlyphs::Bars( FText::FromString( FString( TEXT("\xf0c9") ) ) );
+FText FEditorFontGlyphs::Bed( FText::FromString( FString( TEXT("\xf236") ) ) );
+FText FEditorFontGlyphs::Beer( FText::FromString( FString( TEXT("\xf0fc") ) ) );
+FText FEditorFontGlyphs::Behance( FText::FromString( FString( TEXT("\xf1b4") ) ) );
+FText FEditorFontGlyphs::Behance_Square( FText::FromString( FString( TEXT("\xf1b5") ) ) );
+FText FEditorFontGlyphs::Bell( FText::FromString( FString( TEXT("\xf0f3") ) ) );
+FText FEditorFontGlyphs::Bell_O( FText::FromString( FString( TEXT("\xf0a2") ) ) );
+FText FEditorFontGlyphs::Bell_Slash( FText::FromString( FString( TEXT("\xf1f6") ) ) );
+FText FEditorFontGlyphs::Bell_Slash_O( FText::FromString( FString( TEXT("\xf1f7") ) ) );
+FText FEditorFontGlyphs::Bicycle( FText::FromString( FString( TEXT("\xf206") ) ) );
+FText FEditorFontGlyphs::Binoculars( FText::FromString( FString( TEXT("\xf1e5") ) ) );
+FText FEditorFontGlyphs::Birthday_Cake( FText::FromString( FString( TEXT("\xf1fd") ) ) );
+FText FEditorFontGlyphs::Bitbucket( FText::FromString( FString( TEXT("\xf171") ) ) );
+FText FEditorFontGlyphs::Bitbucket_Square( FText::FromString( FString( TEXT("\xf172") ) ) );
+FText FEditorFontGlyphs::Bold( FText::FromString( FString( TEXT("\xf032") ) ) );
+FText FEditorFontGlyphs::Bolt( FText::FromString( FString( TEXT("\xf0e7") ) ) );
+FText FEditorFontGlyphs::Bomb( FText::FromString( FString( TEXT("\xf1e2") ) ) );
+FText FEditorFontGlyphs::Book( FText::FromString( FString( TEXT("\xf02d") ) ) );
+FText FEditorFontGlyphs::Bookmark( FText::FromString( FString( TEXT("\xf02e") ) ) );
+FText FEditorFontGlyphs::Bookmark_O( FText::FromString( FString( TEXT("\xf097") ) ) );
+FText FEditorFontGlyphs::Briefcase( FText::FromString( FString( TEXT("\xf0b1") ) ) );
+FText FEditorFontGlyphs::Btc( FText::FromString( FString( TEXT("\xf15a") ) ) );
+FText FEditorFontGlyphs::Bug( FText::FromString( FString( TEXT("\xf188") ) ) );
+FText FEditorFontGlyphs::Building( FText::FromString( FString( TEXT("\xf1ad") ) ) );
+FText FEditorFontGlyphs::Building_O( FText::FromString( FString( TEXT("\xf0f7") ) ) );
+FText FEditorFontGlyphs::Bullhorn( FText::FromString( FString( TEXT("\xf0a1") ) ) );
+FText FEditorFontGlyphs::Bullseye( FText::FromString( FString( TEXT("\xf140") ) ) );
+FText FEditorFontGlyphs::Bus( FText::FromString( FString( TEXT("\xf207") ) ) );
+FText FEditorFontGlyphs::Buysellads( FText::FromString( FString( TEXT("\xf20d") ) ) );
+FText FEditorFontGlyphs::Calculator( FText::FromString( FString( TEXT("\xf1ec") ) ) );
+FText FEditorFontGlyphs::Calendar( FText::FromString( FString( TEXT("\xf073") ) ) );
+FText FEditorFontGlyphs::Calendar_O( FText::FromString( FString( TEXT("\xf133") ) ) );
+FText FEditorFontGlyphs::Camera( FText::FromString( FString( TEXT("\xf030") ) ) );
+FText FEditorFontGlyphs::Camera_Retro( FText::FromString( FString( TEXT("\xf083") ) ) );
+FText FEditorFontGlyphs::Car( FText::FromString( FString( TEXT("\xf1b9") ) ) );
+FText FEditorFontGlyphs::Caret_Down( FText::FromString( FString( TEXT("\xf0d7") ) ) );
+FText FEditorFontGlyphs::Caret_Left( FText::FromString( FString( TEXT("\xf0d9") ) ) );
+FText FEditorFontGlyphs::Caret_Right( FText::FromString( FString( TEXT("\xf0da") ) ) );
+FText FEditorFontGlyphs::Caret_Square_O_Down( FText::FromString( FString( TEXT("\xf150") ) ) );
+FText FEditorFontGlyphs::Caret_Square_O_Left( FText::FromString( FString( TEXT("\xf191") ) ) );
+FText FEditorFontGlyphs::Caret_Square_O_Right( FText::FromString( FString( TEXT("\xf152") ) ) );
+FText FEditorFontGlyphs::Caret_Square_O_Up( FText::FromString( FString( TEXT("\xf151") ) ) );
+FText FEditorFontGlyphs::Caret_Up( FText::FromString( FString( TEXT("\xf0d8") ) ) );
+FText FEditorFontGlyphs::Cart_Arrow_Down( FText::FromString( FString( TEXT("\xf218") ) ) );
+FText FEditorFontGlyphs::Cart_Plus( FText::FromString( FString( TEXT("\xf217") ) ) );
+FText FEditorFontGlyphs::Cc( FText::FromString( FString( TEXT("\xf20a") ) ) );
+FText FEditorFontGlyphs::Cc_Amex( FText::FromString( FString( TEXT("\xf1f3") ) ) );
+FText FEditorFontGlyphs::Cc_Discover( FText::FromString( FString( TEXT("\xf1f2") ) ) );
+FText FEditorFontGlyphs::Cc_Mastercard( FText::FromString( FString( TEXT("\xf1f1") ) ) );
+FText FEditorFontGlyphs::Cc_Paypal( FText::FromString( FString( TEXT("\xf1f4") ) ) );
+FText FEditorFontGlyphs::Cc_Stripe( FText::FromString( FString( TEXT("\xf1f5") ) ) );
+FText FEditorFontGlyphs::Cc_Visa( FText::FromString( FString( TEXT("\xf1f0") ) ) );
+FText FEditorFontGlyphs::Certificate( FText::FromString( FString( TEXT("\xf0a3") ) ) );
+FText FEditorFontGlyphs::Chain_Broken( FText::FromString( FString( TEXT("\xf127") ) ) );
+FText FEditorFontGlyphs::Check( FText::FromString( FString( TEXT("\xf00c") ) ) );
+FText FEditorFontGlyphs::Check_Circle( FText::FromString( FString( TEXT("\xf058") ) ) );
+FText FEditorFontGlyphs::Check_Circle_O( FText::FromString( FString( TEXT("\xf05d") ) ) );
+FText FEditorFontGlyphs::Check_Square( FText::FromString( FString( TEXT("\xf14a") ) ) );
+FText FEditorFontGlyphs::Check_Square_O( FText::FromString( FString( TEXT("\xf046") ) ) );
+FText FEditorFontGlyphs::Chevron_Circle_Down( FText::FromString( FString( TEXT("\xf13a") ) ) );
+FText FEditorFontGlyphs::Chevron_Circle_Left( FText::FromString( FString( TEXT("\xf137") ) ) );
+FText FEditorFontGlyphs::Chevron_Circle_Right( FText::FromString( FString( TEXT("\xf138") ) ) );
+FText FEditorFontGlyphs::Chevron_Circle_Up( FText::FromString( FString( TEXT("\xf139") ) ) );
+FText FEditorFontGlyphs::Chevron_Down( FText::FromString( FString( TEXT("\xf078") ) ) );
+FText FEditorFontGlyphs::Chevron_Left( FText::FromString( FString( TEXT("\xf053") ) ) );
+FText FEditorFontGlyphs::Chevron_Right( FText::FromString( FString( TEXT("\xf054") ) ) );
+FText FEditorFontGlyphs::Chevron_Up( FText::FromString( FString( TEXT("\xf077") ) ) );
+FText FEditorFontGlyphs::Child( FText::FromString( FString( TEXT("\xf1ae") ) ) );
+FText FEditorFontGlyphs::Circle( FText::FromString( FString( TEXT("\xf111") ) ) );
+FText FEditorFontGlyphs::Circle_O( FText::FromString( FString( TEXT("\xf10c") ) ) );
+FText FEditorFontGlyphs::Circle_O_Notch( FText::FromString( FString( TEXT("\xf1ce") ) ) );
+FText FEditorFontGlyphs::Circle_Thin( FText::FromString( FString( TEXT("\xf1db") ) ) );
+FText FEditorFontGlyphs::Clipboard( FText::FromString( FString( TEXT("\xf0ea") ) ) );
+FText FEditorFontGlyphs::Clock_O( FText::FromString( FString( TEXT("\xf017") ) ) );
+FText FEditorFontGlyphs::Cloud( FText::FromString( FString( TEXT("\xf0c2") ) ) );
+FText FEditorFontGlyphs::Cloud_Download( FText::FromString( FString( TEXT("\xf0ed") ) ) );
+FText FEditorFontGlyphs::Cloud_Upload( FText::FromString( FString( TEXT("\xf0ee") ) ) );
+FText FEditorFontGlyphs::Code( FText::FromString( FString( TEXT("\xf121") ) ) );
+FText FEditorFontGlyphs::Code_Fork( FText::FromString( FString( TEXT("\xf126") ) ) );
+FText FEditorFontGlyphs::Codepen( FText::FromString( FString( TEXT("\xf1cb") ) ) );
+FText FEditorFontGlyphs::Coffee( FText::FromString( FString( TEXT("\xf0f4") ) ) );
+FText FEditorFontGlyphs::Cog( FText::FromString( FString( TEXT("\xf013") ) ) );
+FText FEditorFontGlyphs::Cogs( FText::FromString( FString( TEXT("\xf085") ) ) );
+FText FEditorFontGlyphs::Columns( FText::FromString( FString( TEXT("\xf0db") ) ) );
+FText FEditorFontGlyphs::Comment( FText::FromString( FString( TEXT("\xf075") ) ) );
+FText FEditorFontGlyphs::Comment_O( FText::FromString( FString( TEXT("\xf0e5") ) ) );
+FText FEditorFontGlyphs::Comments( FText::FromString( FString( TEXT("\xf086") ) ) );
+FText FEditorFontGlyphs::Comments_O( FText::FromString( FString( TEXT("\xf0e6") ) ) );
+FText FEditorFontGlyphs::Compass( FText::FromString( FString( TEXT("\xf14e") ) ) );
+FText FEditorFontGlyphs::Compress( FText::FromString( FString( TEXT("\xf066") ) ) );
+FText FEditorFontGlyphs::Connectdevelop( FText::FromString( FString( TEXT("\xf20e") ) ) );
+FText FEditorFontGlyphs::Copyright( FText::FromString( FString( TEXT("\xf1f9") ) ) );
+FText FEditorFontGlyphs::Credit_Card( FText::FromString( FString( TEXT("\xf09d") ) ) );
+FText FEditorFontGlyphs::Crop( FText::FromString( FString( TEXT("\xf125") ) ) );
+FText FEditorFontGlyphs::Crosshairs( FText::FromString( FString( TEXT("\xf05b") ) ) );
+FText FEditorFontGlyphs::Css3( FText::FromString( FString( TEXT("\xf13c") ) ) );
+FText FEditorFontGlyphs::Cube( FText::FromString( FString( TEXT("\xf1b2") ) ) );
+FText FEditorFontGlyphs::Cubes( FText::FromString( FString( TEXT("\xf1b3") ) ) );
+FText FEditorFontGlyphs::Cutlery( FText::FromString( FString( TEXT("\xf0f5") ) ) );
+FText FEditorFontGlyphs::Dashcube( FText::FromString( FString( TEXT("\xf210") ) ) );
+FText FEditorFontGlyphs::Database( FText::FromString( FString( TEXT("\xf1c0") ) ) );
+FText FEditorFontGlyphs::Delicious( FText::FromString( FString( TEXT("\xf1a5") ) ) );
+FText FEditorFontGlyphs::Desktop( FText::FromString( FString( TEXT("\xf108") ) ) );
+FText FEditorFontGlyphs::Deviantart( FText::FromString( FString( TEXT("\xf1bd") ) ) );
+FText FEditorFontGlyphs::Diamond( FText::FromString( FString( TEXT("\xf219") ) ) );
+FText FEditorFontGlyphs::Digg( FText::FromString( FString( TEXT("\xf1a6") ) ) );
+FText FEditorFontGlyphs::Dot_Circle_O( FText::FromString( FString( TEXT("\xf192") ) ) );
+FText FEditorFontGlyphs::Download( FText::FromString( FString( TEXT("\xf019") ) ) );
+FText FEditorFontGlyphs::Dribbble( FText::FromString( FString( TEXT("\xf17d") ) ) );
+FText FEditorFontGlyphs::Dropbox( FText::FromString( FString( TEXT("\xf16b") ) ) );
+FText FEditorFontGlyphs::Drupal( FText::FromString( FString( TEXT("\xf1a9") ) ) );
+FText FEditorFontGlyphs::Eject( FText::FromString( FString( TEXT("\xf052") ) ) );
+FText FEditorFontGlyphs::Ellipsis_H( FText::FromString( FString( TEXT("\xf141") ) ) );
+FText FEditorFontGlyphs::Ellipsis_V( FText::FromString( FString( TEXT("\xf142") ) ) );
+FText FEditorFontGlyphs::Empire( FText::FromString( FString( TEXT("\xf1d1") ) ) );
+FText FEditorFontGlyphs::Envelope( FText::FromString( FString( TEXT("\xf0e0") ) ) );
+FText FEditorFontGlyphs::Envelope_O( FText::FromString( FString( TEXT("\xf003") ) ) );
+FText FEditorFontGlyphs::Envelope_Square( FText::FromString( FString( TEXT("\xf199") ) ) );
+FText FEditorFontGlyphs::Eraser( FText::FromString( FString( TEXT("\xf12d") ) ) );
+FText FEditorFontGlyphs::Eur( FText::FromString( FString( TEXT("\xf153") ) ) );
+FText FEditorFontGlyphs::Exchange( FText::FromString( FString( TEXT("\xf0ec") ) ) );
+FText FEditorFontGlyphs::Exclamation( FText::FromString( FString( TEXT("\xf12a") ) ) );
+FText FEditorFontGlyphs::Exclamation_Circle( FText::FromString( FString( TEXT("\xf06a") ) ) );
+FText FEditorFontGlyphs::Exclamation_Triangle( FText::FromString( FString( TEXT("\xf071") ) ) );
+FText FEditorFontGlyphs::Expand( FText::FromString( FString( TEXT("\xf065") ) ) );
+FText FEditorFontGlyphs::External_Link( FText::FromString( FString( TEXT("\xf08e") ) ) );
+FText FEditorFontGlyphs::External_Link_Square( FText::FromString( FString( TEXT("\xf14c") ) ) );
+FText FEditorFontGlyphs::Eye( FText::FromString( FString( TEXT("\xf06e") ) ) );
+FText FEditorFontGlyphs::Eyedropper( FText::FromString( FString( TEXT("\xf1fb") ) ) );
+FText FEditorFontGlyphs::Eye_Slash( FText::FromString( FString( TEXT("\xf070") ) ) );
+FText FEditorFontGlyphs::Facebook( FText::FromString( FString( TEXT("\xf09a") ) ) );
+FText FEditorFontGlyphs::Facebook_Official( FText::FromString( FString( TEXT("\xf230") ) ) );
+FText FEditorFontGlyphs::Facebook_Square( FText::FromString( FString( TEXT("\xf082") ) ) );
+FText FEditorFontGlyphs::Fast_Backward( FText::FromString( FString( TEXT("\xf049") ) ) );
+FText FEditorFontGlyphs::Fast_Forward( FText::FromString( FString( TEXT("\xf050") ) ) );
+FText FEditorFontGlyphs::Fax( FText::FromString( FString( TEXT("\xf1ac") ) ) );
+FText FEditorFontGlyphs::Female( FText::FromString( FString( TEXT("\xf182") ) ) );
+FText FEditorFontGlyphs::Fighter_Jet( FText::FromString( FString( TEXT("\xf0fb") ) ) );
+FText FEditorFontGlyphs::File( FText::FromString( FString( TEXT("\xf15b") ) ) );
+FText FEditorFontGlyphs::File_Archive_O( FText::FromString( FString( TEXT("\xf1c6") ) ) );
+FText FEditorFontGlyphs::File_Audio_O( FText::FromString( FString( TEXT("\xf1c7") ) ) );
+FText FEditorFontGlyphs::File_Code_O( FText::FromString( FString( TEXT("\xf1c9") ) ) );
+FText FEditorFontGlyphs::File_Excel_O( FText::FromString( FString( TEXT("\xf1c3") ) ) );
+FText FEditorFontGlyphs::File_Image_O( FText::FromString( FString( TEXT("\xf1c5") ) ) );
+FText FEditorFontGlyphs::File_O( FText::FromString( FString( TEXT("\xf016") ) ) );
+FText FEditorFontGlyphs::File_Pdf_O( FText::FromString( FString( TEXT("\xf1c1") ) ) );
+FText FEditorFontGlyphs::File_Powerpoint_O( FText::FromString( FString( TEXT("\xf1c4") ) ) );
+FText FEditorFontGlyphs::Files_O( FText::FromString( FString( TEXT("\xf0c5") ) ) );
+FText FEditorFontGlyphs::File_Text( FText::FromString( FString( TEXT("\xf15c") ) ) );
+FText FEditorFontGlyphs::File_Text_O( FText::FromString( FString( TEXT("\xf0f6") ) ) );
+FText FEditorFontGlyphs::File_Video_O( FText::FromString( FString( TEXT("\xf1c8") ) ) );
+FText FEditorFontGlyphs::File_Word_O( FText::FromString( FString( TEXT("\xf1c2") ) ) );
+FText FEditorFontGlyphs::Film( FText::FromString( FString( TEXT("\xf008") ) ) );
+FText FEditorFontGlyphs::Filter( FText::FromString( FString( TEXT("\xf0b0") ) ) );
+FText FEditorFontGlyphs::Fire( FText::FromString( FString( TEXT("\xf06d") ) ) );
+FText FEditorFontGlyphs::Fire_Extinguisher( FText::FromString( FString( TEXT("\xf134") ) ) );
+FText FEditorFontGlyphs::Flag( FText::FromString( FString( TEXT("\xf024") ) ) );
+FText FEditorFontGlyphs::Flag_Checkered( FText::FromString( FString( TEXT("\xf11e") ) ) );
+FText FEditorFontGlyphs::Flag_O( FText::FromString( FString( TEXT("\xf11d") ) ) );
+FText FEditorFontGlyphs::Flask( FText::FromString( FString( TEXT("\xf0c3") ) ) );
+FText FEditorFontGlyphs::Flickr( FText::FromString( FString( TEXT("\xf16e") ) ) );
+FText FEditorFontGlyphs::Floppy_O( FText::FromString( FString( TEXT("\xf0c7") ) ) );
+FText FEditorFontGlyphs::Folder( FText::FromString( FString( TEXT("\xf07b") ) ) );
+FText FEditorFontGlyphs::Folder_O( FText::FromString( FString( TEXT("\xf114") ) ) );
+FText FEditorFontGlyphs::Folder_Open( FText::FromString( FString( TEXT("\xf07c") ) ) );
+FText FEditorFontGlyphs::Folder_Open_O( FText::FromString( FString( TEXT("\xf115") ) ) );
+FText FEditorFontGlyphs::Font( FText::FromString( FString( TEXT("\xf031") ) ) );
+FText FEditorFontGlyphs::Forumbee( FText::FromString( FString( TEXT("\xf211") ) ) );
+FText FEditorFontGlyphs::Forward( FText::FromString( FString( TEXT("\xf04e") ) ) );
+FText FEditorFontGlyphs::Foursquare( FText::FromString( FString( TEXT("\xf180") ) ) );
+FText FEditorFontGlyphs::Frown_O( FText::FromString( FString( TEXT("\xf119") ) ) );
+FText FEditorFontGlyphs::Futbol_O( FText::FromString( FString( TEXT("\xf1e3") ) ) );
+FText FEditorFontGlyphs::Gamepad( FText::FromString( FString( TEXT("\xf11b") ) ) );
+FText FEditorFontGlyphs::Gavel( FText::FromString( FString( TEXT("\xf0e3") ) ) );
+FText FEditorFontGlyphs::Gbp( FText::FromString( FString( TEXT("\xf154") ) ) );
+FText FEditorFontGlyphs::Gift( FText::FromString( FString( TEXT("\xf06b") ) ) );
+FText FEditorFontGlyphs::Git( FText::FromString( FString( TEXT("\xf1d3") ) ) );
+FText FEditorFontGlyphs::Github( FText::FromString( FString( TEXT("\xf09b") ) ) );
+FText FEditorFontGlyphs::Github_Alt( FText::FromString( FString( TEXT("\xf113") ) ) );
+FText FEditorFontGlyphs::Github_Square( FText::FromString( FString( TEXT("\xf092") ) ) );
+FText FEditorFontGlyphs::Git_Square( FText::FromString( FString( TEXT("\xf1d2") ) ) );
+FText FEditorFontGlyphs::Glass( FText::FromString( FString( TEXT("\xf000") ) ) );
+FText FEditorFontGlyphs::Globe( FText::FromString( FString( TEXT("\xf0ac") ) ) );
+FText FEditorFontGlyphs::Google( FText::FromString( FString( TEXT("\xf1a0") ) ) );
+FText FEditorFontGlyphs::Google_Plus( FText::FromString( FString( TEXT("\xf0d5") ) ) );
+FText FEditorFontGlyphs::Google_Plus_Square( FText::FromString( FString( TEXT("\xf0d4") ) ) );
+FText FEditorFontGlyphs::Google_Wallet( FText::FromString( FString( TEXT("\xf1ee") ) ) );
+FText FEditorFontGlyphs::Graduation_Cap( FText::FromString( FString( TEXT("\xf19d") ) ) );
+FText FEditorFontGlyphs::Gratipay( FText::FromString( FString( TEXT("\xf184") ) ) );
+FText FEditorFontGlyphs::Hacker_News( FText::FromString( FString( TEXT("\xf1d4") ) ) );
+FText FEditorFontGlyphs::Hand_O_Down( FText::FromString( FString( TEXT("\xf0a7") ) ) );
+FText FEditorFontGlyphs::Hand_O_Left( FText::FromString( FString( TEXT("\xf0a5") ) ) );
+FText FEditorFontGlyphs::Hand_O_Right( FText::FromString( FString( TEXT("\xf0a4") ) ) );
+FText FEditorFontGlyphs::Hand_O_Up( FText::FromString( FString( TEXT("\xf0a6") ) ) );
+FText FEditorFontGlyphs::Hdd_O( FText::FromString( FString( TEXT("\xf0a0") ) ) );
+FText FEditorFontGlyphs::Header( FText::FromString( FString( TEXT("\xf1dc") ) ) );
+FText FEditorFontGlyphs::Headphones( FText::FromString( FString( TEXT("\xf025") ) ) );
+FText FEditorFontGlyphs::Heart( FText::FromString( FString( TEXT("\xf004") ) ) );
+FText FEditorFontGlyphs::Heartbeat( FText::FromString( FString( TEXT("\xf21e") ) ) );
+FText FEditorFontGlyphs::Heart_O( FText::FromString( FString( TEXT("\xf08a") ) ) );
+FText FEditorFontGlyphs::History( FText::FromString( FString( TEXT("\xf1da") ) ) );
+FText FEditorFontGlyphs::Home( FText::FromString( FString( TEXT("\xf015") ) ) );
+FText FEditorFontGlyphs::Hospital_O( FText::FromString( FString( TEXT("\xf0f8") ) ) );
+FText FEditorFontGlyphs::H_Square( FText::FromString( FString( TEXT("\xf0fd") ) ) );
+FText FEditorFontGlyphs::Html5( FText::FromString( FString( TEXT("\xf13b") ) ) );
+FText FEditorFontGlyphs::Ils( FText::FromString( FString( TEXT("\xf20b") ) ) );
+FText FEditorFontGlyphs::Inbox( FText::FromString( FString( TEXT("\xf01c") ) ) );
+FText FEditorFontGlyphs::Indent( FText::FromString( FString( TEXT("\xf03c") ) ) );
+FText FEditorFontGlyphs::Info( FText::FromString( FString( TEXT("\xf129") ) ) );
+FText FEditorFontGlyphs::Info_Circle( FText::FromString( FString( TEXT("\xf05a") ) ) );
+FText FEditorFontGlyphs::Inr( FText::FromString( FString( TEXT("\xf156") ) ) );
+FText FEditorFontGlyphs::Instagram( FText::FromString( FString( TEXT("\xf16d") ) ) );
+FText FEditorFontGlyphs::Ioxhost( FText::FromString( FString( TEXT("\xf208") ) ) );
+FText FEditorFontGlyphs::Italic( FText::FromString( FString( TEXT("\xf033") ) ) );
+FText FEditorFontGlyphs::Joomla( FText::FromString( FString( TEXT("\xf1aa") ) ) );
+FText FEditorFontGlyphs::Jpy( FText::FromString( FString( TEXT("\xf157") ) ) );
+FText FEditorFontGlyphs::Jsfiddle( FText::FromString( FString( TEXT("\xf1cc") ) ) );
+FText FEditorFontGlyphs::Key( FText::FromString( FString( TEXT("\xf084") ) ) );
+FText FEditorFontGlyphs::Keyboard_O( FText::FromString( FString( TEXT("\xf11c") ) ) );
+FText FEditorFontGlyphs::Krw( FText::FromString( FString( TEXT("\xf159") ) ) );
+FText FEditorFontGlyphs::Language( FText::FromString( FString( TEXT("\xf1ab") ) ) );
+FText FEditorFontGlyphs::Laptop( FText::FromString( FString( TEXT("\xf109") ) ) );
+FText FEditorFontGlyphs::Lastfm( FText::FromString( FString( TEXT("\xf202") ) ) );
+FText FEditorFontGlyphs::Lastfm_Square( FText::FromString( FString( TEXT("\xf203") ) ) );
+FText FEditorFontGlyphs::Leaf( FText::FromString( FString( TEXT("\xf06c") ) ) );
+FText FEditorFontGlyphs::Leanpub( FText::FromString( FString( TEXT("\xf212") ) ) );
+FText FEditorFontGlyphs::Lemon_O( FText::FromString( FString( TEXT("\xf094") ) ) );
+FText FEditorFontGlyphs::Level_Down( FText::FromString( FString( TEXT("\xf149") ) ) );
+FText FEditorFontGlyphs::Level_Up( FText::FromString( FString( TEXT("\xf148") ) ) );
+FText FEditorFontGlyphs::Life_Ring( FText::FromString( FString( TEXT("\xf1cd") ) ) );
+FText FEditorFontGlyphs::Lightbulb_O( FText::FromString( FString( TEXT("\xf0eb") ) ) );
+FText FEditorFontGlyphs::Line_Chart( FText::FromString( FString( TEXT("\xf201") ) ) );
+FText FEditorFontGlyphs::Link( FText::FromString( FString( TEXT("\xf0c1") ) ) );
+FText FEditorFontGlyphs::Linkedin( FText::FromString( FString( TEXT("\xf0e1") ) ) );
+FText FEditorFontGlyphs::Linkedin_Square( FText::FromString( FString( TEXT("\xf08c") ) ) );
+FText FEditorFontGlyphs::Linux( FText::FromString( FString( TEXT("\xf17c") ) ) );
+FText FEditorFontGlyphs::List( FText::FromString( FString( TEXT("\xf03a") ) ) );
+FText FEditorFontGlyphs::List_Alt( FText::FromString( FString( TEXT("\xf022") ) ) );
+FText FEditorFontGlyphs::List_Ol( FText::FromString( FString( TEXT("\xf0cb") ) ) );
+FText FEditorFontGlyphs::List_Ul( FText::FromString( FString( TEXT("\xf0ca") ) ) );
+FText FEditorFontGlyphs::Location_Arrow( FText::FromString( FString( TEXT("\xf124") ) ) );
+FText FEditorFontGlyphs::Lock( FText::FromString( FString( TEXT("\xf023") ) ) );
+FText FEditorFontGlyphs::Long_Arrow_Down( FText::FromString( FString( TEXT("\xf175") ) ) );
+FText FEditorFontGlyphs::Long_Arrow_Left( FText::FromString( FString( TEXT("\xf177") ) ) );
+FText FEditorFontGlyphs::Long_Arrow_Right( FText::FromString( FString( TEXT("\xf178") ) ) );
+FText FEditorFontGlyphs::Long_Arrow_Up( FText::FromString( FString( TEXT("\xf176") ) ) );
+FText FEditorFontGlyphs::Magic( FText::FromString( FString( TEXT("\xf0d0") ) ) );
+FText FEditorFontGlyphs::Magnet( FText::FromString( FString( TEXT("\xf076") ) ) );
+FText FEditorFontGlyphs::Male( FText::FromString( FString( TEXT("\xf183") ) ) );
+FText FEditorFontGlyphs::Map_Marker( FText::FromString( FString( TEXT("\xf041") ) ) );
+FText FEditorFontGlyphs::Mars( FText::FromString( FString( TEXT("\xf222") ) ) );
+FText FEditorFontGlyphs::Mars_Double( FText::FromString( FString( TEXT("\xf227") ) ) );
+FText FEditorFontGlyphs::Mars_Stroke( FText::FromString( FString( TEXT("\xf229") ) ) );
+FText FEditorFontGlyphs::Mars_Stroke_H( FText::FromString( FString( TEXT("\xf22b") ) ) );
+FText FEditorFontGlyphs::Mars_Stroke_V( FText::FromString( FString( TEXT("\xf22a") ) ) );
+FText FEditorFontGlyphs::Maxcdn( FText::FromString( FString( TEXT("\xf136") ) ) );
+FText FEditorFontGlyphs::Meanpath( FText::FromString( FString( TEXT("\xf20c") ) ) );
+FText FEditorFontGlyphs::Medium( FText::FromString( FString( TEXT("\xf23a") ) ) );
+FText FEditorFontGlyphs::Medkit( FText::FromString( FString( TEXT("\xf0fa") ) ) );
+FText FEditorFontGlyphs::Meh_O( FText::FromString( FString( TEXT("\xf11a") ) ) );
+FText FEditorFontGlyphs::Mercury( FText::FromString( FString( TEXT("\xf223") ) ) );
+FText FEditorFontGlyphs::Microphone( FText::FromString( FString( TEXT("\xf130") ) ) );
+FText FEditorFontGlyphs::Microphone_Slash( FText::FromString( FString( TEXT("\xf131") ) ) );
+FText FEditorFontGlyphs::Minus( FText::FromString( FString( TEXT("\xf068") ) ) );
+FText FEditorFontGlyphs::Minus_Circle( FText::FromString( FString( TEXT("\xf056") ) ) );
+FText FEditorFontGlyphs::Minus_Square( FText::FromString( FString( TEXT("\xf146") ) ) );
+FText FEditorFontGlyphs::Minus_Square_O( FText::FromString( FString( TEXT("\xf147") ) ) );
+FText FEditorFontGlyphs::Mobile( FText::FromString( FString( TEXT("\xf10b") ) ) );
+FText FEditorFontGlyphs::Money( FText::FromString( FString( TEXT("\xf0d6") ) ) );
+FText FEditorFontGlyphs::Moon_O( FText::FromString( FString( TEXT("\xf186") ) ) );
+FText FEditorFontGlyphs::Motorcycle( FText::FromString( FString( TEXT("\xf21c") ) ) );
+FText FEditorFontGlyphs::Music( FText::FromString( FString( TEXT("\xf001") ) ) );
+FText FEditorFontGlyphs::Neuter( FText::FromString( FString( TEXT("\xf22c") ) ) );
+FText FEditorFontGlyphs::Newspaper_O( FText::FromString( FString( TEXT("\xf1ea") ) ) );
+FText FEditorFontGlyphs::Openid( FText::FromString( FString( TEXT("\xf19b") ) ) );
+FText FEditorFontGlyphs::Outdent( FText::FromString( FString( TEXT("\xf03b") ) ) );
+FText FEditorFontGlyphs::Pagelines( FText::FromString( FString( TEXT("\xf18c") ) ) );
+FText FEditorFontGlyphs::Paint_Brush( FText::FromString( FString( TEXT("\xf1fc") ) ) );
+FText FEditorFontGlyphs::Paperclip( FText::FromString( FString( TEXT("\xf0c6") ) ) );
+FText FEditorFontGlyphs::Paper_Plane( FText::FromString( FString( TEXT("\xf1d8") ) ) );
+FText FEditorFontGlyphs::Paper_Plane_O( FText::FromString( FString( TEXT("\xf1d9") ) ) );
+FText FEditorFontGlyphs::Paragraph( FText::FromString( FString( TEXT("\xf1dd") ) ) );
+FText FEditorFontGlyphs::Pause( FText::FromString( FString( TEXT("\xf04c") ) ) );
+FText FEditorFontGlyphs::Paw( FText::FromString( FString( TEXT("\xf1b0") ) ) );
+FText FEditorFontGlyphs::Paypal( FText::FromString( FString( TEXT("\xf1ed") ) ) );
+FText FEditorFontGlyphs::Pencil( FText::FromString( FString( TEXT("\xf040") ) ) );
+FText FEditorFontGlyphs::Pencil_Square( FText::FromString( FString( TEXT("\xf14b") ) ) );
+FText FEditorFontGlyphs::Pencil_Square_O( FText::FromString( FString( TEXT("\xf044") ) ) );
+FText FEditorFontGlyphs::Phone( FText::FromString( FString( TEXT("\xf095") ) ) );
+FText FEditorFontGlyphs::Phone_Square( FText::FromString( FString( TEXT("\xf098") ) ) );
+FText FEditorFontGlyphs::Picture_O( FText::FromString( FString( TEXT("\xf03e") ) ) );
+FText FEditorFontGlyphs::Pie_Chart( FText::FromString( FString( TEXT("\xf200") ) ) );
+FText FEditorFontGlyphs::Pied_Piper( FText::FromString( FString( TEXT("\xf1a7") ) ) );
+FText FEditorFontGlyphs::Pied_Piper_Alt( FText::FromString( FString( TEXT("\xf1a8") ) ) );
+FText FEditorFontGlyphs::Pinterest( FText::FromString( FString( TEXT("\xf0d2") ) ) );
+FText FEditorFontGlyphs::Pinterest_P( FText::FromString( FString( TEXT("\xf231") ) ) );
+FText FEditorFontGlyphs::Pinterest_Square( FText::FromString( FString( TEXT("\xf0d3") ) ) );
+FText FEditorFontGlyphs::Plane( FText::FromString( FString( TEXT("\xf072") ) ) );
+FText FEditorFontGlyphs::Play( FText::FromString( FString( TEXT("\xf04b") ) ) );
+FText FEditorFontGlyphs::Play_Circle( FText::FromString( FString( TEXT("\xf144") ) ) );
+FText FEditorFontGlyphs::Play_Circle_O( FText::FromString( FString( TEXT("\xf01d") ) ) );
+FText FEditorFontGlyphs::Plug( FText::FromString( FString( TEXT("\xf1e6") ) ) );
+FText FEditorFontGlyphs::Plus( FText::FromString( FString( TEXT("\xf067") ) ) );
+FText FEditorFontGlyphs::Plus_Circle( FText::FromString( FString( TEXT("\xf055") ) ) );
+FText FEditorFontGlyphs::Plus_Square( FText::FromString( FString( TEXT("\xf0fe") ) ) );
+FText FEditorFontGlyphs::Plus_Square_O( FText::FromString( FString( TEXT("\xf196") ) ) );
+FText FEditorFontGlyphs::Power_Off( FText::FromString( FString( TEXT("\xf011") ) ) );
+FText FEditorFontGlyphs::Print( FText::FromString( FString( TEXT("\xf02f") ) ) );
+FText FEditorFontGlyphs::Puzzle_Piece( FText::FromString( FString( TEXT("\xf12e") ) ) );
+FText FEditorFontGlyphs::Qq( FText::FromString( FString( TEXT("\xf1d6") ) ) );
+FText FEditorFontGlyphs::Qrcode( FText::FromString( FString( TEXT("\xf029") ) ) );
+FText FEditorFontGlyphs::Question( FText::FromString( FString( TEXT("\xf128") ) ) );
+FText FEditorFontGlyphs::Question_Circle( FText::FromString( FString( TEXT("\xf059") ) ) );
+FText FEditorFontGlyphs::Quote_Left( FText::FromString( FString( TEXT("\xf10d") ) ) );
+FText FEditorFontGlyphs::Quote_Right( FText::FromString( FString( TEXT("\xf10e") ) ) );
+FText FEditorFontGlyphs::Random( FText::FromString( FString( TEXT("\xf074") ) ) );
+FText FEditorFontGlyphs::Rebel( FText::FromString( FString( TEXT("\xf1d0") ) ) );
+FText FEditorFontGlyphs::Recycle( FText::FromString( FString( TEXT("\xf1b8") ) ) );
+FText FEditorFontGlyphs::Reddit( FText::FromString( FString( TEXT("\xf1a1") ) ) );
+FText FEditorFontGlyphs::Reddit_Square( FText::FromString( FString( TEXT("\xf1a2") ) ) );
+FText FEditorFontGlyphs::Refresh( FText::FromString( FString( TEXT("\xf021") ) ) );
+FText FEditorFontGlyphs::Renren( FText::FromString( FString( TEXT("\xf18b") ) ) );
+FText FEditorFontGlyphs::Repeat( FText::FromString( FString( TEXT("\xf01e") ) ) );
+FText FEditorFontGlyphs::Reply( FText::FromString( FString( TEXT("\xf112") ) ) );
+FText FEditorFontGlyphs::Reply_All( FText::FromString( FString( TEXT("\xf122") ) ) );
+FText FEditorFontGlyphs::Retweet( FText::FromString( FString( TEXT("\xf079") ) ) );
+FText FEditorFontGlyphs::Road( FText::FromString( FString( TEXT("\xf018") ) ) );
+FText FEditorFontGlyphs::Rocket( FText::FromString( FString( TEXT("\xf135") ) ) );
+FText FEditorFontGlyphs::Rss( FText::FromString( FString( TEXT("\xf09e") ) ) );
+FText FEditorFontGlyphs::Rss_Square( FText::FromString( FString( TEXT("\xf143") ) ) );
+FText FEditorFontGlyphs::Rub( FText::FromString( FString( TEXT("\xf158") ) ) );
+FText FEditorFontGlyphs::Scissors( FText::FromString( FString( TEXT("\xf0c4") ) ) );
+FText FEditorFontGlyphs::Search( FText::FromString( FString( TEXT("\xf002") ) ) );
+FText FEditorFontGlyphs::Search_Minus( FText::FromString( FString( TEXT("\xf010") ) ) );
+FText FEditorFontGlyphs::Search_Plus( FText::FromString( FString( TEXT("\xf00e") ) ) );
+FText FEditorFontGlyphs::Sellsy( FText::FromString( FString( TEXT("\xf213") ) ) );
+FText FEditorFontGlyphs::Server( FText::FromString( FString( TEXT("\xf233") ) ) );
+FText FEditorFontGlyphs::Share( FText::FromString( FString( TEXT("\xf064") ) ) );
+FText FEditorFontGlyphs::Share_Alt( FText::FromString( FString( TEXT("\xf1e0") ) ) );
+FText FEditorFontGlyphs::Share_Alt_Square( FText::FromString( FString( TEXT("\xf1e1") ) ) );
+FText FEditorFontGlyphs::Share_Square( FText::FromString( FString( TEXT("\xf14d") ) ) );
+FText FEditorFontGlyphs::Share_Square_O( FText::FromString( FString( TEXT("\xf045") ) ) );
+FText FEditorFontGlyphs::Shield( FText::FromString( FString( TEXT("\xf132") ) ) );
+FText FEditorFontGlyphs::Ship( FText::FromString( FString( TEXT("\xf21a") ) ) );
+FText FEditorFontGlyphs::Shirtsinbulk( FText::FromString( FString( TEXT("\xf214") ) ) );
+FText FEditorFontGlyphs::Shopping_Cart( FText::FromString( FString( TEXT("\xf07a") ) ) );
+FText FEditorFontGlyphs::Signal( FText::FromString( FString( TEXT("\xf012") ) ) );
+FText FEditorFontGlyphs::Sign_In( FText::FromString( FString( TEXT("\xf090") ) ) );
+FText FEditorFontGlyphs::Sign_Out( FText::FromString( FString( TEXT("\xf08b") ) ) );
+FText FEditorFontGlyphs::Simplybuilt( FText::FromString( FString( TEXT("\xf215") ) ) );
+FText FEditorFontGlyphs::Sitemap( FText::FromString( FString( TEXT("\xf0e8") ) ) );
+FText FEditorFontGlyphs::Skyatlas( FText::FromString( FString( TEXT("\xf216") ) ) );
+FText FEditorFontGlyphs::Skype( FText::FromString( FString( TEXT("\xf17e") ) ) );
+FText FEditorFontGlyphs::Slack( FText::FromString( FString( TEXT("\xf198") ) ) );
+FText FEditorFontGlyphs::Sliders( FText::FromString( FString( TEXT("\xf1de") ) ) );
+FText FEditorFontGlyphs::Slideshare( FText::FromString( FString( TEXT("\xf1e7") ) ) );
+FText FEditorFontGlyphs::Smile_O( FText::FromString( FString( TEXT("\xf118") ) ) );
+FText FEditorFontGlyphs::Sort( FText::FromString( FString( TEXT("\xf0dc") ) ) );
+FText FEditorFontGlyphs::Sort_Alpha_Asc( FText::FromString( FString( TEXT("\xf15d") ) ) );
+FText FEditorFontGlyphs::Sort_Alpha_Desc( FText::FromString( FString( TEXT("\xf15e") ) ) );
+FText FEditorFontGlyphs::Sort_Amount_Asc( FText::FromString( FString( TEXT("\xf160") ) ) );
+FText FEditorFontGlyphs::Sort_Amount_Desc( FText::FromString( FString( TEXT("\xf161") ) ) );
+FText FEditorFontGlyphs::Sort_Asc( FText::FromString( FString( TEXT("\xf0de") ) ) );
+FText FEditorFontGlyphs::Sort_Desc( FText::FromString( FString( TEXT("\xf0dd") ) ) );
+FText FEditorFontGlyphs::Sort_Numeric_Asc( FText::FromString( FString( TEXT("\xf162") ) ) );
+FText FEditorFontGlyphs::Sort_Numeric_Desc( FText::FromString( FString( TEXT("\xf163") ) ) );
+FText FEditorFontGlyphs::Soundcloud( FText::FromString( FString( TEXT("\xf1be") ) ) );
+FText FEditorFontGlyphs::Space_Shuttle( FText::FromString( FString( TEXT("\xf197") ) ) );
+FText FEditorFontGlyphs::Spinner( FText::FromString( FString( TEXT("\xf110") ) ) );
+FText FEditorFontGlyphs::Spoon( FText::FromString( FString( TEXT("\xf1b1") ) ) );
+FText FEditorFontGlyphs::Spotify( FText::FromString( FString( TEXT("\xf1bc") ) ) );
+FText FEditorFontGlyphs::Square( FText::FromString( FString( TEXT("\xf0c8") ) ) );
+FText FEditorFontGlyphs::Square_O( FText::FromString( FString( TEXT("\xf096") ) ) );
+FText FEditorFontGlyphs::Stack_Exchange( FText::FromString( FString( TEXT("\xf18d") ) ) );
+FText FEditorFontGlyphs::Stack_Overflow( FText::FromString( FString( TEXT("\xf16c") ) ) );
+FText FEditorFontGlyphs::Star( FText::FromString( FString( TEXT("\xf005") ) ) );
+FText FEditorFontGlyphs::Star_Half( FText::FromString( FString( TEXT("\xf089") ) ) );
+FText FEditorFontGlyphs::Star_Half_O( FText::FromString( FString( TEXT("\xf123") ) ) );
+FText FEditorFontGlyphs::Star_O( FText::FromString( FString( TEXT("\xf006") ) ) );
+FText FEditorFontGlyphs::Steam( FText::FromString( FString( TEXT("\xf1b6") ) ) );
+FText FEditorFontGlyphs::Steam_Square( FText::FromString( FString( TEXT("\xf1b7") ) ) );
+FText FEditorFontGlyphs::Step_Backward( FText::FromString( FString( TEXT("\xf048") ) ) );
+FText FEditorFontGlyphs::Step_Forward( FText::FromString( FString( TEXT("\xf051") ) ) );
+FText FEditorFontGlyphs::Stethoscope( FText::FromString( FString( TEXT("\xf0f1") ) ) );
+FText FEditorFontGlyphs::Stop( FText::FromString( FString( TEXT("\xf04d") ) ) );
+FText FEditorFontGlyphs::Street_View( FText::FromString( FString( TEXT("\xf21d") ) ) );
+FText FEditorFontGlyphs::Strikethrough( FText::FromString( FString( TEXT("\xf0cc") ) ) );
+FText FEditorFontGlyphs::Stumbleupon( FText::FromString( FString( TEXT("\xf1a4") ) ) );
+FText FEditorFontGlyphs::Stumbleupon_Circle( FText::FromString( FString( TEXT("\xf1a3") ) ) );
+FText FEditorFontGlyphs::Subscript( FText::FromString( FString( TEXT("\xf12c") ) ) );
+FText FEditorFontGlyphs::Subway( FText::FromString( FString( TEXT("\xf239") ) ) );
+FText FEditorFontGlyphs::Suitcase( FText::FromString( FString( TEXT("\xf0f2") ) ) );
+FText FEditorFontGlyphs::Sun_O( FText::FromString( FString( TEXT("\xf185") ) ) );
+FText FEditorFontGlyphs::Superscript( FText::FromString( FString( TEXT("\xf12b") ) ) );
+FText FEditorFontGlyphs::Table( FText::FromString( FString( TEXT("\xf0ce") ) ) );
+FText FEditorFontGlyphs::Tablet( FText::FromString( FString( TEXT("\xf10a") ) ) );
+FText FEditorFontGlyphs::Tachometer( FText::FromString( FString( TEXT("\xf0e4") ) ) );
+FText FEditorFontGlyphs::Tag( FText::FromString( FString( TEXT("\xf02b") ) ) );
+FText FEditorFontGlyphs::Tags( FText::FromString( FString( TEXT("\xf02c") ) ) );
+FText FEditorFontGlyphs::Tasks( FText::FromString( FString( TEXT("\xf0ae") ) ) );
+FText FEditorFontGlyphs::Taxi( FText::FromString( FString( TEXT("\xf1ba") ) ) );
+FText FEditorFontGlyphs::Tencent_Weibo( FText::FromString( FString( TEXT("\xf1d5") ) ) );
+FText FEditorFontGlyphs::Terminal( FText::FromString( FString( TEXT("\xf120") ) ) );
+FText FEditorFontGlyphs::Text_Height( FText::FromString( FString( TEXT("\xf034") ) ) );
+FText FEditorFontGlyphs::Text_Width( FText::FromString( FString( TEXT("\xf035") ) ) );
+FText FEditorFontGlyphs::Th( FText::FromString( FString( TEXT("\xf00a") ) ) );
+FText FEditorFontGlyphs::Th_Large( FText::FromString( FString( TEXT("\xf009") ) ) );
+FText FEditorFontGlyphs::Th_List( FText::FromString( FString( TEXT("\xf00b") ) ) );
+FText FEditorFontGlyphs::Thumbs_Down( FText::FromString( FString( TEXT("\xf165") ) ) );
+FText FEditorFontGlyphs::Thumbs_O_Down( FText::FromString( FString( TEXT("\xf088") ) ) );
+FText FEditorFontGlyphs::Thumbs_O_Up( FText::FromString( FString( TEXT("\xf087") ) ) );
+FText FEditorFontGlyphs::Thumbs_Up( FText::FromString( FString( TEXT("\xf164") ) ) );
+FText FEditorFontGlyphs::Thumb_Tack( FText::FromString( FString( TEXT("\xf08d") ) ) );
+FText FEditorFontGlyphs::Ticket( FText::FromString( FString( TEXT("\xf145") ) ) );
+FText FEditorFontGlyphs::Times( FText::FromString( FString( TEXT("\xf00d") ) ) );
+FText FEditorFontGlyphs::Times_Circle( FText::FromString( FString( TEXT("\xf057") ) ) );
+FText FEditorFontGlyphs::Times_Circle_O( FText::FromString( FString( TEXT("\xf05c") ) ) );
+FText FEditorFontGlyphs::Tint( FText::FromString( FString( TEXT("\xf043") ) ) );
+FText FEditorFontGlyphs::Toggle_Off( FText::FromString( FString( TEXT("\xf204") ) ) );
+FText FEditorFontGlyphs::Toggle_On( FText::FromString( FString( TEXT("\xf205") ) ) );
+FText FEditorFontGlyphs::Train( FText::FromString( FString( TEXT("\xf238") ) ) );
+FText FEditorFontGlyphs::Transgender( FText::FromString( FString( TEXT("\xf224") ) ) );
+FText FEditorFontGlyphs::Transgender_Alt( FText::FromString( FString( TEXT("\xf225") ) ) );
+FText FEditorFontGlyphs::Trash( FText::FromString( FString( TEXT("\xf1f8") ) ) );
+FText FEditorFontGlyphs::Trash_O( FText::FromString( FString( TEXT("\xf014") ) ) );
+FText FEditorFontGlyphs::Tree( FText::FromString( FString( TEXT("\xf1bb") ) ) );
+FText FEditorFontGlyphs::Trello( FText::FromString( FString( TEXT("\xf181") ) ) );
+FText FEditorFontGlyphs::Trophy( FText::FromString( FString( TEXT("\xf091") ) ) );
+FText FEditorFontGlyphs::Truck( FText::FromString( FString( TEXT("\xf0d1") ) ) );
+FText FEditorFontGlyphs::Try( FText::FromString( FString( TEXT("\xf195") ) ) );
+FText FEditorFontGlyphs::Tty( FText::FromString( FString( TEXT("\xf1e4") ) ) );
+FText FEditorFontGlyphs::Tumblr( FText::FromString( FString( TEXT("\xf173") ) ) );
+FText FEditorFontGlyphs::Tumblr_Square( FText::FromString( FString( TEXT("\xf174") ) ) );
+FText FEditorFontGlyphs::Twitch( FText::FromString( FString( TEXT("\xf1e8") ) ) );
+FText FEditorFontGlyphs::Twitter( FText::FromString( FString( TEXT("\xf099") ) ) );
+FText FEditorFontGlyphs::Twitter_Square( FText::FromString( FString( TEXT("\xf081") ) ) );
+FText FEditorFontGlyphs::Umbrella( FText::FromString( FString( TEXT("\xf0e9") ) ) );
+FText FEditorFontGlyphs::Underline( FText::FromString( FString( TEXT("\xf0cd") ) ) );
+FText FEditorFontGlyphs::Undo( FText::FromString( FString( TEXT("\xf0e2") ) ) );
+FText FEditorFontGlyphs::University( FText::FromString( FString( TEXT("\xf19c") ) ) );
+FText FEditorFontGlyphs::Unlock( FText::FromString( FString( TEXT("\xf09c") ) ) );
+FText FEditorFontGlyphs::Unlock_Alt( FText::FromString( FString( TEXT("\xf13e") ) ) );
+FText FEditorFontGlyphs::Upload( FText::FromString( FString( TEXT("\xf093") ) ) );
+FText FEditorFontGlyphs::Usd( FText::FromString( FString( TEXT("\xf155") ) ) );
+FText FEditorFontGlyphs::User( FText::FromString( FString( TEXT("\xf007") ) ) );
+FText FEditorFontGlyphs::User_Md( FText::FromString( FString( TEXT("\xf0f0") ) ) );
+FText FEditorFontGlyphs::User_Plus( FText::FromString( FString( TEXT("\xf234") ) ) );
+FText FEditorFontGlyphs::Users( FText::FromString( FString( TEXT("\xf0c0") ) ) );
+FText FEditorFontGlyphs::User_Secret( FText::FromString( FString( TEXT("\xf21b") ) ) );
+FText FEditorFontGlyphs::User_Times( FText::FromString( FString( TEXT("\xf235") ) ) );
+FText FEditorFontGlyphs::Venus( FText::FromString( FString( TEXT("\xf221") ) ) );
+FText FEditorFontGlyphs::Venus_Double( FText::FromString( FString( TEXT("\xf226") ) ) );
+FText FEditorFontGlyphs::Venus_Mars( FText::FromString( FString( TEXT("\xf228") ) ) );
+FText FEditorFontGlyphs::Viacoin( FText::FromString( FString( TEXT("\xf237") ) ) );
+FText FEditorFontGlyphs::Video_Camera( FText::FromString( FString( TEXT("\xf03d") ) ) );
+FText FEditorFontGlyphs::Vimeo_Square( FText::FromString( FString( TEXT("\xf194") ) ) );
+FText FEditorFontGlyphs::Vine( FText::FromString( FString( TEXT("\xf1ca") ) ) );
+FText FEditorFontGlyphs::Vk( FText::FromString( FString( TEXT("\xf189") ) ) );
+FText FEditorFontGlyphs::Volume_Down( FText::FromString( FString( TEXT("\xf027") ) ) );
+FText FEditorFontGlyphs::Volume_Off( FText::FromString( FString( TEXT("\xf026") ) ) );
+FText FEditorFontGlyphs::Volume_Up( FText::FromString( FString( TEXT("\xf028") ) ) );
+FText FEditorFontGlyphs::Weibo( FText::FromString( FString( TEXT("\xf18a") ) ) );
+FText FEditorFontGlyphs::Weixin( FText::FromString( FString( TEXT("\xf1d7") ) ) );
+FText FEditorFontGlyphs::Whatsapp( FText::FromString( FString( TEXT("\xf232") ) ) );
+FText FEditorFontGlyphs::Wheelchair( FText::FromString( FString( TEXT("\xf193") ) ) );
+FText FEditorFontGlyphs::Wifi( FText::FromString( FString( TEXT("\xf1eb") ) ) );
+FText FEditorFontGlyphs::Windows( FText::FromString( FString( TEXT("\xf17a") ) ) );
+FText FEditorFontGlyphs::Wordpress( FText::FromString( FString( TEXT("\xf19a") ) ) );
+FText FEditorFontGlyphs::Wrench( FText::FromString( FString( TEXT("\xf0ad") ) ) );
+FText FEditorFontGlyphs::Xing( FText::FromString( FString( TEXT("\xf168") ) ) );
+FText FEditorFontGlyphs::Xing_Square( FText::FromString( FString( TEXT("\xf169") ) ) );
+FText FEditorFontGlyphs::Yahoo( FText::FromString( FString( TEXT("\xf19e") ) ) );
+FText FEditorFontGlyphs::Yelp( FText::FromString( FString( TEXT("\xf1e9") ) ) );
+FText FEditorFontGlyphs::Youtube( FText::FromString( FString( TEXT("\xf167") ) ) );
+FText FEditorFontGlyphs::Youtube_Play( FText::FromString( FString( TEXT("\xf16a") ) ) );
+FText FEditorFontGlyphs::Youtube_Square( FText::FromString( FString( TEXT("\xf166") ) ) );
diff --git a/Engine/Source/Editor/EditorStyle/Public/EditorFontGlyphs.h b/Engine/Source/Editor/EditorStyle/Public/EditorFontGlyphs.h
new file mode 100644
--- /dev/null
+++ b/Engine/Source/Editor/EditorStyle/Public/EditorFontGlyphs.h
@@ -0,0 +1,1564 @@
+// Copyright 1998-2015 Epic Games, Inc. All Rights Reserved.
+
+#pragma once
+
+class EDITORSTYLE_API FEditorFontGlyphs
+{
+public:
+	// (contrast) http://fortawesome.github.io/Font-Awesome/icon/adjust/
+	static FText Adjust;
+
+	// http://fortawesome.github.io/Font-Awesome/icon/adn/
+	static FText Adn;
+
+	// (middle, text) http://fortawesome.github.io/Font-Awesome/icon/align-center/
+	static FText Align_Center;
+
+	// (text) http://fortawesome.github.io/Font-Awesome/icon/align-justify/
+	static FText Align_Justify;
+
+	// (text) http://fortawesome.github.io/Font-Awesome/icon/align-left/
+	static FText Align_Left;
+
+	// (text) http://fortawesome.github.io/Font-Awesome/icon/align-right/
+	static FText Align_Right;
+
+	// (support, help) http://fortawesome.github.io/Font-Awesome/icon/ambulance/
+	static FText Ambulance;
+
+	// (link) http://fortawesome.github.io/Font-Awesome/icon/anchor/
+	static FText Anchor;
+
+	// http://fortawesome.github.io/Font-Awesome/icon/android/
+	static FText Android;
+
+	// http://fortawesome.github.io/Font-Awesome/icon/angellist/
+	static FText Angellist;
+
+	// http://fortawesome.github.io/Font-Awesome/icon/angle-double-down/
+	static FText Angle_Double_Down;
+
+	// (laquo, quote, previous, back) http://fortawesome.github.io/Font-Awesome/icon/angle-double-left/
+	static FText Angle_Double_Left;
+
+	// (raquo, quote, next, forward) http://fortawesome.github.io/Font-Awesome/icon/angle-double-right/
+	static FText Angle_Double_Right;
+
+	// http://fortawesome.github.io/Font-Awesome/icon/angle-double-up/
+	static FText Angle_Double_Up;
+
+	// http://fortawesome.github.io/Font-Awesome/icon/angle-down/
+	static FText Angle_Down;
+
+	// (previous, back) http://fortawesome.github.io/Font-Awesome/icon/angle-left/
+	static FText Angle_Left;
+
+	// (next, forward) http://fortawesome.github.io/Font-Awesome/icon/angle-right/
+	static FText Angle_Right;
+
+	// http://fortawesome.github.io/Font-Awesome/icon/angle-up/
+	static FText Angle_Up;
+
+	// (osx) http://fortawesome.github.io/Font-Awesome/icon/apple/
+	static FText Apple;
+
+	// (box, storage) http://fortawesome.github.io/Font-Awesome/icon/archive/
+	static FText Archive;
+
+	// http://fortawesome.github.io/Font-Awesome/icon/area-chart/
+	static FText Area_Chart;
+
+	// (download) http://fortawesome.github.io/Font-Awesome/icon/arrow-circle-down/
+	static FText Arrow_Circle_Down;
+
+	// (previous, back) http://fortawesome.github.io/Font-Awesome/icon/arrow-circle-left/
+	static FText Arrow_Circle_Left;
+
+	// (download) http://fortawesome.github.io/Font-Awesome/icon/arrow-circle-o-down/
+	static FText Arrow_Circle_O_Down;
+
+	// (previous, back) http://fortawesome.github.io/Font-Awesome/icon/arrow-circle-o-left/
+	static FText Arrow_Circle_O_Left;
+
+	// (next, forward) http://fortawesome.github.io/Font-Awesome/icon/arrow-circle-o-right/
+	static FText Arrow_Circle_O_Right;
+
+	// http://fortawesome.github.io/Font-Awesome/icon/arrow-circle-o-up/
+	static FText Arrow_Circle_O_Up;
+
+	// (next, forward) http://fortawesome.github.io/Font-Awesome/icon/arrow-circle-right/
+	static FText Arrow_Circle_Right;
+
+	// http://fortawesome.github.io/Font-Awesome/icon/arrow-circle-up/
+	static FText Arrow_Circle_Up;
+
+	// (download) http://fortawesome.github.io/Font-Awesome/icon/arrow-down/
+	static FText Arrow_Down;
+
+	// (previous, back) http://fortawesome.github.io/Font-Awesome/icon/arrow-left/
+	static FText Arrow_Left;
+
+	// (next, forward) http://fortawesome.github.io/Font-Awesome/icon/arrow-right/
+	static FText Arrow_Right;
+
+	// (move, reorder, resize) http://fortawesome.github.io/Font-Awesome/icon/arrows/
+	static FText Arrows;
+
+	// (expand, enlarge, bigger, move, reorder, resize) http://fortawesome.github.io/Font-Awesome/icon/arrows-alt/
+	static FText Arrows_Alt;
+
+	// (resize) http://fortawesome.github.io/Font-Awesome/icon/arrows-h/
+	static FText Arrows_H;
+
+	// (resize) http://fortawesome.github.io/Font-Awesome/icon/arrows-v/
+	static FText Arrows_V;
+
+	// http://fortawesome.github.io/Font-Awesome/icon/arrow-up/
+	static FText Arrow_Up;
+
+	// (details) http://fortawesome.github.io/Font-Awesome/icon/asterisk/
+	static FText Asterisk;
+
+	// http://fortawesome.github.io/Font-Awesome/icon/at/
+	static FText At;
+
+	// (rewind, previous) http://fortawesome.github.io/Font-Awesome/icon/backward/
+	static FText Backward;
+
+	// (block, abort) http://fortawesome.github.io/Font-Awesome/icon/ban/
+	static FText Ban;
+
+	// (graph) http://fortawesome.github.io/Font-Awesome/icon/bar-chart/
+	static FText Bar_Chart;
+
+	// (scan) http://fortawesome.github.io/Font-Awesome/icon/barcode/
+	static FText Barcode;
+
+	// (menu, drag, reorder, settings, list, ul, ol, checklist, todo, list) http://fortawesome.github.io/Font-Awesome/icon/bars/
+	static FText Bars;
+
+	// (travel) http://fortawesome.github.io/Font-Awesome/icon/bed/
+	static FText Bed;
+
+	// (alcohol, stein, drink, mug, bar, liquor) http://fortawesome.github.io/Font-Awesome/icon/beer/
+	static FText Beer;
+
+	// http://fortawesome.github.io/Font-Awesome/icon/behance/
+	static FText Behance;
+
+	// http://fortawesome.github.io/Font-Awesome/icon/behance-square/
+	static FText Behance_Square;
+
+	// (alert, reminder, notification) http://fortawesome.github.io/Font-Awesome/icon/bell/
+	static FText Bell;
+
+	// (alert, reminder, notification) http://fortawesome.github.io/Font-Awesome/icon/bell-o/
+	static FText Bell_O;
+
+	// http://fortawesome.github.io/Font-Awesome/icon/bell-slash/
+	static FText Bell_Slash;
+
+	// http://fortawesome.github.io/Font-Awesome/icon/bell-slash-o/
+	static FText Bell_Slash_O;
+
+	// (vehicle, bike) http://fortawesome.github.io/Font-Awesome/icon/bicycle/
+	static FText Bicycle;
+
+	// http://fortawesome.github.io/Font-Awesome/icon/binoculars/
+	static FText Binoculars;
+
+	// http://fortawesome.github.io/Font-Awesome/icon/birthday-cake/
+	static FText Birthday_Cake;
+
+	// (git) http://fortawesome.github.io/Font-Awesome/icon/bitbucket/
+	static FText Bitbucket;
+
+	// (git) http://fortawesome.github.io/Font-Awesome/icon/bitbucket-square/
+	static FText Bitbucket_Square;
+
+	// http://fortawesome.github.io/Font-Awesome/icon/bold/
+	static FText Bold;
+
+	// (lightning, weather) http://fortawesome.github.io/Font-Awesome/icon/bolt/
+	static FText Bolt;
+
+	// http://fortawesome.github.io/Font-Awesome/icon/bomb/
+	static FText Bomb;
+
+	// (read, documentation) http://fortawesome.github.io/Font-Awesome/icon/book/
+	static FText Book;
+
+	// (save) http://fortawesome.github.io/Font-Awesome/icon/bookmark/
+	static FText Bookmark;
+
+	// (save) http://fortawesome.github.io/Font-Awesome/icon/bookmark-o/
+	static FText Bookmark_O;
+
+	// (work, business, office, luggage, bag) http://fortawesome.github.io/Font-Awesome/icon/briefcase/
+	static FText Briefcase;
+
+	// http://fortawesome.github.io/Font-Awesome/icon/btc/
+	static FText Btc;
+
+	// (report) http://fortawesome.github.io/Font-Awesome/icon/bug/
+	static FText Bug;
+
+	// http://fortawesome.github.io/Font-Awesome/icon/building/
+	static FText Building;
+
+	// (work, business, apartment, office) http://fortawesome.github.io/Font-Awesome/icon/building-o/
+	static FText Building_O;
+
+	// (announcement, share, broadcast, louder) http://fortawesome.github.io/Font-Awesome/icon/bullhorn/
+	static FText Bullhorn;
+
+	// (target) http://fortawesome.github.io/Font-Awesome/icon/bullseye/
+	static FText Bullseye;
+
+	// (vehicle) http://fortawesome.github.io/Font-Awesome/icon/bus/
+	static FText Bus;
+
+	// http://fortawesome.github.io/Font-Awesome/icon/buysellads/
+	static FText Buysellads;
+
+	// http://fortawesome.github.io/Font-Awesome/icon/calculator/
+	static FText Calculator;
+
+	// (date, time, when) http://fortawesome.github.io/Font-Awesome/icon/calendar/
+	static FText Calendar;
+
+	// (date, time, when) http://fortawesome.github.io/Font-Awesome/icon/calendar-o/
+	static FText Calendar_O;
+
+	// (photo, picture, record) http://fortawesome.github.io/Font-Awesome/icon/camera/
+	static FText Camera;
+
+	// (photo, picture, record) http://fortawesome.github.io/Font-Awesome/icon/camera-retro/
+	static FText Camera_Retro;
+
+	// (vehicle) http://fortawesome.github.io/Font-Awesome/icon/car/
+	static FText Car;
+
+	// (more, dropdown, menu) http://fortawesome.github.io/Font-Awesome/icon/caret-down/
+	static FText Caret_Down;
+
+	// (previous, back) http://fortawesome.github.io/Font-Awesome/icon/caret-left/
+	static FText Caret_Left;
+
+	// (next, forward) http://fortawesome.github.io/Font-Awesome/icon/caret-right/
+	static FText Caret_Right;
+
+	// (more, dropdown, menu) http://fortawesome.github.io/Font-Awesome/icon/caret-square-o-down/
+	static FText Caret_Square_O_Down;
+
+	// (previous, back) http://fortawesome.github.io/Font-Awesome/icon/caret-square-o-left/
+	static FText Caret_Square_O_Left;
+
+	// (next, forward) http://fortawesome.github.io/Font-Awesome/icon/caret-square-o-right/
+	static FText Caret_Square_O_Right;
+
+	// http://fortawesome.github.io/Font-Awesome/icon/caret-square-o-up/
+	static FText Caret_Square_O_Up;
+
+	// http://fortawesome.github.io/Font-Awesome/icon/caret-up/
+	static FText Caret_Up;
+
+	// (shopping) http://fortawesome.github.io/Font-Awesome/icon/cart-arrow-down/
+	static FText Cart_Arrow_Down;
+
+	// (add, shopping) http://fortawesome.github.io/Font-Awesome/icon/cart-plus/
+	static FText Cart_Plus;
+
+	// http://fortawesome.github.io/Font-Awesome/icon/cc/
+	static FText Cc;
+
+	// http://fortawesome.github.io/Font-Awesome/icon/cc-amex/
+	static FText Cc_Amex;
+
+	// http://fortawesome.github.io/Font-Awesome/icon/cc-discover/
+	static FText Cc_Discover;
+
+	// http://fortawesome.github.io/Font-Awesome/icon/cc-mastercard/
+	static FText Cc_Mastercard;
+
+	// http://fortawesome.github.io/Font-Awesome/icon/cc-paypal/
+	static FText Cc_Paypal;
+
+	// http://fortawesome.github.io/Font-Awesome/icon/cc-stripe/
+	static FText Cc_Stripe;
+
+	// http://fortawesome.github.io/Font-Awesome/icon/cc-visa/
+	static FText Cc_Visa;
+
+	// (badge, star) http://fortawesome.github.io/Font-Awesome/icon/certificate/
+	static FText Certificate;
+
+	// (remove) http://fortawesome.github.io/Font-Awesome/icon/chain-broken/
+	static FText Chain_Broken;
+
+	// (checkmark, done, todo, agree, accept, confirm) http://fortawesome.github.io/Font-Awesome/icon/check/
+	static FText Check;
+
+	// (todo, done, agree, accept, confirm) http://fortawesome.github.io/Font-Awesome/icon/check-circle/
+	static FText Check_Circle;
+
+	// (todo, done, agree, accept, confirm) http://fortawesome.github.io/Font-Awesome/icon/check-circle-o/
+	static FText Check_Circle_O;
+
+	// (checkmark, done, todo, agree, accept, confirm) http://fortawesome.github.io/Font-Awesome/icon/check-square/
+	static FText Check_Square;
+
+	// (todo, done, agree, accept, confirm) http://fortawesome.github.io/Font-Awesome/icon/check-square-o/
+	static FText Check_Square_O;
+
+	// (more, dropdown, menu) http://fortawesome.github.io/Font-Awesome/icon/chevron-circle-down/
+	static FText Chevron_Circle_Down;
+
+	// (previous, back) http://fortawesome.github.io/Font-Awesome/icon/chevron-circle-left/
+	static FText Chevron_Circle_Left;
+
+	// (next, forward) http://fortawesome.github.io/Font-Awesome/icon/chevron-circle-right/
+	static FText Chevron_Circle_Right;
+
+	// http://fortawesome.github.io/Font-Awesome/icon/chevron-circle-up/
+	static FText Chevron_Circle_Up;
+
+	// http://fortawesome.github.io/Font-Awesome/icon/chevron-down/
+	static FText Chevron_Down;
+
+	// (bracket, previous, back) http://fortawesome.github.io/Font-Awesome/icon/chevron-left/
+	static FText Chevron_Left;
+
+	// (bracket, next, forward) http://fortawesome.github.io/Font-Awesome/icon/chevron-right/
+	static FText Chevron_Right;
+
+	// http://fortawesome.github.io/Font-Awesome/icon/chevron-up/
+	static FText Chevron_Up;
+
+	// http://fortawesome.github.io/Font-Awesome/icon/child/
+	static FText Child;
+
+	// (dot, notification) http://fortawesome.github.io/Font-Awesome/icon/circle/
+	static FText Circle;
+
+	// http://fortawesome.github.io/Font-Awesome/icon/circle-o/
+	static FText Circle_O;
+
+	// http://fortawesome.github.io/Font-Awesome/icon/circle-o-notch/
+	static FText Circle_O_Notch;
+
+	// http://fortawesome.github.io/Font-Awesome/icon/circle-thin/
+	static FText Circle_Thin;
+
+	// (copy) http://fortawesome.github.io/Font-Awesome/icon/clipboard/
+	static FText Clipboard;
+
+	// (watch, timer, late, timestamp) http://fortawesome.github.io/Font-Awesome/icon/clock-o/
+	static FText Clock_O;
+
+	// (save) http://fortawesome.github.io/Font-Awesome/icon/cloud/
+	static FText Cloud;
+
+	// (import) http://fortawesome.github.io/Font-Awesome/icon/cloud-download/
+	static FText Cloud_Download;
+
+	// (import) http://fortawesome.github.io/Font-Awesome/icon/cloud-upload/
+	static FText Cloud_Upload;
+
+	// (html, brackets) http://fortawesome.github.io/Font-Awesome/icon/code/
+	static FText Code;
+
+	// (git, fork, vcs, svn, github, rebase, version, merge) http://fortawesome.github.io/Font-Awesome/icon/code-fork/
+	static FText Code_Fork;
+
+	// http://fortawesome.github.io/Font-Awesome/icon/codepen/
+	static FText Codepen;
+
+	// (morning, mug, breakfast, tea, drink, cafe) http://fortawesome.github.io/Font-Awesome/icon/coffee/
+	static FText Coffee;
+
+	// (settings) http://fortawesome.github.io/Font-Awesome/icon/cog/
+	static FText Cog;
+
+	// (settings) http://fortawesome.github.io/Font-Awesome/icon/cogs/
+	static FText Cogs;
+
+	// (split, panes) http://fortawesome.github.io/Font-Awesome/icon/columns/
+	static FText Columns;
+
+	// (speech, notification, note) http://fortawesome.github.io/Font-Awesome/icon/comment/
+	static FText Comment;
+
+	// (notification, note) http://fortawesome.github.io/Font-Awesome/icon/comment-o/
+	static FText Comment_O;
+
+	// (conversation, notification, notes) http://fortawesome.github.io/Font-Awesome/icon/comments/
+	static FText Comments;
+
+	// (conversation, notification, notes) http://fortawesome.github.io/Font-Awesome/icon/comments-o/
+	static FText Comments_O;
+
+	// (safari, directory, menu, location) http://fortawesome.github.io/Font-Awesome/icon/compass/
+	static FText Compass;
+
+	// (combine, merge, smaller) http://fortawesome.github.io/Font-Awesome/icon/compress/
+	static FText Compress;
+
+	// http://fortawesome.github.io/Font-Awesome/icon/connectdevelop/
+	static FText Connectdevelop;
+
+	// http://fortawesome.github.io/Font-Awesome/icon/copyright/
+	static FText Copyright;
+
+	// (money, buy, debit, checkout, purchase, payment) http://fortawesome.github.io/Font-Awesome/icon/credit-card/
+	static FText Credit_Card;
+
+	// http://fortawesome.github.io/Font-Awesome/icon/crop/
+	static FText Crop;
+
+	// (picker) http://fortawesome.github.io/Font-Awesome/icon/crosshairs/
+	static FText Crosshairs;
+
+	// (code) http://fortawesome.github.io/Font-Awesome/icon/css3/
+	static FText Css3;
+
+	// http://fortawesome.github.io/Font-Awesome/icon/cube/
+	static FText Cube;
+
+	// http://fortawesome.github.io/Font-Awesome/icon/cubes/
+	static FText Cubes;
+
+	// (food, restaurant, spoon, knife, dinner, eat) http://fortawesome.github.io/Font-Awesome/icon/cutlery/
+	static FText Cutlery;
+
+	// http://fortawesome.github.io/Font-Awesome/icon/dashcube/
+	static FText Dashcube;
+
+	// http://fortawesome.github.io/Font-Awesome/icon/database/
+	static FText Database;
+
+	// http://fortawesome.github.io/Font-Awesome/icon/delicious/
+	static FText Delicious;
+
+	// (monitor, screen, desktop, computer, demo, device) http://fortawesome.github.io/Font-Awesome/icon/desktop/
+	static FText Desktop;
+
+	// http://fortawesome.github.io/Font-Awesome/icon/deviantart/
+	static FText Deviantart;
+
+	// (gem, gemstone) http://fortawesome.github.io/Font-Awesome/icon/diamond/
+	static FText Diamond;
+
+	// http://fortawesome.github.io/Font-Awesome/icon/digg/
+	static FText Digg;
+
+	// (target, bullseye, notification) http://fortawesome.github.io/Font-Awesome/icon/dot-circle-o/
+	static FText Dot_Circle_O;
+
+	// (import) http://fortawesome.github.io/Font-Awesome/icon/download/
+	static FText Download;
+
+	// http://fortawesome.github.io/Font-Awesome/icon/dribbble/
+	static FText Dribbble;
+
+	// http://fortawesome.github.io/Font-Awesome/icon/dropbox/
+	static FText Dropbox;
+
+	// http://fortawesome.github.io/Font-Awesome/icon/drupal/
+	static FText Drupal;
+
+	// http://fortawesome.github.io/Font-Awesome/icon/eject/
+	static FText Eject;
+
+	// (dots) http://fortawesome.github.io/Font-Awesome/icon/ellipsis-h/
+	static FText Ellipsis_H;
+
+	// (dots) http://fortawesome.github.io/Font-Awesome/icon/ellipsis-v/
+	static FText Ellipsis_V;
+
+	// http://fortawesome.github.io/Font-Awesome/icon/empire/
+	static FText Empire;
+
+	// http://fortawesome.github.io/Font-Awesome/icon/envelope/
+	static FText Envelope;
+
+	// (email, support, e-mail, letter, mail, notification) http://fortawesome.github.io/Font-Awesome/icon/envelope-o/
+	static FText Envelope_O;
+
+	// http://fortawesome.github.io/Font-Awesome/icon/envelope-square/
+	static FText Envelope_Square;
+
+	// http://fortawesome.github.io/Font-Awesome/icon/eraser/
+	static FText Eraser;
+
+	// http://fortawesome.github.io/Font-Awesome/icon/eur/
+	static FText Eur;
+
+	// http://fortawesome.github.io/Font-Awesome/icon/exchange/
+	static FText Exchange;
+
+	// (warning, error, problem, notification, notify, alert) http://fortawesome.github.io/Font-Awesome/icon/exclamation/
+	static FText Exclamation;
+
+	// (warning, error, problem, notification, alert) http://fortawesome.github.io/Font-Awesome/icon/exclamation-circle/
+	static FText Exclamation_Circle;
+
+	// (warning, error, problem, notification, alert) http://fortawesome.github.io/Font-Awesome/icon/exclamation-triangle/
+	static FText Exclamation_Triangle;
+
+	// (enlarge, bigger, resize) http://fortawesome.github.io/Font-Awesome/icon/expand/
+	static FText Expand;
+
+	// (open, new) http://fortawesome.github.io/Font-Awesome/icon/external-link/
+	static FText External_Link;
+
+	// (open, new) http://fortawesome.github.io/Font-Awesome/icon/external-link-square/
+	static FText External_Link_Square;
+
+	// (show, visible, views) http://fortawesome.github.io/Font-Awesome/icon/eye/
+	static FText Eye;
+
+	// http://fortawesome.github.io/Font-Awesome/icon/eyedropper/
+	static FText Eyedropper;
+
+	// (toggle, show, hide, visible, visiblity, views) http://fortawesome.github.io/Font-Awesome/icon/eye-slash/
+	static FText Eye_Slash;
+
+	// (social network) http://fortawesome.github.io/Font-Awesome/icon/facebook/
+	static FText Facebook;
+
+	// http://fortawesome.github.io/Font-Awesome/icon/facebook-official/
+	static FText Facebook_Official;
+
+	// (social network) http://fortawesome.github.io/Font-Awesome/icon/facebook-square/
+	static FText Facebook_Square;
+
+	// (rewind, previous, beginning, start, first) http://fortawesome.github.io/Font-Awesome/icon/fast-backward/
+	static FText Fast_Backward;
+
+	// (next, end, last) http://fortawesome.github.io/Font-Awesome/icon/fast-forward/
+	static FText Fast_Forward;
+
+	// http://fortawesome.github.io/Font-Awesome/icon/fax/
+	static FText Fax;
+
+	// (woman, user, person, profile) http://fortawesome.github.io/Font-Awesome/icon/female/
+	static FText Female;
+
+	// (fly, plane, airplane, quick, fast, travel) http://fortawesome.github.io/Font-Awesome/icon/fighter-jet/
+	static FText Fighter_Jet;
+
+	// (new, page, pdf, document) http://fortawesome.github.io/Font-Awesome/icon/file/
+	static FText File;
+
+	// http://fortawesome.github.io/Font-Awesome/icon/file-archive-o/
+	static FText File_Archive_O;
+
+	// http://fortawesome.github.io/Font-Awesome/icon/file-audio-o/
+	static FText File_Audio_O;
+
+	// http://fortawesome.github.io/Font-Awesome/icon/file-code-o/
+	static FText File_Code_O;
+
+	// http://fortawesome.github.io/Font-Awesome/icon/file-excel-o/
+	static FText File_Excel_O;
+
+	// http://fortawesome.github.io/Font-Awesome/icon/file-image-o/
+	static FText File_Image_O;
+
+	// (new, page, pdf, document) http://fortawesome.github.io/Font-Awesome/icon/file-o/
+	static FText File_O;
+
+	// http://fortawesome.github.io/Font-Awesome/icon/file-pdf-o/
+	static FText File_Pdf_O;
+
+	// http://fortawesome.github.io/Font-Awesome/icon/file-powerpoint-o/
+	static FText File_Powerpoint_O;
+
+	// (duplicate) http://fortawesome.github.io/Font-Awesome/icon/files-o/
+	static FText Files_O;
+
+	// (new, page, pdf, document) http://fortawesome.github.io/Font-Awesome/icon/file-text/
+	static FText File_Text;
+
+	// (new, page, pdf, document) http://fortawesome.github.io/Font-Awesome/icon/file-text-o/
+	static FText File_Text_O;
+
+	// http://fortawesome.github.io/Font-Awesome/icon/file-video-o/
+	static FText File_Video_O;
+
+	// http://fortawesome.github.io/Font-Awesome/icon/file-word-o/
+	static FText File_Word_O;
+
+	// (movie) http://fortawesome.github.io/Font-Awesome/icon/film/
+	static FText Film;
+
+	// (funnel, options) http://fortawesome.github.io/Font-Awesome/icon/filter/
+	static FText Filter;
+
+	// (flame, hot, popular) http://fortawesome.github.io/Font-Awesome/icon/fire/
+	static FText Fire;
+
+	// http://fortawesome.github.io/Font-Awesome/icon/fire-extinguisher/
+	static FText Fire_Extinguisher;
+
+	// (report, notification, notify) http://fortawesome.github.io/Font-Awesome/icon/flag/
+	static FText Flag;
+
+	// (report, notification, notify) http://fortawesome.github.io/Font-Awesome/icon/flag-checkered/
+	static FText Flag_Checkered;
+
+	// (report, notification) http://fortawesome.github.io/Font-Awesome/icon/flag-o/
+	static FText Flag_O;
+
+	// (science, beaker, experimental, labs) http://fortawesome.github.io/Font-Awesome/icon/flask/
+	static FText Flask;
+
+	// http://fortawesome.github.io/Font-Awesome/icon/flickr/
+	static FText Flickr;
+
+	// http://fortawesome.github.io/Font-Awesome/icon/floppy-o/
+	static FText Floppy_O;
+
+	// http://fortawesome.github.io/Font-Awesome/icon/folder/
+	static FText Folder;
+
+	// http://fortawesome.github.io/Font-Awesome/icon/folder-o/
+	static FText Folder_O;
+
+	// http://fortawesome.github.io/Font-Awesome/icon/folder-open/
+	static FText Folder_Open;
+
+	// http://fortawesome.github.io/Font-Awesome/icon/folder-open-o/
+	static FText Folder_Open_O;
+
+	// (text) http://fortawesome.github.io/Font-Awesome/icon/font/
+	static FText Font;
+
+	// http://fortawesome.github.io/Font-Awesome/icon/forumbee/
+	static FText Forumbee;
+
+	// (forward, next) http://fortawesome.github.io/Font-Awesome/icon/forward/
+	static FText Forward;
+
+	// http://fortawesome.github.io/Font-Awesome/icon/foursquare/
+	static FText Foursquare;
+
+	// (emoticon, sad, disapprove, rating) http://fortawesome.github.io/Font-Awesome/icon/frown-o/
+	static FText Frown_O;
+
+	// http://fortawesome.github.io/Font-Awesome/icon/futbol-o/
+	static FText Futbol_O;
+
+	// (controller) http://fortawesome.github.io/Font-Awesome/icon/gamepad/
+	static FText Gamepad;
+
+	// http://fortawesome.github.io/Font-Awesome/icon/gavel/
+	static FText Gavel;
+
+	// http://fortawesome.github.io/Font-Awesome/icon/gbp/
+	static FText Gbp;
+
+	// (present) http://fortawesome.github.io/Font-Awesome/icon/gift/
+	static FText Gift;
+
+	// http://fortawesome.github.io/Font-Awesome/icon/git/
+	static FText Git;
+
+	// (octocat) http://fortawesome.github.io/Font-Awesome/icon/github/
+	static FText Github;
+
+	// (octocat) http://fortawesome.github.io/Font-Awesome/icon/github-alt/
+	static FText Github_Alt;
+
+	// (octocat) http://fortawesome.github.io/Font-Awesome/icon/github-square/
+	static FText Github_Square;
+
+	// http://fortawesome.github.io/Font-Awesome/icon/git-square/
+	static FText Git_Square;
+
+	// (martini, drink, bar, alcohol, liquor) http://fortawesome.github.io/Font-Awesome/icon/glass/
+	static FText Glass;
+
+	// (world, planet, map, place, travel, earth, global, translate, all, language, localize, location, coordinates, country) http://fortawesome.github.io/Font-Awesome/icon/globe/
+	static FText Globe;
+
+	// http://fortawesome.github.io/Font-Awesome/icon/google/
+	static FText Google;
+
+	// (social network) http://fortawesome.github.io/Font-Awesome/icon/google-plus/
+	static FText Google_Plus;
+
+	// (social network) http://fortawesome.github.io/Font-Awesome/icon/google-plus-square/
+	static FText Google_Plus_Square;
+
+	// http://fortawesome.github.io/Font-Awesome/icon/google-wallet/
+	static FText Google_Wallet;
+
+	// http://fortawesome.github.io/Font-Awesome/icon/graduation-cap/
+	static FText Graduation_Cap;
+
+	// (heart, like, favorite, love) http://fortawesome.github.io/Font-Awesome/icon/gratipay/
+	static FText Gratipay;
+
+	// http://fortawesome.github.io/Font-Awesome/icon/hacker-news/
+	static FText Hacker_News;
+
+	// (point) http://fortawesome.github.io/Font-Awesome/icon/hand-o-down/
+	static FText Hand_O_Down;
+
+	// (point, left, previous, back) http://fortawesome.github.io/Font-Awesome/icon/hand-o-left/
+	static FText Hand_O_Left;
+
+	// (point, right, next, forward) http://fortawesome.github.io/Font-Awesome/icon/hand-o-right/
+	static FText Hand_O_Right;
+
+	// (point) http://fortawesome.github.io/Font-Awesome/icon/hand-o-up/
+	static FText Hand_O_Up;
+
+	// (harddrive, hard drive, storage, save) http://fortawesome.github.io/Font-Awesome/icon/hdd-o/
+	static FText Hdd_O;
+
+	// http://fortawesome.github.io/Font-Awesome/icon/header/
+	static FText Header;
+
+	// (sound, listen, music) http://fortawesome.github.io/Font-Awesome/icon/headphones/
+	static FText Headphones;
+
+	// (love, like, favorite) http://fortawesome.github.io/Font-Awesome/icon/heart/
+	static FText Heart;
+
+	// (ekg) http://fortawesome.github.io/Font-Awesome/icon/heartbeat/
+	static FText Heartbeat;
+
+	// (love, like, favorite) http://fortawesome.github.io/Font-Awesome/icon/heart-o/
+	static FText Heart_O;
+
+	// http://fortawesome.github.io/Font-Awesome/icon/history/
+	static FText History;
+
+	// (main, house) http://fortawesome.github.io/Font-Awesome/icon/home/
+	static FText Home;
+
+	// (building) http://fortawesome.github.io/Font-Awesome/icon/hospital-o/
+	static FText Hospital_O;
+
+	// (hospital, hotel) http://fortawesome.github.io/Font-Awesome/icon/h-square/
+	static FText H_Square;
+
+	// http://fortawesome.github.io/Font-Awesome/icon/html5/
+	static FText Html5;
+
+	// http://fortawesome.github.io/Font-Awesome/icon/ils/
+	static FText Ils;
+
+	// http://fortawesome.github.io/Font-Awesome/icon/inbox/
+	static FText Inbox;
+
+	// http://fortawesome.github.io/Font-Awesome/icon/indent/
+	static FText Indent;
+
+	// (help, information, more, details) http://fortawesome.github.io/Font-Awesome/icon/info/
+	static FText Info;
+
+	// (help, information, more, details) http://fortawesome.github.io/Font-Awesome/icon/info-circle/
+	static FText Info_Circle;
+
+	// http://fortawesome.github.io/Font-Awesome/icon/inr/
+	static FText Inr;
+
+	// http://fortawesome.github.io/Font-Awesome/icon/instagram/
+	static FText Instagram;
+
+	// http://fortawesome.github.io/Font-Awesome/icon/ioxhost/
+	static FText Ioxhost;
+
+	// (italics) http://fortawesome.github.io/Font-Awesome/icon/italic/
+	static FText Italic;
+
+	// http://fortawesome.github.io/Font-Awesome/icon/joomla/
+	static FText Joomla;
+
+	// http://fortawesome.github.io/Font-Awesome/icon/jpy/
+	static FText Jpy;
+
+	// http://fortawesome.github.io/Font-Awesome/icon/jsfiddle/
+	static FText Jsfiddle;
+
+	// (unlock, password) http://fortawesome.github.io/Font-Awesome/icon/key/
+	static FText Key;
+
+	// (type, input) http://fortawesome.github.io/Font-Awesome/icon/keyboard-o/
+	static FText Keyboard_O;
+
+	// http://fortawesome.github.io/Font-Awesome/icon/krw/
+	static FText Krw;
+
+	// http://fortawesome.github.io/Font-Awesome/icon/language/
+	static FText Language;
+
+	// (demo, computer, device) http://fortawesome.github.io/Font-Awesome/icon/laptop/
+	static FText Laptop;
+
+	// http://fortawesome.github.io/Font-Awesome/icon/lastfm/
+	static FText Lastfm;
+
+	// http://fortawesome.github.io/Font-Awesome/icon/lastfm-square/
+	static FText Lastfm_Square;
+
+	// (eco, nature) http://fortawesome.github.io/Font-Awesome/icon/leaf/
+	static FText Leaf;
+
+	// http://fortawesome.github.io/Font-Awesome/icon/leanpub/
+	static FText Leanpub;
+
+	// http://fortawesome.github.io/Font-Awesome/icon/lemon-o/
+	static FText Lemon_O;
+
+	// http://fortawesome.github.io/Font-Awesome/icon/level-down/
+	static FText Level_Down;
+
+	// http://fortawesome.github.io/Font-Awesome/icon/level-up/
+	static FText Level_Up;
+
+	// http://fortawesome.github.io/Font-Awesome/icon/life-ring/
+	static FText Life_Ring;
+
+	// (idea, inspiration) http://fortawesome.github.io/Font-Awesome/icon/lightbulb-o/
+	static FText Lightbulb_O;
+
+	// http://fortawesome.github.io/Font-Awesome/icon/line-chart/
+	static FText Line_Chart;
+
+	// (chain) http://fortawesome.github.io/Font-Awesome/icon/link/
+	static FText Link;
+
+	// http://fortawesome.github.io/Font-Awesome/icon/linkedin/
+	static FText Linkedin;
+
+	// http://fortawesome.github.io/Font-Awesome/icon/linkedin-square/
+	static FText Linkedin_Square;
+
+	// (tux) http://fortawesome.github.io/Font-Awesome/icon/linux/
+	static FText Linux;
+
+	// (ul, ol, checklist, finished, completed, done, todo) http://fortawesome.github.io/Font-Awesome/icon/list/
+	static FText List;
+
+	// (ul, ol, checklist, finished, completed, done, todo) http://fortawesome.github.io/Font-Awesome/icon/list-alt/
+	static FText List_Alt;
+
+	// (ul, ol, checklist, list, todo, list, numbers) http://fortawesome.github.io/Font-Awesome/icon/list-ol/
+	static FText List_Ol;
+
+	// (ul, ol, checklist, todo, list) http://fortawesome.github.io/Font-Awesome/icon/list-ul/
+	static FText List_Ul;
+
+	// (map, coordinates, location, address, place, where) http://fortawesome.github.io/Font-Awesome/icon/location-arrow/
+	static FText Location_Arrow;
+
+	// (protect, admin) http://fortawesome.github.io/Font-Awesome/icon/lock/
+	static FText Lock;
+
+	// http://fortawesome.github.io/Font-Awesome/icon/long-arrow-down/
+	static FText Long_Arrow_Down;
+
+	// (previous, back) http://fortawesome.github.io/Font-Awesome/icon/long-arrow-left/
+	static FText Long_Arrow_Left;
+
+	// http://fortawesome.github.io/Font-Awesome/icon/long-arrow-right/
+	static FText Long_Arrow_Right;
+
+	// http://fortawesome.github.io/Font-Awesome/icon/long-arrow-up/
+	static FText Long_Arrow_Up;
+
+	// (wizard, automatic, autocomplete) http://fortawesome.github.io/Font-Awesome/icon/magic/
+	static FText Magic;
+
+	// http://fortawesome.github.io/Font-Awesome/icon/magnet/
+	static FText Magnet;
+
+	// (man, user, person, profile) http://fortawesome.github.io/Font-Awesome/icon/male/
+	static FText Male;
+
+	// (map, pin, location, coordinates, localize, address, travel, where, place) http://fortawesome.github.io/Font-Awesome/icon/map-marker/
+	static FText Map_Marker;
+
+	// (male) http://fortawesome.github.io/Font-Awesome/icon/mars/
+	static FText Mars;
+
+	// http://fortawesome.github.io/Font-Awesome/icon/mars-double/
+	static FText Mars_Double;
+
+	// http://fortawesome.github.io/Font-Awesome/icon/mars-stroke/
+	static FText Mars_Stroke;
+
+	// http://fortawesome.github.io/Font-Awesome/icon/mars-stroke-h/
+	static FText Mars_Stroke_H;
+
+	// http://fortawesome.github.io/Font-Awesome/icon/mars-stroke-v/
+	static FText Mars_Stroke_V;
+
+	// http://fortawesome.github.io/Font-Awesome/icon/maxcdn/
+	static FText Maxcdn;
+
+	// http://fortawesome.github.io/Font-Awesome/icon/meanpath/
+	static FText Meanpath;
+
+	// http://fortawesome.github.io/Font-Awesome/icon/medium/
+	static FText Medium;
+
+	// (first aid, firstaid, help, support) http://fortawesome.github.io/Font-Awesome/icon/medkit/
+	static FText Medkit;
+
+	// (emoticon, rating, neutral) http://fortawesome.github.io/Font-Awesome/icon/meh-o/
+	static FText Meh_O;
+
+	// (transgender) http://fortawesome.github.io/Font-Awesome/icon/mercury/
+	static FText Mercury;
+
+	// (record, voice, sound) http://fortawesome.github.io/Font-Awesome/icon/microphone/
+	static FText Microphone;
+
+	// (record, voice, sound, mute) http://fortawesome.github.io/Font-Awesome/icon/microphone-slash/
+	static FText Microphone_Slash;
+
+	// (hide, minify, delete, remove, trash, hide, collapse) http://fortawesome.github.io/Font-Awesome/icon/minus/
+	static FText Minus;
+
+	// (delete, remove, trash, hide) http://fortawesome.github.io/Font-Awesome/icon/minus-circle/
+	static FText Minus_Circle;
+
+	// (hide, minify, delete, remove, trash, hide, collapse) http://fortawesome.github.io/Font-Awesome/icon/minus-square/
+	static FText Minus_Square;
+
+	// (hide, minify, delete, remove, trash, hide, collapse) http://fortawesome.github.io/Font-Awesome/icon/minus-square-o/
+	static FText Minus_Square_O;
+
+	// (cell phone, cellphone, text, call, iphone, number) http://fortawesome.github.io/Font-Awesome/icon/mobile/
+	static FText Mobile;
+
+	// (cash, money, buy, checkout, purchase, payment) http://fortawesome.github.io/Font-Awesome/icon/money/
+	static FText Money;
+
+	// (night, darker, contrast) http://fortawesome.github.io/Font-Awesome/icon/moon-o/
+	static FText Moon_O;
+
+	// (vehicle, bike) http://fortawesome.github.io/Font-Awesome/icon/motorcycle/
+	static FText Motorcycle;
+
+	// (note, sound) http://fortawesome.github.io/Font-Awesome/icon/music/
+	static FText Music;
+
+	// http://fortawesome.github.io/Font-Awesome/icon/neuter/
+	static FText Neuter;
+
+	// http://fortawesome.github.io/Font-Awesome/icon/newspaper-o/
+	static FText Newspaper_O;
+
+	// http://fortawesome.github.io/Font-Awesome/icon/openid/
+	static FText Openid;
+
+	// http://fortawesome.github.io/Font-Awesome/icon/outdent/
+	static FText Outdent;
+
+	// (leaf, leaves, tree, plant, eco, nature) http://fortawesome.github.io/Font-Awesome/icon/pagelines/
+	static FText Pagelines;
+
+	// http://fortawesome.github.io/Font-Awesome/icon/paint-brush/
+	static FText Paint_Brush;
+
+	// (attachment) http://fortawesome.github.io/Font-Awesome/icon/paperclip/
+	static FText Paperclip;
+
+	// http://fortawesome.github.io/Font-Awesome/icon/paper-plane/
+	static FText Paper_Plane;
+
+	// http://fortawesome.github.io/Font-Awesome/icon/paper-plane-o/
+	static FText Paper_Plane_O;
+
+	// http://fortawesome.github.io/Font-Awesome/icon/paragraph/
+	static FText Paragraph;
+
+	// (wait) http://fortawesome.github.io/Font-Awesome/icon/pause/
+	static FText Pause;
+
+	// http://fortawesome.github.io/Font-Awesome/icon/paw/
+	static FText Paw;
+
+	// http://fortawesome.github.io/Font-Awesome/icon/paypal/
+	static FText Paypal;
+
+	// (write, edit, update) http://fortawesome.github.io/Font-Awesome/icon/pencil/
+	static FText Pencil;
+
+	// (write, edit, update) http://fortawesome.github.io/Font-Awesome/icon/pencil-square/
+	static FText Pencil_Square;
+
+	// (write, edit, update) http://fortawesome.github.io/Font-Awesome/icon/pencil-square-o/
+	static FText Pencil_Square_O;
+
+	// (call, voice, number, support) http://fortawesome.github.io/Font-Awesome/icon/phone/
+	static FText Phone;
+
+	// (call, voice, number, support) http://fortawesome.github.io/Font-Awesome/icon/phone-square/
+	static FText Phone_Square;
+
+	// http://fortawesome.github.io/Font-Awesome/icon/picture-o/
+	static FText Picture_O;
+
+	// http://fortawesome.github.io/Font-Awesome/icon/pie-chart/
+	static FText Pie_Chart;
+
+	// http://fortawesome.github.io/Font-Awesome/icon/pied-piper/
+	static FText Pied_Piper;
+
+	// http://fortawesome.github.io/Font-Awesome/icon/pied-piper-alt/
+	static FText Pied_Piper_Alt;
+
+	// http://fortawesome.github.io/Font-Awesome/icon/pinterest/
+	static FText Pinterest;
+
+	// http://fortawesome.github.io/Font-Awesome/icon/pinterest-p/
+	static FText Pinterest_P;
+
+	// http://fortawesome.github.io/Font-Awesome/icon/pinterest-square/
+	static FText Pinterest_Square;
+
+	// (travel, trip, location, destination, airplane, fly, mode) http://fortawesome.github.io/Font-Awesome/icon/plane/
+	static FText Plane;
+
+	// (start, playing, music, sound) http://fortawesome.github.io/Font-Awesome/icon/play/
+	static FText Play;
+
+	// (start, playing) http://fortawesome.github.io/Font-Awesome/icon/play-circle/
+	static FText Play_Circle;
+
+	// http://fortawesome.github.io/Font-Awesome/icon/play-circle-o/
+	static FText Play_Circle_O;
+
+	// http://fortawesome.github.io/Font-Awesome/icon/plug/
+	static FText Plug;
+
+	// (add, new, create, expand) http://fortawesome.github.io/Font-Awesome/icon/plus/
+	static FText Plus;
+
+	// (add, new, create, expand) http://fortawesome.github.io/Font-Awesome/icon/plus-circle/
+	static FText Plus_Circle;
+
+	// (add, new, create, expand) http://fortawesome.github.io/Font-Awesome/icon/plus-square/
+	static FText Plus_Square;
+
+	// (add, new, create, expand) http://fortawesome.github.io/Font-Awesome/icon/plus-square-o/
+	static FText Plus_Square_O;
+
+	// (on) http://fortawesome.github.io/Font-Awesome/icon/power-off/
+	static FText Power_Off;
+
+	// http://fortawesome.github.io/Font-Awesome/icon/print/
+	static FText Print;
+
+	// (addon, add-on, section) http://fortawesome.github.io/Font-Awesome/icon/puzzle-piece/
+	static FText Puzzle_Piece;
+
+	// http://fortawesome.github.io/Font-Awesome/icon/qq/
+	static FText Qq;
+
+	// (scan) http://fortawesome.github.io/Font-Awesome/icon/qrcode/
+	static FText Qrcode;
+
+	// (help, information, unknown, support) http://fortawesome.github.io/Font-Awesome/icon/question/
+	static FText Question;
+
+	// (help, information, unknown, support) http://fortawesome.github.io/Font-Awesome/icon/question-circle/
+	static FText Question_Circle;
+
+	// http://fortawesome.github.io/Font-Awesome/icon/quote-left/
+	static FText Quote_Left;
+
+	// http://fortawesome.github.io/Font-Awesome/icon/quote-right/
+	static FText Quote_Right;
+
+	// (sort) http://fortawesome.github.io/Font-Awesome/icon/random/
+	static FText Random;
+
+	// http://fortawesome.github.io/Font-Awesome/icon/rebel/
+	static FText Rebel;
+
+	// http://fortawesome.github.io/Font-Awesome/icon/recycle/
+	static FText Recycle;
+
+	// http://fortawesome.github.io/Font-Awesome/icon/reddit/
+	static FText Reddit;
+
+	// http://fortawesome.github.io/Font-Awesome/icon/reddit-square/
+	static FText Reddit_Square;
+
+	// (reload) http://fortawesome.github.io/Font-Awesome/icon/refresh/
+	static FText Refresh;
+
+	// http://fortawesome.github.io/Font-Awesome/icon/renren/
+	static FText Renren;
+
+	// (redo, forward) http://fortawesome.github.io/Font-Awesome/icon/repeat/
+	static FText Repeat;
+
+	// http://fortawesome.github.io/Font-Awesome/icon/reply/
+	static FText Reply;
+
+	// http://fortawesome.github.io/Font-Awesome/icon/reply-all/
+	static FText Reply_All;
+
+	// (refresh, reload, share) http://fortawesome.github.io/Font-Awesome/icon/retweet/
+	static FText Retweet;
+
+	// (street) http://fortawesome.github.io/Font-Awesome/icon/road/
+	static FText Road;
+
+	// (app) http://fortawesome.github.io/Font-Awesome/icon/rocket/
+	static FText Rocket;
+
+	// (feed, blog) http://fortawesome.github.io/Font-Awesome/icon/rss/
+	static FText Rss;
+
+	// (feed, blog) http://fortawesome.github.io/Font-Awesome/icon/rss-square/
+	static FText Rss_Square;
+
+	// http://fortawesome.github.io/Font-Awesome/icon/rub/
+	static FText Rub;
+
+	// http://fortawesome.github.io/Font-Awesome/icon/scissors/
+	static FText Scissors;
+
+	// (magnify, zoom, enlarge, bigger) http://fortawesome.github.io/Font-Awesome/icon/search/
+	static FText Search;
+
+	// (magnify, minify, zoom, smaller) http://fortawesome.github.io/Font-Awesome/icon/search-minus/
+	static FText Search_Minus;
+
+	// (magnify, zoom, enlarge, bigger) http://fortawesome.github.io/Font-Awesome/icon/search-plus/
+	static FText Search_Plus;
+
+	// http://fortawesome.github.io/Font-Awesome/icon/sellsy/
+	static FText Sellsy;
+
+	// http://fortawesome.github.io/Font-Awesome/icon/server/
+	static FText Server;
+
+	// http://fortawesome.github.io/Font-Awesome/icon/share/
+	static FText Share;
+
+	// http://fortawesome.github.io/Font-Awesome/icon/share-alt/
+	static FText Share_Alt;
+
+	// http://fortawesome.github.io/Font-Awesome/icon/share-alt-square/
+	static FText Share_Alt_Square;
+
+	// (social, send) http://fortawesome.github.io/Font-Awesome/icon/share-square/
+	static FText Share_Square;
+
+	// (social, send) http://fortawesome.github.io/Font-Awesome/icon/share-square-o/
+	static FText Share_Square_O;
+
+	// (award, achievement, winner) http://fortawesome.github.io/Font-Awesome/icon/shield/
+	static FText Shield;
+
+	// (boat, sea) http://fortawesome.github.io/Font-Awesome/icon/ship/
+	static FText Ship;
+
+	// http://fortawesome.github.io/Font-Awesome/icon/shirtsinbulk/
+	static FText Shirtsinbulk;
+
+	// (checkout, buy, purchase, payment) http://fortawesome.github.io/Font-Awesome/icon/shopping-cart/
+	static FText Shopping_Cart;
+
+	// http://fortawesome.github.io/Font-Awesome/icon/signal/
+	static FText Signal;
+
+	// (enter, join, sign up, sign in, signin, signup, arrow) http://fortawesome.github.io/Font-Awesome/icon/sign-in/
+	static FText Sign_In;
+
+	// (log out, logout, leave, exit, arrow) http://fortawesome.github.io/Font-Awesome/icon/sign-out/
+	static FText Sign_Out;
+
+	// http://fortawesome.github.io/Font-Awesome/icon/simplybuilt/
+	static FText Simplybuilt;
+
+	// (directory, hierarchy, organization) http://fortawesome.github.io/Font-Awesome/icon/sitemap/
+	static FText Sitemap;
+
+	// http://fortawesome.github.io/Font-Awesome/icon/skyatlas/
+	static FText Skyatlas;
+
+	// http://fortawesome.github.io/Font-Awesome/icon/skype/
+	static FText Skype;
+
+	// http://fortawesome.github.io/Font-Awesome/icon/slack/
+	static FText Slack;
+
+	// http://fortawesome.github.io/Font-Awesome/icon/sliders/
+	static FText Sliders;
+
+	// http://fortawesome.github.io/Font-Awesome/icon/slideshare/
+	static FText Slideshare;
+
+	// (emoticon, happy, approve, satisfied, rating) http://fortawesome.github.io/Font-Awesome/icon/smile-o/
+	static FText Smile_O;
+
+	// (order) http://fortawesome.github.io/Font-Awesome/icon/sort/
+	static FText Sort;
+
+	// http://fortawesome.github.io/Font-Awesome/icon/sort-alpha-asc/
+	static FText Sort_Alpha_Asc;
+
+	// http://fortawesome.github.io/Font-Awesome/icon/sort-alpha-desc/
+	static FText Sort_Alpha_Desc;
+
+	// http://fortawesome.github.io/Font-Awesome/icon/sort-amount-asc/
+	static FText Sort_Amount_Asc;
+
+	// http://fortawesome.github.io/Font-Awesome/icon/sort-amount-desc/
+	static FText Sort_Amount_Desc;
+
+	// http://fortawesome.github.io/Font-Awesome/icon/sort-asc/
+	static FText Sort_Asc;
+
+	// (dropdown, more, menu) http://fortawesome.github.io/Font-Awesome/icon/sort-desc/
+	static FText Sort_Desc;
+
+	// (numbers) http://fortawesome.github.io/Font-Awesome/icon/sort-numeric-asc/
+	static FText Sort_Numeric_Asc;
+
+	// (numbers) http://fortawesome.github.io/Font-Awesome/icon/sort-numeric-desc/
+	static FText Sort_Numeric_Desc;
+
+	// http://fortawesome.github.io/Font-Awesome/icon/soundcloud/
+	static FText Soundcloud;
+
+	// http://fortawesome.github.io/Font-Awesome/icon/space-shuttle/
+	static FText Space_Shuttle;
+
+	// (loading, progress) http://fortawesome.github.io/Font-Awesome/icon/spinner/
+	static FText Spinner;
+
+	// http://fortawesome.github.io/Font-Awesome/icon/spoon/
+	static FText Spoon;
+
+	// http://fortawesome.github.io/Font-Awesome/icon/spotify/
+	static FText Spotify;
+
+	// (block, box) http://fortawesome.github.io/Font-Awesome/icon/square/
+	static FText Square;
+
+	// (block, square, box) http://fortawesome.github.io/Font-Awesome/icon/square-o/
+	static FText Square_O;
+
+	// http://fortawesome.github.io/Font-Awesome/icon/stack-exchange/
+	static FText Stack_Exchange;
+
+	// http://fortawesome.github.io/Font-Awesome/icon/stack-overflow/
+	static FText Stack_Overflow;
+
+	// (award, achievement, night, rating, score) http://fortawesome.github.io/Font-Awesome/icon/star/
+	static FText Star;
+
+	// (award, achievement, rating, score) http://fortawesome.github.io/Font-Awesome/icon/star-half/
+	static FText Star_Half;
+
+	// (award, achievement, rating, score) http://fortawesome.github.io/Font-Awesome/icon/star-half-o/
+	static FText Star_Half_O;
+
+	// (award, achievement, night, rating, score) http://fortawesome.github.io/Font-Awesome/icon/star-o/
+	static FText Star_O;
+
+	// http://fortawesome.github.io/Font-Awesome/icon/steam/
+	static FText Steam;
+
+	// http://fortawesome.github.io/Font-Awesome/icon/steam-square/
+	static FText Steam_Square;
+
+	// (rewind, previous, beginning, start, first) http://fortawesome.github.io/Font-Awesome/icon/step-backward/
+	static FText Step_Backward;
+
+	// (next, end, last) http://fortawesome.github.io/Font-Awesome/icon/step-forward/
+	static FText Step_Forward;
+
+	// http://fortawesome.github.io/Font-Awesome/icon/stethoscope/
+	static FText Stethoscope;
+
+	// (block, box, square) http://fortawesome.github.io/Font-Awesome/icon/stop/
+	static FText Stop;
+
+	// (map) http://fortawesome.github.io/Font-Awesome/icon/street-view/
+	static FText Street_View;
+
+	// http://fortawesome.github.io/Font-Awesome/icon/strikethrough/
+	static FText Strikethrough;
+
+	// http://fortawesome.github.io/Font-Awesome/icon/stumbleupon/
+	static FText Stumbleupon;
+
+	// http://fortawesome.github.io/Font-Awesome/icon/stumbleupon-circle/
+	static FText Stumbleupon_Circle;
+
+	// http://fortawesome.github.io/Font-Awesome/icon/subscript/
+	static FText Subscript;
+
+	// http://fortawesome.github.io/Font-Awesome/icon/subway/
+	static FText Subway;
+
+	// (trip, luggage, travel, move, baggage) http://fortawesome.github.io/Font-Awesome/icon/suitcase/
+	static FText Suitcase;
+
+	// (weather, contrast, lighter, brighten, day) http://fortawesome.github.io/Font-Awesome/icon/sun-o/
+	static FText Sun_O;
+
+	// (exponential) http://fortawesome.github.io/Font-Awesome/icon/superscript/
+	static FText Superscript;
+
+	// (data, excel, spreadsheet) http://fortawesome.github.io/Font-Awesome/icon/table/
+	static FText Table;
+
+	// (ipad, device) http://fortawesome.github.io/Font-Awesome/icon/tablet/
+	static FText Tablet;
+
+	// http://fortawesome.github.io/Font-Awesome/icon/tachometer/
+	static FText Tachometer;
+
+	// (label) http://fortawesome.github.io/Font-Awesome/icon/tag/
+	static FText Tag;
+
+	// (labels) http://fortawesome.github.io/Font-Awesome/icon/tags/
+	static FText Tags;
+
+	// (progress, loading, downloading, downloads, settings) http://fortawesome.github.io/Font-Awesome/icon/tasks/
+	static FText Tasks;
+
+	// (vehicle) http://fortawesome.github.io/Font-Awesome/icon/taxi/
+	static FText Taxi;
+
+	// http://fortawesome.github.io/Font-Awesome/icon/tencent-weibo/
+	static FText Tencent_Weibo;
+
+	// (command, prompt, code) http://fortawesome.github.io/Font-Awesome/icon/terminal/
+	static FText Terminal;
+
+	// http://fortawesome.github.io/Font-Awesome/icon/text-height/
+	static FText Text_Height;
+
+	// http://fortawesome.github.io/Font-Awesome/icon/text-width/
+	static FText Text_Width;
+
+	// (blocks, squares, boxes) http://fortawesome.github.io/Font-Awesome/icon/th/
+	static FText Th;
+
+	// (blocks, squares, boxes) http://fortawesome.github.io/Font-Awesome/icon/th-large/
+	static FText Th_Large;
+
+	// (ul, ol, checklist, finished, completed, done, todo) http://fortawesome.github.io/Font-Awesome/icon/th-list/
+	static FText Th_List;
+
+	// (dislike, disapprove, disagree) http://fortawesome.github.io/Font-Awesome/icon/thumbs-down/
+	static FText Thumbs_Down;
+
+	// (dislike, disapprove, disagree) http://fortawesome.github.io/Font-Awesome/icon/thumbs-o-down/
+	static FText Thumbs_O_Down;
+
+	// (like, approve, favorite, agree) http://fortawesome.github.io/Font-Awesome/icon/thumbs-o-up/
+	static FText Thumbs_O_Up;
+
+	// (like, favorite, approve, agree) http://fortawesome.github.io/Font-Awesome/icon/thumbs-up/
+	static FText Thumbs_Up;
+
+	// (marker, pin, location, coordinates) http://fortawesome.github.io/Font-Awesome/icon/thumb-tack/
+	static FText Thumb_Tack;
+
+	// (movie, pass, support) http://fortawesome.github.io/Font-Awesome/icon/ticket/
+	static FText Ticket;
+
+	// (close, exit, x) http://fortawesome.github.io/Font-Awesome/icon/times/
+	static FText Times;
+
+	// (close, exit, x) http://fortawesome.github.io/Font-Awesome/icon/times-circle/
+	static FText Times_Circle;
+
+	// (close, exit, x) http://fortawesome.github.io/Font-Awesome/icon/times-circle-o/
+	static FText Times_Circle_O;
+
+	// (raindrop) http://fortawesome.github.io/Font-Awesome/icon/tint/
+	static FText Tint;
+
+	// http://fortawesome.github.io/Font-Awesome/icon/toggle-off/
+	static FText Toggle_Off;
+
+	// http://fortawesome.github.io/Font-Awesome/icon/toggle-on/
+	static FText Toggle_On;
+
+	// http://fortawesome.github.io/Font-Awesome/icon/train/
+	static FText Train;
+
+	// http://fortawesome.github.io/Font-Awesome/icon/transgender/
+	static FText Transgender;
+
+	// http://fortawesome.github.io/Font-Awesome/icon/transgender-alt/
+	static FText Transgender_Alt;
+
+	// http://fortawesome.github.io/Font-Awesome/icon/trash/
+	static FText Trash;
+
+	// (garbage, delete, remove, trash, hide) http://fortawesome.github.io/Font-Awesome/icon/trash-o/
+	static FText Trash_O;
+
+	// http://fortawesome.github.io/Font-Awesome/icon/tree/
+	static FText Tree;
+
+	// http://fortawesome.github.io/Font-Awesome/icon/trello/
+	static FText Trello;
+
+	// (award, achievement, winner, game) http://fortawesome.github.io/Font-Awesome/icon/trophy/
+	static FText Trophy;
+
+	// (shipping) http://fortawesome.github.io/Font-Awesome/icon/truck/
+	static FText Truck;
+
+	// http://fortawesome.github.io/Font-Awesome/icon/try/
+	static FText Try;
+
+	// http://fortawesome.github.io/Font-Awesome/icon/tty/
+	static FText Tty;
+
+	// http://fortawesome.github.io/Font-Awesome/icon/tumblr/
+	static FText Tumblr;
+
+	// http://fortawesome.github.io/Font-Awesome/icon/tumblr-square/
+	static FText Tumblr_Square;
+
+	// http://fortawesome.github.io/Font-Awesome/icon/twitch/
+	static FText Twitch;
+
+	// (tweet, social network) http://fortawesome.github.io/Font-Awesome/icon/twitter/
+	static FText Twitter;
+
+	// (tweet, social network) http://fortawesome.github.io/Font-Awesome/icon/twitter-square/
+	static FText Twitter_Square;
+
+	// http://fortawesome.github.io/Font-Awesome/icon/umbrella/
+	static FText Umbrella;
+
+	// http://fortawesome.github.io/Font-Awesome/icon/underline/
+	static FText Underline;
+
+	// (back) http://fortawesome.github.io/Font-Awesome/icon/undo/
+	static FText Undo;
+
+	// http://fortawesome.github.io/Font-Awesome/icon/university/
+	static FText University;
+
+	// (protect, admin, password) http://fortawesome.github.io/Font-Awesome/icon/unlock/
+	static FText Unlock;
+
+	// (protect, admin, password) http://fortawesome.github.io/Font-Awesome/icon/unlock-alt/
+	static FText Unlock_Alt;
+
+	// (import) http://fortawesome.github.io/Font-Awesome/icon/upload/
+	static FText Upload;
+
+	// http://fortawesome.github.io/Font-Awesome/icon/usd/
+	static FText Usd;
+
+	// (person, man, head, profile) http://fortawesome.github.io/Font-Awesome/icon/user/
+	static FText User;
+
+	// (doctor, profile, medical, nurse) http://fortawesome.github.io/Font-Awesome/icon/user-md/
+	static FText User_Md;
+
+	// http://fortawesome.github.io/Font-Awesome/icon/user-plus/
+	static FText User_Plus;
+
+	// (people, profiles, persons) http://fortawesome.github.io/Font-Awesome/icon/users/
+	static FText Users;
+
+	// (whisper, spy, incognito) http://fortawesome.github.io/Font-Awesome/icon/user-secret/
+	static FText User_Secret;
+
+	// http://fortawesome.github.io/Font-Awesome/icon/user-times/
+	static FText User_Times;
+
+	// (female) http://fortawesome.github.io/Font-Awesome/icon/venus/
+	static FText Venus;
+
+	// http://fortawesome.github.io/Font-Awesome/icon/venus-double/
+	static FText Venus_Double;
+
+	// http://fortawesome.github.io/Font-Awesome/icon/venus-mars/
+	static FText Venus_Mars;
+
+	// http://fortawesome.github.io/Font-Awesome/icon/viacoin/
+	static FText Viacoin;
+
+	// (film, movie, record) http://fortawesome.github.io/Font-Awesome/icon/video-camera/
+	static FText Video_Camera;
+
+	// http://fortawesome.github.io/Font-Awesome/icon/vimeo-square/
+	static FText Vimeo_Square;
+
+	// http://fortawesome.github.io/Font-Awesome/icon/vine/
+	static FText Vine;
+
+	// http://fortawesome.github.io/Font-Awesome/icon/vk/
+	static FText Vk;
+
+	// (lower, quieter, sound, music) http://fortawesome.github.io/Font-Awesome/icon/volume-down/
+	static FText Volume_Down;
+
+	// (mute, sound, music) http://fortawesome.github.io/Font-Awesome/icon/volume-off/
+	static FText Volume_Off;
+
+	// (higher, louder, sound, music) http://fortawesome.github.io/Font-Awesome/icon/volume-up/
+	static FText Volume_Up;
+
+	// http://fortawesome.github.io/Font-Awesome/icon/weibo/
+	static FText Weibo;
+
+	// http://fortawesome.github.io/Font-Awesome/icon/weixin/
+	static FText Weixin;
+
+	// http://fortawesome.github.io/Font-Awesome/icon/whatsapp/
+	static FText Whatsapp;
+
+	// (handicap, person, accessibility, accessibile) http://fortawesome.github.io/Font-Awesome/icon/wheelchair/
+	static FText Wheelchair;
+
+	// http://fortawesome.github.io/Font-Awesome/icon/wifi/
+	static FText Wifi;
+
+	// http://fortawesome.github.io/Font-Awesome/icon/windows/
+	static FText Windows;
+
+	// http://fortawesome.github.io/Font-Awesome/icon/wordpress/
+	static FText Wordpress;
+
+	// (settings, fix, update) http://fortawesome.github.io/Font-Awesome/icon/wrench/
+	static FText Wrench;
+
+	// http://fortawesome.github.io/Font-Awesome/icon/xing/
+	static FText Xing;
+
+	// http://fortawesome.github.io/Font-Awesome/icon/xing-square/
+	static FText Xing_Square;
+
+	// http://fortawesome.github.io/Font-Awesome/icon/yahoo/
+	static FText Yahoo;
+
+	// http://fortawesome.github.io/Font-Awesome/icon/yelp/
+	static FText Yelp;
+
+	// (video, film) http://fortawesome.github.io/Font-Awesome/icon/youtube/
+	static FText Youtube;
+
+	// (start, playing) http://fortawesome.github.io/Font-Awesome/icon/youtube-play/
+	static FText Youtube_Play;
+
+	// (video, film) http://fortawesome.github.io/Font-Awesome/icon/youtube-square/
+	static FText Youtube_Square;
+};
\ No newline at end of file
diff --git a/Engine/Source/Editor/EditorStyle/Public/EditorStyle.h b/Engine/Source/Editor/EditorStyle/Public/EditorStyle.h
--- a/Engine/Source/Editor/EditorStyle/Public/EditorStyle.h
+++ b/Engine/Source/Editor/EditorStyle/Public/EditorStyle.h
@@ -22,3 +22,4 @@
  *****************************************************************************/
 
 #include "EditorStyleSettings.h"
+#include "EditorFontGlyphs.h"
\ No newline at end of file
