Misc changes needed for VS2015
- no ctor available in an FPawnActionStack struct.
- Shadow variable changes, including header changes.
- uint vs int narrowing issue fixed.

diff --git a/Engine/Source/Developer/StandaloneRenderer/Private/Windows/D3D/SlateD3DTextures.cpp b/Engine/Source/Developer/StandaloneRenderer/Private/Windows/D3D/SlateD3DTextures.cpp
--- a/Engine/Source/Developer/StandaloneRenderer/Private/Windows/D3D/SlateD3DTextures.cpp
+++ b/Engine/Source/Developer/StandaloneRenderer/Private/Windows/D3D/SlateD3DTextures.cpp
@@ -49,7 +49,7 @@ void FSlateD3DTexture::Init( DXGI_FORMAT InFormat, D3D11_SUBRESOURCE_DATA* Inita
 		TexDesc.BindFlags = 0;
 		TexDesc.CPUAccessFlags = D3D11_CPU_ACCESS_WRITE | D3D11_CPU_ACCESS_READ;
 
-		HRESULT Hr = GD3DDevice->CreateTexture2D(&TexDesc, InitalData, StagingTexture.GetInitReference());
+		Hr = GD3DDevice->CreateTexture2D(&TexDesc, InitalData, StagingTexture.GetInitReference());
 		checkf(SUCCEEDED(Hr), TEXT("D3D11 Error Result %X"), Hr);
 	}
 	else
diff --git a/Engine/Source/Editor/EditorStyle/Private/EditorStyleClasses.cpp b/Engine/Source/Editor/EditorStyle/Private/EditorStyleClasses.cpp
--- a/Engine/Source/Editor/EditorStyle/Private/EditorStyleClasses.cpp
+++ b/Engine/Source/Editor/EditorStyle/Private/EditorStyleClasses.cpp
@@ -33,9 +33,9 @@ void UEditorStyleSettings::PostEditChangeProperty(struct FPropertyChangedEvent&
 {
 	Super::PostEditChangeProperty(PropertyChangedEvent);
 
-	const FName Name = (PropertyChangedEvent.Property != nullptr) ? PropertyChangedEvent.Property->GetFName() : NAME_None;
+	const FName PropertyName = (PropertyChangedEvent.Property != nullptr) ? PropertyChangedEvent.Property->GetFName() : NAME_None;
 
-	if (Name == GET_MEMBER_NAME_CHECKED(UEditorStyleSettings, bEnableWindowAnimations))
+	if (PropertyName == GET_MEMBER_NAME_CHECKED(UEditorStyleSettings, bEnableWindowAnimations))
 	{
 		FSlateApplication::Get().EnableMenuAnimations(bEnableWindowAnimations);
 	}
@@ -45,7 +45,7 @@ void UEditorStyleSettings::PostEditChangeProperty(struct FPropertyChangedEvent&
 		SaveConfig();
 	}
 
-	SettingChangedEvent.Broadcast(Name);
+	SettingChangedEvent.Broadcast(PropertyName);
 }
 
 #endif
diff --git a/Engine/Source/Editor/UnrealEd/Classes/ActorFactories/ActorFactory.h b/Engine/Source/Editor/UnrealEd/Classes/ActorFactories/ActorFactory.h
--- a/Engine/Source/Editor/UnrealEd/Classes/ActorFactories/ActorFactory.h
+++ b/Engine/Source/Editor/UnrealEd/Classes/ActorFactories/ActorFactory.h
@@ -41,7 +41,7 @@ class UNREALED_API UActorFactory : public UObject
 	FVector SpawnPositionOffset;
 
 	/** Called to actual create an actor with the supplied transform (scale is ignored), using the properties in the ActorFactory */
-	AActor* CreateActor( UObject* Asset, ULevel* InLevel, FTransform Transform, EObjectFlags ObjectFlags = RF_Transactional, const FName Name = NAME_None );
+	AActor* CreateActor( UObject* Asset, ULevel* InLevel, FTransform Transform, EObjectFlags InObjectFlags = RF_Transactional, const FName InName = NAME_None );
 
 	/** Called to create a blueprint class that can be used to spawn an actor from this factory */
 	UBlueprint* CreateBlueprint( UObject* Instance, UObject* Outer, const FName Name, const FName CallingContext = NAME_None );
diff --git a/Engine/Source/Runtime/AIModule/Classes/Actions/PawnActionsComponent.h b/Engine/Source/Runtime/AIModule/Classes/Actions/PawnActionsComponent.h
--- a/Engine/Source/Runtime/AIModule/Classes/Actions/PawnActionsComponent.h
+++ b/Engine/Source/Runtime/AIModule/Classes/Actions/PawnActionsComponent.h
@@ -39,6 +39,10 @@ struct AIMODULE_API FPawnActionStack
 {
 	GENERATED_USTRUCT_BODY()
 
+	FPawnActionStack()
+		: TopAction(nullptr)
+	{}
+
 private:
 	UPROPERTY()
 	UPawnAction* TopAction;
diff --git a/Engine/Source/Runtime/AIModule/Classes/EnvironmentQuery/EnvQueryTypes.h b/Engine/Source/Runtime/AIModule/Classes/EnvironmentQuery/EnvQueryTypes.h
--- a/Engine/Source/Runtime/AIModule/Classes/EnvironmentQuery/EnvQueryTypes.h
+++ b/Engine/Source/Runtime/AIModule/Classes/EnvironmentQuery/EnvQueryTypes.h
@@ -379,7 +379,7 @@ struct AIMODULE_API FEnvQueryItem
 	int32 DataOffset:31;
 
 	/** has this item been discarded? */
-	int32 bIsDiscarded:1;
+	uint32 bIsDiscarded:1;
 
 	FORCEINLINE bool IsValid() const { return DataOffset >= 0 && !bIsDiscarded; }
 	FORCEINLINE void Discard() { bIsDiscarded = true; }
diff --git a/Engine/Source/Runtime/AIModule/Classes/Navigation/PathFollowingComponent.h b/Engine/Source/Runtime/AIModule/Classes/Navigation/PathFollowingComponent.h
--- a/Engine/Source/Runtime/AIModule/Classes/Navigation/PathFollowingComponent.h
+++ b/Engine/Source/Runtime/AIModule/Classes/Navigation/PathFollowingComponent.h
@@ -299,8 +299,8 @@ class AIMODULE_API UPathFollowingComponent : public UActorComponent, public IAIR
 	void OnPathEvent(FNavigationPath* InvalidatedPath, ENavPathEvent::Type Event);
 
 	/** helper function for sending a path for visual log */
-	static void LogPathHelper(const AActor* LogOwner, FNavPathSharedPtr LogPath, const AActor* LogGoalActor);
-	static void LogPathHelper(const AActor* LogOwner, FNavigationPath* LogPath, const AActor* LogGoalActor);
+	static void LogPathHelper(const AActor* LogOwner, FNavPathSharedPtr InLogPath, const AActor* LogGoalActor);
+	static void LogPathHelper(const AActor* LogOwner, FNavigationPath* InLogPath, const AActor* LogGoalActor);
 
 protected:
 
