Hooked up the new source control settings icon

diff --git a/Engine/Source/Editor/EditorStyle/Private/SlateEditorStyle.cpp b/Engine/Source/Editor/EditorStyle/Private/SlateEditorStyle.cpp
--- a/Engine/Source/Editor/EditorStyle/Private/SlateEditorStyle.cpp
+++ b/Engine/Source/Editor/EditorStyle/Private/SlateEditorStyle.cpp
@@ -6268,6 +6268,7 @@ void FSlateEditorStyle::FStyle::SetupSourceControlStyles()
 		Set( "SourceControl.Actions.CheckOut", new IMAGE_BRUSH( "ContentBrowser/SCC_CheckedOut", Icon16x16 ) );
 		Set( "SourceControl.Actions.Add", new IMAGE_BRUSH( "ContentBrowser/SCC_ContentAdd", Icon16x16 ) );
 		Set( "SourceControl.Actions.Refresh", new IMAGE_BRUSH( "Icons/icon_Refresh_16x", Icon16x16 ) );
+		Set( "SourceControl.Actions.ChangeSettings", new IMAGE_BRUSH( "Icons/icon_SCC_Change_Source_Control_Settings_16x", Icon16x16 ) );
 	}
 #endif // WITH_EDITOR || IS_PROGRAM
 
diff --git a/Engine/Source/Editor/LevelEditor/Private/LevelEditorToolBar.cpp b/Engine/Source/Editor/LevelEditor/Private/LevelEditorToolBar.cpp
--- a/Engine/Source/Editor/LevelEditor/Private/LevelEditorToolBar.cpp
+++ b/Engine/Source/Editor/LevelEditor/Private/LevelEditorToolBar.cpp
@@ -1806,7 +1806,7 @@ TSharedRef< SWidget > FLevelEditorToolBar::GenerateSourceControlMenu(TSharedRef<
 			NAME_None,
 			TAttribute<FText>(),
 			TAttribute<FText>(),
-			FSlateIcon(FEditorStyle::GetStyleSetName(), "SourceControl.StatusIcon.Unknown")
+			FSlateIcon(FEditorStyle::GetStyleSetName(), "SourceControl.Actions.ChangeSettings")
 			);
 	}
 	else
