Profiler: Made shadow color for display text more opaque to increase readability
[INTEGRATE] Change 2671802 by [REDACTED_PERSON][REDACTED_HANDLE] on 2015/08/27 19:49:51

diff --git a/Engine/Source/Editor/EditorStyle/Private/SlateEditorStyle.cpp b/Engine/Source/Editor/EditorStyle/Private/SlateEditorStyle.cpp
--- a/Engine/Source/Editor/EditorStyle/Private/SlateEditorStyle.cpp
+++ b/Engine/Source/Editor/EditorStyle/Private/SlateEditorStyle.cpp
@@ -3325,28 +3325,28 @@ void FSlateEditorStyle::FStyle::SetupProfilerStyle()
 			.SetFont( TTF_CORE_FONT( "Fonts/Roboto-Bold", 10 ) )
 			.SetColorAndOpacity( FLinearColor::White )
 			.SetShadowOffset( FVector2D(1.0f, 1.0f) )
-			.SetShadowColorAndOpacity( FLinearColor(0.f,0.f,0.f,0.5f) )
+			.SetShadowColorAndOpacity( FLinearColor(0.f,0.f,0.f,0.8f) )
 		);
 
 		Set( "Profiler.Caption", FTextBlockStyle(NormalText)
 			.SetFont( TTF_CORE_FONT( "Fonts/Roboto-Regular", 10 ) )
 			.SetColorAndOpacity( FLinearColor::White )
 			.SetShadowOffset( FVector2D(1.0f, 1.0f) )
-			.SetShadowColorAndOpacity( FLinearColor(0.f,0.f,0.f,0.5f) )
+			.SetShadowColorAndOpacity( FLinearColor(0.f,0.f,0.f,0.8f) )
 		);
 
 		Set( "Profiler.TooltipBold", FTextBlockStyle(NormalText)
 			.SetFont( TTF_CORE_FONT( "Fonts/Roboto-Bold", 8 ) )
 			.SetColorAndOpacity( FLinearColor(0.5f,0.5f,0.5f,1.0f) )
 			.SetShadowOffset( FVector2D(1.0f, 1.0f) )
-			.SetShadowColorAndOpacity( FLinearColor(0.f,0.f,0.f,0.5f) )
+			.SetShadowColorAndOpacity( FLinearColor(0.f,0.f,0.f,0.8f) )
 		);
 
 		Set( "Profiler.Tooltip", FTextBlockStyle(NormalText)
 			.SetFont( TTF_CORE_FONT( "Fonts/Roboto-Regular", 8 ) )
 			.SetColorAndOpacity( FLinearColor::White )
 			.SetShadowOffset( FVector2D(1.0f, 1.0f) )
-			.SetShadowColorAndOpacity( FLinearColor(0.f,0.f,0.f,0.5f) )
+			.SetShadowColorAndOpacity( FLinearColor(0.f,0.f,0.f,0.8f) )
 		);
 
 		// Event graph icons
