Sequencer: Hook up new sequencer icons.

diff --git a/Engine/Plugins/MovieScene/ActorAnimationEditor/Source/ActorAnimationEditor/Private/ActorAnimationEditorToolkit.cpp b/Engine/Plugins/MovieScene/ActorAnimationEditor/Source/ActorAnimationEditor/Private/ActorAnimationEditorToolkit.cpp
--- a/Engine/Plugins/MovieScene/ActorAnimationEditor/Source/ActorAnimationEditor/Private/ActorAnimationEditorToolkit.cpp
+++ b/Engine/Plugins/MovieScene/ActorAnimationEditor/Source/ActorAnimationEditor/Private/ActorAnimationEditorToolkit.cpp
@@ -264,7 +264,7 @@ TSharedRef<SWidget> FActorAnimationEditorToolkit::HandleSequencerGetAddMenuConte
 	MenuBuilder.AddMenuEntry(
 		LOCTEXT("AddEventTrack", "Add Event Track"),
 		LOCTEXT("AddEventTooltip", "Adds a new event track that can trigger events on the timeline."),
-		FSlateIcon(),
+		FSlateIcon(Style->GetStyleSetName(), "ActorAnimationEditor.Tracks.Event"),
 		FUIAction(FExecuteAction::CreateRaw(this, &FActorAnimationEditorToolkit::HandleAddMenuAddEventTrackExecute))
 	);
 
@@ -272,7 +272,7 @@ TSharedRef<SWidget> FActorAnimationEditorToolkit::HandleSequencerGetAddMenuConte
 	MenuBuilder.AddMenuEntry(
 		LOCTEXT("AddShotTrack", "Add Shot Track"),
 		LOCTEXT("AddShotTooltip", "Adds a shot track, as well as a new shot at the current scrubber location if a camera is selected."),
-		FSlateIcon(),
+		FSlateIcon(Style->GetStyleSetName(), "ActorAnimationEditor.Tracks.Shot"),
 		FUIAction(FExecuteAction::CreateRaw(this, &FActorAnimationEditorToolkit::HandleAddMenuAddShotTrackExecute))
 	);
 
@@ -280,7 +280,7 @@ TSharedRef<SWidget> FActorAnimationEditorToolkit::HandleSequencerGetAddMenuConte
 	MenuBuilder.AddMenuEntry(
 		LOCTEXT("AddSlomoTrack", "Add Slomo Track"),
 		LOCTEXT("AddSlomoTooltip", "Adds a new slomo track that controls the playback speed of the sequence."),
-		FSlateIcon(),
+		FSlateIcon(Style->GetStyleSetName(), "ActorAnimationEditor.Tracks.Slomo"),
 		FUIAction(FExecuteAction::CreateRaw(this, &FActorAnimationEditorToolkit::HandleAddMenuAddSlomoTrackExecute))
 	);
 
diff --git a/Engine/Plugins/MovieScene/ActorAnimationEditor/Source/ActorAnimationEditor/Private/Styles/ActorAnimationEditorStyle.h b/Engine/Plugins/MovieScene/ActorAnimationEditor/Source/ActorAnimationEditor/Private/Styles/ActorAnimationEditorStyle.h
--- a/Engine/Plugins/MovieScene/ActorAnimationEditor/Source/ActorAnimationEditor/Private/Styles/ActorAnimationEditorStyle.h
+++ b/Engine/Plugins/MovieScene/ActorAnimationEditor/Source/ActorAnimationEditor/Private/Styles/ActorAnimationEditorStyle.h
@@ -30,6 +30,11 @@ public:
 
 		// tab icons
 		Set("ActorAnimationEditor.Tabs.Sequencer", new IMAGE_BRUSH("icon_tab_sequencer_16x", Icon16x16));
+
+		// track icons
+		Set("ActorAnimationEditor.Tracks.Event", new IMAGE_BRUSH("icon_Sequencer_Event_16x", Icon16x16));
+		Set("ActorAnimationEditor.Tracks.Shot", new IMAGE_BRUSH("icon_Sequencer_Shot_16x", Icon16x16));
+		Set("ActorAnimationEditor.Tracks.Slomo", new IMAGE_BRUSH("icon_Sequencer_Slomo_16x", Icon16x16));
 		
 		FSlateStyleRegistry::RegisterSlateStyle(*this);
 	 }
diff --git a/Engine/Source/Editor/EditorStyle/Private/SlateEditorStyle.cpp b/Engine/Source/Editor/EditorStyle/Private/SlateEditorStyle.cpp
--- a/Engine/Source/Editor/EditorStyle/Private/SlateEditorStyle.cpp
+++ b/Engine/Source/Editor/EditorStyle/Private/SlateEditorStyle.cpp
@@ -1584,8 +1584,8 @@ void FSlateEditorStyle::FStyle::SetupGeneralStyles()
 		Set( "Sequencer.Timeline.RangeHandle", new BOX_BRUSH( "Sequencer/GenericSectionBackground", FMargin(5.f/16.f) ) );
 		Set( "Sequencer.Timeline.NotifyAlignmentMarker", new IMAGE_BRUSH( "Sequencer/NotifyAlignmentMarker", FVector2D(10,19) ) );
 		Set( "Sequencer.NotificationImage_AddedPlayMovieSceneEvent", new IMAGE_BRUSH( "Old/Checkbox_checked", Icon16x16 ) );
-		Set( "Sequencer.Save", new IMAGE_BRUSH( "Icons/icon_file_save_40x", Icon40x40 ) );
-		Set( "Sequencer.Save.Small", new IMAGE_BRUSH( "Icons/icon_file_save_16px", Icon20x20 ) );
+		Set( "Sequencer.Save", new IMAGE_BRUSH( "Icons/SequencerIcons/icon_Sequencer_Save_24x", Icon40x40 ) );
+		Set( "Sequencer.Save.Small", new IMAGE_BRUSH( "Icons/SequencerIcons/icon_Sequencer_Save_24x", Icon20x20 ) );
 		Set( "Sequencer.Key", new IMAGE_BRUSH( "Common/Key", FVector2D(7.0f,12.0f) ) );
 		Set( "Sequencer.PartialKey", new IMAGE_BRUSH( "Sequencer/PartialKey", FVector2D(11.f, 11.f) ) );
 		Set( "Sequencer.GenericSection.Background", new BOX_BRUSH( "Sequencer/GenericSectionBackground", FMargin(5.f/16.f) ) );
@@ -1598,25 +1598,33 @@ void FSlateEditorStyle::FStyle::SetupGeneralStyles()
 		Set( "Sequencer.AnimationOutliner.RegularFont", TTF_CORE_FONT( "Fonts/Roboto-Regular", 9 ) );
 		Set( "Sequencer.ShotFilter", new IMAGE_BRUSH( "Sequencer/FilteredArea", FVector2D(74,74), FLinearColor::White, ESlateBrushTileType::Both ) );
 		Set( "Sequencer.KeyMark", new IMAGE_BRUSH("Sequencer/KeyMark", FVector2D(3,21), FLinearColor::White, ESlateBrushTileType::NoTile ) );
-		Set( "Sequencer.AddObject", new IMAGE_BRUSH( "Icons/icon_Sequencer_AddObject_40x", Icon40x40 ) );
-		Set( "Sequencer.AddObject.Small", new IMAGE_BRUSH( "Icons/icon_Sequencer_AddObject_40x", Icon20x20 ) );
-		Set( "Sequencer.ToggleAutoKeyEnabled", new IMAGE_BRUSH( "Icons/icon_Sequencer_AutoKey_40x", Icon40x40 ) );
-		Set( "Sequencer.ToggleAutoKeyEnabled.Small", new IMAGE_BRUSH( "Icons/icon_Sequencer_AutoKey_40x", Icon20x20 ) );
-		Set( "Sequencer.ToggleKeyAllEnabled", new IMAGE_BRUSH( "Icons/icon_Sequencer_KeyAll_40x", Icon40x40 ) );
-		Set( "Sequencer.ToggleKeyAllEnabled.Small", new IMAGE_BRUSH( "Icons/icon_Sequencer_KeyAll_40x", Icon20x20 ) );
-		Set( "Sequencer.ToggleIsSnapEnabled", new IMAGE_BRUSH( "Icons/icon_Sequencer_ToggleSnap_40x", Icon40x40 ) );
-		Set( "Sequencer.ToggleIsSnapEnabled.Small", new IMAGE_BRUSH( "Icons/icon_Sequencer_ToggleSnap_40x", Icon20x20 ) );
-		Set( "Sequencer.ToggleShowCurveEditor", new IMAGE_BRUSH("Icons/icon_MatineeCurveView_40px", Icon40x40) );
-		Set( "Sequencer.ToggleShowCurveEditor.Small", new IMAGE_BRUSH("Icons/icon_MatineeCurveView_40px", Icon20x20) );
+		Set( "Sequencer.AddObject", new IMAGE_BRUSH( "Icons/SequencerIcons/icon_Sequencer_NewTrack_24x", Icon40x40 ) );
+		Set( "Sequencer.AddObject.Small", new IMAGE_BRUSH( "Icons/SequencerIcons/icon_Sequencer_NewTrack_24x", Icon20x20 ) );
+		Set( "Sequencer.ToggleAutoKeyEnabled", new IMAGE_BRUSH( "Icons/SequencerIcons/icon_Sequencer_AutoKey_24x", Icon40x40 ) );
+		Set( "Sequencer.ToggleAutoKeyEnabled.Small", new IMAGE_BRUSH( "Icons/SequencerIcons/icon_Sequencer_AutoKey_24x", Icon20x20 ) );
+		Set( "Sequencer.KeyAllEnabled", new IMAGE_BRUSH( "Icons/SequencerIcons/icon_Sequencer_WholeKeyframing_24x", Icon40x40 ) );
+		Set( "Sequencer.KeyAllEnabled.Small", new IMAGE_BRUSH( "Icons/SequencerIcons/icon_Sequencer_WholeKeyframing_24x", Icon20x20 ) );
+		Set( "Sequencer.KeyAllDisabled", new IMAGE_BRUSH( "Icons/SequencerIcons/icon_Sequencer_PartialKeyframing_24x", Icon40x40 ) );
+		Set( "Sequencer.KeyAllDisabled.Small", new IMAGE_BRUSH( "Icons/SequencerIcons/icon_Sequencer_PartialKeyframing_24x", Icon20x20 ) );
+		Set( "Sequencer.ToggleIsSnapEnabled", new IMAGE_BRUSH( "Icons/SequencerIcons/icon_Sequencer_Snap_24x", Icon40x40 ) );
+		Set( "Sequencer.ToggleIsSnapEnabled.Small", new IMAGE_BRUSH( "Icons/SequencerIcons/icon_Sequencer_Snap_24x", Icon20x20 ) );
+		Set( "Sequencer.ToggleShowCurveEditor", new IMAGE_BRUSH("Icons/SequencerIcons/icon_Sequencer_CurveEditor_24x", Icon40x40) );
+		Set( "Sequencer.ToggleShowCurveEditor.Small", new IMAGE_BRUSH("Icons/SequencerIcons/icon_Sequencer_CurveEditor_24x", Icon20x20) );
 		Set( "Sequencer.ToggleAutoScroll", new IMAGE_BRUSH( "Icons/icon_Sequencer_ToggleAutoScroll_40x", Icon40x40 ) );
 		Set( "Sequencer.ToggleAutoScroll.Small", new IMAGE_BRUSH( "Icons/icon_Sequencer_ToggleAutoScroll_16x", Icon16x16 ) );
-		Set( "Sequencer.MoveTool.Small", new IMAGE_BRUSH( "Sequencer/MoveTool_16px", Icon16x16 ) );
-		Set( "Sequencer.MarqueeTool_Keys.Small", new IMAGE_BRUSH( "Sequencer/MarqueeTool_Keys_16px", Icon16x16 ) );
-		Set( "Sequencer.MarqueeTool_Sections.Small", new IMAGE_BRUSH( "Sequencer/MarqueeTool_Sections_16px", Icon16x16 ) );
+		Set( "Sequencer.MoveTool.Small", new IMAGE_BRUSH( "Icons/SequencerIcons/icon_Sequencer_Move_24x", Icon16x16 ) );
+		Set( "Sequencer.MarqueeTool_Keys.Small", new IMAGE_BRUSH( "Sequencer/MarqueeTool_Keys_16px", Icon16x16 ) ); //@todo
+		Set( "Sequencer.MarqueeTool_Sections.Small", new IMAGE_BRUSH( "Sequencer/MarqueeTool_Sections_16px", Icon16x16 ) ); //@todo
 		Set( "Sequencer.LockCamera", new IMAGE_BRUSH( "Icons/icon_ToggleActorPilotCameraView_16x", Icon16x16 ) );
 		Set( "Sequencer.UnlockCamera", new IMAGE_BRUSH( "Icons/icon_ToggleActorPilotCameraView_16x", Icon16x16, FLinearColor(1.f, 1.f, 1.f, 0.5f) ) );
 		Set( "Sequencer.Section.SelectionBorder", new BORDER_BRUSH( TEXT("Sequencer/SectionHighlight"), FMargin(7.f/16.f) ) );
 		Set( "Sequencer.ShotTrack.SectionHandle", new IMAGE_BRUSH( "Old/White", Icon16x16, FLinearColor::Black ) );
+		Set( "Sequencer.General", new IMAGE_BRUSH( "Icons/SequencerIcons/icon_Sequencer_View_24x", Icon24x24 ) );
+		Set( "Sequencer.General.Small", new IMAGE_BRUSH( "Icons/SequencerIcons/icon_Sequencer_View_24x", Icon20x20 ) );
+		Set( "Sequencer.Time", new IMAGE_BRUSH( "Icons/SequencerIcons/icon_Sequencer_Time_24x", Icon24x24 ) );
+		Set( "Sequencer.Time.Small", new IMAGE_BRUSH( "Icons/SequencerIcons/icon_Sequencer_Time_24x", Icon20x20 ) );
+		Set( "Sequencer.Value", new IMAGE_BRUSH( "Icons/SequencerIcons/icon_Sequencer_Value_24x", Icon24x24 ) );
+		Set( "Sequencer.Value.Small", new IMAGE_BRUSH( "Icons/SequencerIcons/icon_Sequencer_Value_24x", Icon20x20 ) );
 
 		const FButtonStyle DetailsKeyButton = FButtonStyle(NoBorder)
 			.SetNormal( IMAGE_BRUSH("Sequencer/AddKey_Details", FVector2D(11,11) )  )
diff --git a/Engine/Source/Editor/Sequencer/Private/SSequencer.cpp b/Engine/Source/Editor/Sequencer/Private/SSequencer.cpp
--- a/Engine/Source/Editor/Sequencer/Private/SSequencer.cpp
+++ b/Engine/Source/Editor/Sequencer/Private/SSequencer.cpp
@@ -633,16 +633,17 @@ TSharedRef<SWidget> SSequencer::MakeToolBar()
 				, NAME_None
 				, LOCTEXT("SaveDirtyPackages", "Save")
 				, LOCTEXT("SaveDirtyPackagesTooltip", "Saves the current movie scene")
-				//, FSlateIcon(FEditorStyle::GetStyleSetName(), "LevelEditor.Save"));
 				, FSlateIcon(FEditorStyle::GetStyleSetName(), "Sequencer.Save"));
+
+			ToolBarBuilder.AddSeparator();
 		}
+
 		ToolBarBuilder.AddComboButton(
-			FUIAction(),
-			FOnGetContent::CreateSP( this, &SSequencer::MakeGeneralMenu ),
-			LOCTEXT( "GeneralOptions", "General Options" ),
-			LOCTEXT( "GeneralOptionsToolTip", "General Options" ),
-			TAttribute<FSlateIcon>(),
-			true );
+			FUIAction()
+			, FOnGetContent::CreateSP( this, &SSequencer::MakeGeneralMenu )
+			, LOCTEXT( "GeneralOptions", "General Options" )
+			, LOCTEXT( "GeneralOptionsToolTip", "General Options" )
+			, FSlateIcon(FEditorStyle::GetStyleSetName(), "Sequencer.General") );
 
 		ToolBarBuilder.AddSeparator();
 
@@ -655,7 +656,14 @@ TSharedRef<SWidget> SSequencer::MakeToolBar()
 				, LOCTEXT("AddObjectTooltip", "Add selected objects")
 				, FSlateIcon(FEditorStyle::GetStyleSetName(), "Sequencer.AddObject"));
 
-			ToolBarBuilder.AddToolBarButton( FSequencerCommands::Get().ToggleKeyAllEnabled );
+			TAttribute<FSlateIcon> KeyAllIcon;
+			KeyAllIcon.Bind(TAttribute<FSlateIcon>::FGetter::CreateLambda([&]{
+				static FSlateIcon KeyAllEnabledIcon(FEditorStyle::GetStyleSetName(), "Sequencer.KeyAllEnabled");
+				static FSlateIcon KeyAllDisabledIcon(FEditorStyle::GetStyleSetName(), "Sequencer.KeyAllDisabled");
+
+				return Sequencer.Pin()->GetKeyAllEnabled() ? KeyAllEnabledIcon : KeyAllDisabledIcon;
+			}));
+			ToolBarBuilder.AddToolBarButton( FSequencerCommands::Get().ToggleKeyAllEnabled, NAME_None, TAttribute<FText>(), TAttribute<FText>(), KeyAllIcon );
 		}
 
 		ToolBarBuilder.AddToolBarButton( FSequencerCommands::Get().ToggleAutoKeyEnabled );
@@ -710,6 +718,7 @@ TSharedRef<SWidget> SSequencer::MakeToolBar()
 	ToolBarBuilder.BeginSection("Snapping");
 	{
 		ToolBarBuilder.AddToolBarButton( FSequencerCommands::Get().ToggleIsSnapEnabled, NAME_None, TAttribute<FText>( FText::GetEmpty() ) );
+
 		ToolBarBuilder.AddComboButton(
 			FUIAction(),
 			FOnGetContent::CreateSP( this, &SSequencer::MakeSnapMenu ),
@@ -717,25 +726,35 @@ TSharedRef<SWidget> SSequencer::MakeToolBar()
 			LOCTEXT( "SnapOptionsToolTip", "Snapping Options" ),
 			TAttribute<FSlateIcon>(),
 			true );
+
+		ToolBarBuilder.AddSeparator();
+
+		ToolBarBuilder.AddWidget(
+			SNew( SImage )
+			.Image(FEditorStyle::GetBrush("Sequencer.Time")) );
+
 		ToolBarBuilder.AddWidget(
 			SNew( SBox )
 			.VAlign( VAlign_Center )
 			[
 				SNew( SNumericDropDown<float> )
 				.DropDownValues( SequencerSnapValues::GetTimeSnapValues() )
-				.LabelText( LOCTEXT( "TimeSnapLabel", "Time" ) )
 				.bShowNamedValue(true)
 				.ToolTipText( LOCTEXT( "TimeSnappingIntervalToolTip", "Time snapping interval" ) )
 				.Value( this, &SSequencer::OnGetTimeSnapInterval )
 				.OnValueChanged( this, &SSequencer::OnTimeSnapIntervalChanged )
 			] );
+
+		ToolBarBuilder.AddWidget(
+			SNew( SImage )
+			.Image(FEditorStyle::GetBrush("Sequencer.Value")) );
+
 		ToolBarBuilder.AddWidget(
 			SNew( SBox )
 			.VAlign( VAlign_Center )
 			[
 				SNew( SNumericDropDown<float> )
 				.DropDownValues( SequencerSnapValues::GetSnapValues() )
-				.LabelText( LOCTEXT( "ValueSnapLabel", "Value" ) )
 				.ToolTipText( LOCTEXT( "ValueSnappingIntervalToolTip", "Curve value snapping interval" ) )
 				.Value( this, &SSequencer::OnGetValueSnapInterval )
 				.OnValueChanged( this, &SSequencer::OnValueSnapIntervalChanged )
