Sequencer: Move icons for tracks back to SequencerIcons.

diff --git a/Engine/Plugins/MovieScene/ActorAnimationEditor/Source/ActorAnimationEditor/Private/Styles/ActorAnimationEditorStyle.h b/Engine/Plugins/MovieScene/ActorAnimationEditor/Source/ActorAnimationEditor/Private/Styles/ActorAnimationEditorStyle.h
--- a/Engine/Plugins/MovieScene/ActorAnimationEditor/Source/ActorAnimationEditor/Private/Styles/ActorAnimationEditorStyle.h
+++ b/Engine/Plugins/MovieScene/ActorAnimationEditor/Source/ActorAnimationEditor/Private/Styles/ActorAnimationEditorStyle.h
@@ -30,12 +30,7 @@ public:
 
 		// tab icons
 		Set("ActorAnimationEditor.Tabs.Sequencer", new IMAGE_BRUSH("icon_tab_sequencer_16x", Icon16x16));
-
-		// track icons
-		Set("ActorAnimationEditor.Tracks.Event", new IMAGE_BRUSH("icon_Sequencer_Event_16x", Icon16x16));
-		Set("ActorAnimationEditor.Tracks.Shot", new IMAGE_BRUSH("icon_Sequencer_Shot_16x", Icon16x16));
-		Set("ActorAnimationEditor.Tracks.Slomo", new IMAGE_BRUSH("icon_Sequencer_Slomo_16x", Icon16x16));
-		
+	
 		FSlateStyleRegistry::RegisterSlateStyle(*this);
 	 }
 
diff --git a/Engine/Source/Editor/EditorStyle/Private/SlateEditorStyle.cpp b/Engine/Source/Editor/EditorStyle/Private/SlateEditorStyle.cpp
--- a/Engine/Source/Editor/EditorStyle/Private/SlateEditorStyle.cpp
+++ b/Engine/Source/Editor/EditorStyle/Private/SlateEditorStyle.cpp
@@ -1629,6 +1629,9 @@ void FSlateEditorStyle::FStyle::SetupGeneralStyles()
 		Set( "Sequencer.Time.Small", new IMAGE_BRUSH( "Icons/SequencerIcons/icon_Sequencer_Time_24x", Icon20x20 ) );
 		Set( "Sequencer.Value", new IMAGE_BRUSH( "Icons/SequencerIcons/icon_Sequencer_Value_24x", Icon24x24 ) );
 		Set( "Sequencer.Value.Small", new IMAGE_BRUSH( "Icons/SequencerIcons/icon_Sequencer_Value_24x", Icon20x20 ) );
+		Set( "Sequencer.Tracks.Event", new IMAGE_BRUSH("Icons/SequencerIcons/icon_Sequencer_Event_16x", Icon16x16));
+		Set( "Sequencer.Tracks.Shot", new IMAGE_BRUSH("Icons/SequencerIcons/icon_Sequencer_Shot_16x", Icon16x16));
+		Set( "Sequencer.Tracks.Slomo", new IMAGE_BRUSH("Icons/SequencerIcons/icon_Sequencer_Slomo_16x", Icon16x16));
 
 		const FButtonStyle DetailsKeyButton = FButtonStyle(NoBorder)
 			.SetNormal( IMAGE_BRUSH("Sequencer/AddKey_Details", FVector2D(11,11) )  )
diff --git a/Engine/Source/Editor/MovieSceneTools/Private/TrackEditors/EventTrackEditor.cpp b/Engine/Source/Editor/MovieSceneTools/Private/TrackEditors/EventTrackEditor.cpp
--- a/Engine/Source/Editor/MovieSceneTools/Private/TrackEditors/EventTrackEditor.cpp
+++ b/Engine/Source/Editor/MovieSceneTools/Private/TrackEditors/EventTrackEditor.cpp
@@ -47,7 +47,7 @@ void FEventTrackEditor::BuildAddTrackMenu(FMenuBuilder& MenuBuilder)
 	MenuBuilder.AddMenuEntry(
 		LOCTEXT("AddEventTrack", "Add Event Track"),
 		LOCTEXT("AddEventTooltip", "Adds a new event track that can trigger events on the timeline."),
-		FSlateIcon(FEditorStyle::GetStyleSetName(), "ActorAnimationEditor.Tracks.Event"),
+		FSlateIcon(FEditorStyle::GetStyleSetName(), "Sequencer.Tracks.Event"),
 		FUIAction(
 			FExecuteAction::CreateRaw(this, &FEventTrackEditor::HandleAddEventTrackMenuEntryExecute)
 		)
diff --git a/Engine/Source/Editor/MovieSceneTools/Private/TrackEditors/ShotTrackEditor.cpp b/Engine/Source/Editor/MovieSceneTools/Private/TrackEditors/ShotTrackEditor.cpp
--- a/Engine/Source/Editor/MovieSceneTools/Private/TrackEditors/ShotTrackEditor.cpp
+++ b/Engine/Source/Editor/MovieSceneTools/Private/TrackEditors/ShotTrackEditor.cpp
@@ -570,7 +570,7 @@ void FShotTrackEditor::BuildAddTrackMenu(FMenuBuilder& MenuBuilder)
 	MenuBuilder.AddMenuEntry(
 		LOCTEXT("AddShotTrack", "Add Shot Track"),
 		LOCTEXT("AddShotTooltip", "Adds a shot track, as well as a new shot at the current scrubber location if a camera is selected."),
-		FSlateIcon(FEditorStyle::GetStyleSetName(), "ActorAnimationEditor.Tracks.Shot"),
+		FSlateIcon(FEditorStyle::GetStyleSetName(), "Sequencer.Tracks.Shot"),
 		FUIAction(
 			FExecuteAction::CreateRaw(this, &FShotTrackEditor::HandleAddShotTrackMenuEntryExecute),
 			FCanExecuteAction::CreateRaw(this, &FShotTrackEditor::HandleAddShotTrackMenuEntryCanExecute)
diff --git a/Engine/Source/Editor/MovieSceneTools/Private/TrackEditors/SlomoTrackEditor.cpp b/Engine/Source/Editor/MovieSceneTools/Private/TrackEditors/SlomoTrackEditor.cpp
--- a/Engine/Source/Editor/MovieSceneTools/Private/TrackEditors/SlomoTrackEditor.cpp
+++ b/Engine/Source/Editor/MovieSceneTools/Private/TrackEditors/SlomoTrackEditor.cpp
@@ -41,7 +41,7 @@ void FSlomoTrackEditor::BuildAddTrackMenu(FMenuBuilder& MenuBuilder)
 	MenuBuilder.AddMenuEntry(
 		LOCTEXT("AddPlayRateTrack", "Add Play Rate Track"),
 		LOCTEXT("AddPlayRateTrackTooltip", "Adds a new track that controls the playback rate of the sequence."),
-		FSlateIcon(FEditorStyle::GetStyleSetName(), "ActorAnimationEditor.Tracks.Slomo"),
+		FSlateIcon(FEditorStyle::GetStyleSetName(), "Sequencer.Tracks.Slomo"),
 		FUIAction(
 			FExecuteAction::CreateRaw(this, &FSlomoTrackEditor::HandleAddSlomoTrackMenuEntryExecute)
 		)
