Added icons for actor animation assets

diff --git a/Engine/Source/Editor/EditorStyle/Private/SlateEditorStyle.cpp b/Engine/Source/Editor/EditorStyle/Private/SlateEditorStyle.cpp
--- a/Engine/Source/Editor/EditorStyle/Private/SlateEditorStyle.cpp
+++ b/Engine/Source/Editor/EditorStyle/Private/SlateEditorStyle.cpp
@@ -5279,6 +5279,7 @@ void FSlateEditorStyle::FStyle::SetupClassIconsAndThumbnails()
 
 			TEXT("Actor"),
 			TEXT("ActorComponent"),
+			TEXT("ActorAnimation"),
 			TEXT("AIController"),
 			TEXT("AimOffsetBlendSpace"),
 			TEXT("AimOffsetBlendSpace1D"),
