Content Browser: Use larger font size when running with "-VREditorDemo"

diff --git a/Engine/Source/Editor/EditorStyle/Private/SlateEditorStyle.cpp b/Engine/Source/Editor/EditorStyle/Private/SlateEditorStyle.cpp
--- a/Engine/Source/Editor/EditorStyle/Private/SlateEditorStyle.cpp
+++ b/Engine/Source/Editor/EditorStyle/Private/SlateEditorStyle.cpp
@@ -5574,8 +5574,19 @@ void FSlateEditorStyle::FStyle::SetupContentBrowserStyle()
 
 		// Sources View
 		Set( "ContentBrowser.SourceTitleFont", TTF_CORE_FONT( "Fonts/Roboto-Regular", 12 ) );
-		Set( "ContentBrowser.SourceListItemFont", TTF_CORE_FONT( "Fonts/Roboto-Regular", 10 ) );
-		Set( "ContentBrowser.SourceTreeItemFont", TTF_CORE_FONT( "Fonts/Roboto-Regular", 10 ) );
+
+		// @todo vreditor urgent: Increase the size of Content Browser fonts while in VR
+		const bool bIsVREditorDemo = FParse::Param( FCommandLine::Get(), TEXT( "VREditorDemo" ) );	// @todo vreditor: Remove this when no longer needed
+		if( GEnableVREditorHacks && bIsVREditorDemo )
+		{
+			Set( "ContentBrowser.SourceListItemFont", TTF_CORE_FONT( "Fonts/Roboto-Regular", 20 ) );
+			Set( "ContentBrowser.SourceTreeItemFont", TTF_CORE_FONT( "Fonts/Roboto-Regular", 20 ) );
+		}
+		else
+		{
+			Set( "ContentBrowser.SourceListItemFont", TTF_CORE_FONT( "Fonts/Roboto-Regular", 10 ) );
+			Set( "ContentBrowser.SourceTreeItemFont", TTF_CORE_FONT( "Fonts/Roboto-Regular", 10 ) );
+		}
 		Set( "ContentBrowser.SourceTreeRootItemFont", TTF_CORE_FONT( "Fonts/Roboto-Regular", 12 ) );
 		Set( "ContentBrowser.AssetTreeFolderClosed", new IMAGE_BRUSH( "Icons/FolderClosed", FVector2D(18, 16) ) );
 		Set( "ContentBrowser.BreadcrumbPathPickerFolder", new IMAGE_BRUSH( "Icons/FolderClosed", FVector2D(18, 16), FLinearColor::Gray ) );
