VR Editor: Only enable VR-specific HMD hacks when actually using VR mode

diff --git a/Engine/Source/Editor/EditorStyle/Private/SlateEditorStyle.cpp b/Engine/Source/Editor/EditorStyle/Private/SlateEditorStyle.cpp
--- a/Engine/Source/Editor/EditorStyle/Private/SlateEditorStyle.cpp
+++ b/Engine/Source/Editor/EditorStyle/Private/SlateEditorStyle.cpp
@@ -5676,7 +5676,7 @@ void FSlateEditorStyle::FStyle::SetupContentBrowserStyle()
 
 		// @todo vreditor urgent: Increase the size of Content Browser fonts while in VR
 		const bool bIsVREditorDemo = FParse::Param( FCommandLine::Get(), TEXT( "VREditorDemo" ) );	// @todo vreditor: Remove this when no longer needed
-		if( GEnableVREditorHacks && bIsVREditorDemo )
+		if( bIsVREditorDemo )
 		{
 			Set( "ContentBrowser.SourceListItemFont", TTF_CORE_FONT( "Fonts/Roboto-Regular", 20 ) );
 			Set( "ContentBrowser.SourceTreeItemFont", TTF_CORE_FONT( "Fonts/Roboto-Regular", 20 ) );
diff --git a/Engine/Source/Editor/VREditor/VREditorMode.cpp b/Engine/Source/Editor/VREditor/VREditorMode.cpp
--- a/Engine/Source/Editor/VREditor/VREditorMode.cpp
+++ b/Engine/Source/Editor/VREditor/VREditorMode.cpp
@@ -141,6 +141,9 @@ void FVREditorMode::Enter()
 	InputProcessor = MakeShareable(new FVREditorInputProcessor(this));
 	FSlateApplication::Get().SetInputPreProcessor(true, InputProcessor);
 
+	// @todo vreditor urgent: Turn on global editor hacks for VR Editor mode
+	GEnableVREditorHacks = true;
+
 	bIsFullyInitialized = false;
 	bWantsToExitMode = false;
 
@@ -571,6 +574,9 @@ void FVREditorMode::Exit()
 
 	// Call parent implementation
 	FEdMode::Exit();
+
+	// @todo vreditor urgent: Disable global editor hacks for VR Editor mode
+	GEnableVREditorHacks = false;
 }
 
 
diff --git a/Engine/Source/Runtime/Core/Private/Misc/Core.cpp b/Engine/Source/Runtime/Core/Private/Misc/Core.cpp
--- a/Engine/Source/Runtime/Core/Private/Misc/Core.cpp
+++ b/Engine/Source/Runtime/Core/Private/Misc/Core.cpp
@@ -257,7 +257,7 @@ bool					GIsAutomationTesting					= false;
 bool					GPumpingMessagesOutsideOfMainLoop = false;
 
 /** Enables various editor and HMD hacks that allow the experimental VR editor feature to work, perhaps at the expense of other systems */
-bool					GEnableVREditorHacks = true;
+bool					GEnableVREditorHacks = false;
 
 /** Total blueprint compile time.																			*/
 double GBlueprintCompileTime = 0.0;
