Copying //UE4/Dev-Sequencer to //UE4/Main

diff --git a/Engine/Source/Editor/EditorStyle/Private/SlateEditorStyle.cpp b/Engine/Source/Editor/EditorStyle/Private/SlateEditorStyle.cpp
--- a/Engine/Source/Editor/EditorStyle/Private/SlateEditorStyle.cpp
+++ b/Engine/Source/Editor/EditorStyle/Private/SlateEditorStyle.cpp
@@ -948,15 +948,65 @@ void FSlateEditorStyle::FStyle::SetupGeneralStyles()
 	
 	// Scrub control buttons
 	{
-		Set( "Animation.Pause", new IMAGE_BRUSH( "/Animation/pause", Icon24x24 ) );
-		Set( "Animation.Forward", new IMAGE_BRUSH( "/Animation/forward", Icon24x24 ) );
-		Set( "Animation.Forward_Step", new IMAGE_BRUSH( "/Animation/forward_step", Icon16x16 ) );
-		Set( "Animation.Forward_End", new IMAGE_BRUSH( "/Animation/forward_end", Icon16x16 ) );
-		Set( "Animation.Backward", new IMAGE_BRUSH( "/Animation/backward", Icon24x24 ) );
-		Set( "Animation.Backward_Step", new IMAGE_BRUSH( "/Animation/backward_step", Icon16x16 ) );
-		Set( "Animation.Backward_End", new IMAGE_BRUSH( "/Animation/backward_end", Icon16x16 ) );
-		Set( "Animation.Loop", new IMAGE_BRUSH( "/Animation/Loop", Icon24x24 ) );
-		Set( "Animation.Record", new IMAGE_BRUSH( "/Animation/Record", Icon24x24 ) );
+		Set( "Animation.Pause", FButtonStyle( Button )
+			.SetNormal(IMAGE_BRUSH( "/Sequencer/Transport_Bar/Pause_24x_OFF", Icon24x24 ))
+			.SetHovered(IMAGE_BRUSH( "/Sequencer/Transport_Bar/Pause_24x_OFF", Icon24x24 ))
+			.SetPressed(IMAGE_BRUSH( "/Sequencer/Transport_Bar/Pause_24x", Icon24x24 ))
+		);
+
+		Set( "Animation.Forward", FButtonStyle( Button )
+			.SetNormal(IMAGE_BRUSH( "/Sequencer/Transport_Bar/Play_24x_OFF", Icon24x24 ))
+			.SetHovered(IMAGE_BRUSH( "/Sequencer/Transport_Bar/Play_24x_OFF", Icon24x24 ))
+			.SetPressed(IMAGE_BRUSH( "/Sequencer/Transport_Bar/Play_24x", Icon24x24 ))
+		);
+
+		Set( "Animation.Forward_Step", FButtonStyle( Button )
+			.SetNormal(IMAGE_BRUSH( "/Sequencer/Transport_Bar/Step_Forward_24x_OFF", Icon24x24 ))
+			.SetHovered(IMAGE_BRUSH( "/Sequencer/Transport_Bar/Step_Forward_24x_OFF", Icon24x24 ))
+			.SetPressed(IMAGE_BRUSH( "/Sequencer/Transport_Bar/Step_Forward_24x", Icon24x24 ))
+		);
+
+		Set( "Animation.Forward_End", FButtonStyle( Button )
+			.SetNormal(IMAGE_BRUSH( "/Sequencer/Transport_Bar/Go_To_End_24x_OFF", Icon24x24 ))
+			.SetHovered(IMAGE_BRUSH( "/Sequencer/Transport_Bar/Go_To_End_24x_OFF", Icon24x24 ))
+			.SetPressed(IMAGE_BRUSH( "/Sequencer/Transport_Bar/Go_To_End_24x", Icon24x24 ))
+		);
+
+		Set( "Animation.Backward", FButtonStyle( Button )
+			.SetNormal(IMAGE_BRUSH( "/Sequencer/Transport_Bar/Backwards_24x_OFF", Icon24x24 ))
+			.SetHovered(IMAGE_BRUSH( "/Sequencer/Transport_Bar/Backwards_24x_OFF", Icon24x24 ))
+			.SetPressed(IMAGE_BRUSH( "/Sequencer/Transport_Bar/Backwards_24x", Icon24x24 ))
+		);
+
+		Set( "Animation.Backward_Step", FButtonStyle( Button )
+			.SetNormal(IMAGE_BRUSH( "/Sequencer/Transport_Bar/Step_Backwards_24x_OFF", Icon24x24 ))
+			.SetHovered(IMAGE_BRUSH( "/Sequencer/Transport_Bar/Step_Backwards_24x_OFF", Icon24x24 ))
+			.SetPressed(IMAGE_BRUSH( "/Sequencer/Transport_Bar/Step_Backwards_24x", Icon24x24 ))
+		);
+
+		Set( "Animation.Backward_End", FButtonStyle( Button )
+			.SetNormal(IMAGE_BRUSH( "/Sequencer/Transport_Bar/Go_To_Front_24x_OFF", Icon24x24 ))
+			.SetHovered(IMAGE_BRUSH( "/Sequencer/Transport_Bar/Go_To_Front_24x_OFF", Icon24x24 ))
+			.SetPressed(IMAGE_BRUSH( "/Sequencer/Transport_Bar/Go_To_Front_24x", Icon24x24 ))
+		);
+
+		Set( "Animation.Loop.Enabled", FButtonStyle( Button )
+			.SetNormal(IMAGE_BRUSH( "/Sequencer/Transport_Bar/Loop_24x_OFF", Icon24x24 ))
+			.SetHovered(IMAGE_BRUSH( "/Sequencer/Transport_Bar/Loop_24x_OFF", Icon24x24 ))
+			.SetPressed(IMAGE_BRUSH( "/Sequencer/Transport_Bar/Loop_24x", Icon24x24 ))
+		);
+
+		Set( "Animation.Loop.Disabled", FButtonStyle( Button )
+			.SetNormal(IMAGE_BRUSH( "/Sequencer/Transport_Bar/Loop_Toggle_24x_OFF", Icon24x24 ))
+			.SetHovered(IMAGE_BRUSH( "/Sequencer/Transport_Bar/Loop_Toggle_24x_OFF", Icon24x24 ))
+			.SetPressed(IMAGE_BRUSH( "/Sequencer/Transport_Bar/Loop_Toggle_24x", Icon24x24 ))
+		);
+
+		Set( "Animation.Record", FButtonStyle( Button )
+			.SetNormal(IMAGE_BRUSH( "/Sequencer/Transport_Bar/Record_24x_OFF", Icon24x24 ))
+			.SetHovered(IMAGE_BRUSH( "/Sequencer/Transport_Bar/Record_24x_OFF", Icon24x24 ))
+			.SetPressed(IMAGE_BRUSH( "/Sequencer/Transport_Bar/Record_24x", Icon24x24 ))
+		);
 	}
 
 	// Message Log
@@ -1284,6 +1334,7 @@ void FSlateEditorStyle::FStyle::SetupGeneralStyles()
 		Set( "SceneOutliner.NewFolderIcon", new IMAGE_BRUSH("Icons/icon_AddFolder_16x", Icon16x16 ) );
 		Set( "SceneOutliner.MoveToRoot", new IMAGE_BRUSH("Icons/icon_NoFolder_16x", Icon16x16 ) );
 		Set( "SceneOutliner.ChangedItemHighlight", new BOX_BRUSH( "Common/EditableTextSelectionBackground", FMargin(4.f/16.f) ) );
+		Set( "SceneOutliner.World", new IMAGE_BRUSH( "Icons/icon_world_16x", Icon16x16 ) );
 
 		// Selection color should still be orange to align with the editor viewport.
 		// But must also give the hint that the tree is no longer focused.
@@ -1598,8 +1649,8 @@ void FSlateEditorStyle::FStyle::SetupGeneralStyles()
 		Set( "Sequencer.KeyDiamond", new IMAGE_BRUSH( "Sequencer/KeyDiamond", Icon12x12 ) );
 		Set( "Sequencer.KeySquare", new IMAGE_BRUSH( "Sequencer/KeySquare", Icon12x12 ) );
 		Set( "Sequencer.KeyTriangle", new IMAGE_BRUSH( "Sequencer/KeyTriangle", Icon12x12 ) );
-		Set( "Sequencer.KeyLeft", new IMAGE_BRUSH( "Common/KeyLeft", FVector2D( 7.0f, 12.0f ) ) );
-		Set( "Sequencer.KeyRight", new IMAGE_BRUSH( "Common/KeyRight", FVector2D( 7.0f, 12.0f ) ) );
+		Set( "Sequencer.KeyLeft", new IMAGE_BRUSH( "Sequencer/KeyLeft", Icon12x12 ) );
+		Set( "Sequencer.KeyRight", new IMAGE_BRUSH( "Sequencer/KeyRight", Icon12x12 ) );
 		Set( "Sequencer.PartialKey", new IMAGE_BRUSH( "Sequencer/PartialKey", FVector2D(11.f, 11.f) ) );
 		Set( "Sequencer.Star", new IMAGE_BRUSH( "Sequencer/Star", Icon12x12 ) );
 		Set( "Sequencer.Empty", new IMAGE_BRUSH( "Sequencer/Empty", Icon12x12 ) );
@@ -1618,16 +1669,43 @@ void FSlateEditorStyle::FStyle::SetupGeneralStyles()
 		Set( "Sequencer.Timeline.PlayRange_Bottom_L", new BOX_BRUSH( "Sequencer/PlayRange_Bottom_L", FMargin(1.f, 0.5f, 0.f, 0.5f) ) );
 		Set( "Sequencer.Timeline.PlayRange_Bottom_R", new BOX_BRUSH( "Sequencer/PlayRange_Bottom_R", FMargin(0.f, 0.5f, 1.f, 0.5f) ) );
 
+		
+		Set( "Sequencer.Transport.JumpToPreviousKey", FButtonStyle()
+			.SetNormal ( IMAGE_BRUSH( "/Sequencer/Transport_Bar/Previous_Frame_OFF", Icon24x24 ) )
+			.SetPressed( IMAGE_BRUSH( "/Sequencer/Transport_Bar/Previous_Frame", Icon24x24 ) )
+			.SetHovered( IMAGE_BRUSH( "/Sequencer/Transport_Bar/Previous_Frame_OFF", Icon24x24 ) ));
+		Set( "Sequencer.Transport.JumpToNextKey", FButtonStyle()
+			.SetNormal ( IMAGE_BRUSH( "/Sequencer/Transport_Bar/Next_Frame_24x_OFF", Icon24x24 ) )
+			.SetPressed( IMAGE_BRUSH( "/Sequencer/Transport_Bar/Next_Frame_24x", Icon24x24 ) )
+			.SetHovered( IMAGE_BRUSH( "/Sequencer/Transport_Bar/Next_Frame_24x_OFF", Icon24x24 ) ) );
+		Set( "Sequencer.Transport.SetPlayStart", FButtonStyle()
+			.SetNormal ( IMAGE_BRUSH( "/Sequencer/Transport_Bar/Bracket_In_16x24_OFF", FVector2D(16,24) ) )
+			.SetPressed( IMAGE_BRUSH( "/Sequencer/Transport_Bar/Bracket_In_16x24", FVector2D(16,24) ) )
+			.SetHovered( IMAGE_BRUSH( "/Sequencer/Transport_Bar/Bracket_In_16x24_OFF", FVector2D(16,24) ) ) );
+		Set( "Sequencer.Transport.SetPlayEnd", FButtonStyle()
+			.SetNormal ( IMAGE_BRUSH( "/Sequencer/Transport_Bar/Bracket_Out_16x24_OFF", FVector2D(16,24) ) )
+			.SetPressed( IMAGE_BRUSH( "/Sequencer/Transport_Bar/Bracket_Out_16x24", FVector2D(16,24) ) )
+			.SetHovered( IMAGE_BRUSH( "/Sequencer/Transport_Bar/Bracket_Out_16x24_OFF", FVector2D(16,24) ) ) );
+
+		Set( "Sequencer.Transport.CloseButton", FButtonStyle()
+			.SetNormal ( IMAGE_BRUSH( "/Docking/CloseApp_Normal", Icon16x16 ) )
+			.SetPressed( IMAGE_BRUSH( "/Docking/CloseApp_Pressed", Icon16x16 ) )
+			.SetHovered( IMAGE_BRUSH( "/Docking/CloseApp_Hovered", Icon16x16 )));
+
 		Set( "Sequencer.NotificationImage_AddedPlayMovieSceneEvent", new IMAGE_BRUSH( "Old/Checkbox_checked", Icon16x16 ) );
 
-		Set( "Sequencer.Save", new IMAGE_BRUSH( "Icons/SequencerIcons/icon_Sequencer_Save_24x", Icon40x40 ) );
-		Set( "Sequencer.Save.Small", new IMAGE_BRUSH( "Icons/SequencerIcons/icon_Sequencer_Save_24x", Icon20x20 ) );
+		Set( "Sequencer.Save", new IMAGE_BRUSH( "Sequencer/Main_Icons/Icon_Sequencer_Save_24x", Icon40x40 ) );
+		Set( "Sequencer.Save.Small", new IMAGE_BRUSH( "Sequencer/Main_Icons/Icon_Sequencer_Save_24x", Icon20x20 ) );
+		Set( "Sequencer.SaveAs", new IMAGE_BRUSH( "Sequencer/Main_Icons/Icon_Sequencer_Save_As_24x", Icon40x40 ) );
+		Set( "Sequencer.SaveAs.Small", new IMAGE_BRUSH( "Sequencer/Main_Icons/Icon_Sequencer_Save_As_24x", Icon20x20 ) );
+		Set( "Sequencer.DiscardChanges", new IMAGE_BRUSH( "Sequencer/Main_Icons/Icon_Sequencer_Revert_24x", Icon40x40 ) );
+		Set( "Sequencer.DiscardChanges.Small", new IMAGE_BRUSH( "Sequencer/Main_Icons/Icon_Sequencer_Revert_24x", Icon20x20 ) );
 		Set( "Sequencer.GenericGripLeft", new BOX_BRUSH( "Sequencer/GenericGripLeft", FMargin(5.f/16.f) ) );
 		Set( "Sequencer.GenericGripRight", new BOX_BRUSH( "Sequencer/GenericGripRight", FMargin(5.f/16.f) ) );
 		Set( "Sequencer.SectionArea.Background", new FSlateColorBrush( FColor::White ) );
 
-		Set( "Sequencer.Section.Background", new BOX_BRUSH( TEXT("Sequencer/SectionBackground"), FMargin(3.f/100.f) ) );
-		Set( "Sequencer.Section.BackgroundTint", new BOX_BRUSH( TEXT("Sequencer/SectionBackgroundTint"), FMargin(48.f/100.f) ) );
+		Set( "Sequencer.Section.Background", new BORDER_BRUSH( TEXT("Sequencer/SectionBackground"), FMargin(4.f/16.f) ) );
+		Set( "Sequencer.Section.BackgroundTint", new BOX_BRUSH( TEXT("Sequencer/SectionBackgroundTint"), FMargin(4/16.f) ) );
 		Set( "Sequencer.Section.SelectedSectionOverlay", new IMAGE_BRUSH( TEXT("Sequencer/SelectedSectionOverlay"), Icon16x16, FLinearColor::White, ESlateBrushTileType::Both ) );
 		Set( "Sequencer.Section.SelectedTrackTint", new BOX_BRUSH( TEXT("Sequencer/SelectedTrackTint"), FMargin(0.f, 0.5f) ) );
 		Set( "Sequencer.Section.SelectionBorder", new BORDER_BRUSH( TEXT("Sequencer/SectionHighlight"), FMargin(7.f/16.f) ) );
@@ -1654,49 +1732,49 @@ void FSlateEditorStyle::FStyle::SetupGeneralStyles()
 		Set( "Sequencer.AnimationOutliner.RegularFont", TTF_CORE_FONT( "Fonts/Roboto-Regular", 9 ) );
 		Set( "Sequencer.ShotFilter", new IMAGE_BRUSH( "Sequencer/FilteredArea", FVector2D(74,74), FLinearColor::White, ESlateBrushTileType::Both ) );
 		Set( "Sequencer.KeyMark", new IMAGE_BRUSH("Sequencer/KeyMark", FVector2D(3,21), FLinearColor::White, ESlateBrushTileType::NoTile ) );
-		Set( "Sequencer.SetAutoKeyModeAll", new IMAGE_BRUSH( "Icons/SequencerIcons/icon_Sequencer_AutoKeyAll_24x", Icon40x40 ) );
-		Set( "Sequencer.SetAutoKeyModeAll.Small", new IMAGE_BRUSH( "Icons/SequencerIcons/icon_Sequencer_AutoKeyAll_24x", Icon20x20 ) );
-		Set( "Sequencer.SetAutoKeyModeAnimated", new IMAGE_BRUSH("Icons/SequencerIcons/icon_Sequencer_AutoKey_24x", Icon40x40 ) );
-		Set( "Sequencer.SetAutoKeyModeAnimated.Small", new IMAGE_BRUSH("Icons/SequencerIcons/icon_Sequencer_AutoKey_24x", Icon20x20 ) );
-		Set( "Sequencer.SetAutoKeyModeNone", new IMAGE_BRUSH("Icons/SequencerIcons/icon_Sequencer_AutoKeyNone_24x", Icon40x40));
-		Set( "Sequencer.SetAutoKeyModeNone.Small", new IMAGE_BRUSH("Icons/SequencerIcons/icon_Sequencer_AutoKeyNone_24x", Icon20x20 ) );
-		Set( "Sequencer.KeyAllEnabled", new IMAGE_BRUSH( "Icons/SequencerIcons/icon_Sequencer_WholeKeyframing_24x", Icon40x40 ) );
-		Set( "Sequencer.KeyAllEnabled.Small", new IMAGE_BRUSH( "Icons/SequencerIcons/icon_Sequencer_WholeKeyframing_24x", Icon20x20 ) );
-		Set( "Sequencer.KeyAllDisabled", new IMAGE_BRUSH( "Icons/SequencerIcons/icon_Sequencer_PartialKeyframing_24x", Icon40x40 ) );
-		Set( "Sequencer.KeyAllDisabled.Small", new IMAGE_BRUSH( "Icons/SequencerIcons/icon_Sequencer_PartialKeyframing_24x", Icon20x20 ) );
-		Set( "Sequencer.ToggleIsSnapEnabled", new IMAGE_BRUSH( "Icons/SequencerIcons/icon_Sequencer_Snap_24x", Icon40x40 ) );
-		Set( "Sequencer.ToggleIsSnapEnabled.Small", new IMAGE_BRUSH( "Icons/SequencerIcons/icon_Sequencer_Snap_24x", Icon20x20 ) );
+		Set( "Sequencer.SetAutoKeyModeAll", new IMAGE_BRUSH( "Sequencer/Main_Icons/Icon_Sequencer_Auto_Key_All_24x", Icon40x40 ) );
+		Set( "Sequencer.SetAutoKeyModeAll.Small", new IMAGE_BRUSH( "Sequencer/Main_Icons/Icon_Sequencer_Auto_Key_All_24x", Icon20x20 ) );
+		Set( "Sequencer.SetAutoKeyModeAnimated", new IMAGE_BRUSH("Sequencer/Main_Icons/Icon_Sequencer_Auto_Key_All_Animated_24x", Icon40x40 ) );
+		Set( "Sequencer.SetAutoKeyModeAnimated.Small", new IMAGE_BRUSH("Sequencer/Main_Icons/Icon_Sequencer_Auto_Key_All_Animated_24x", Icon20x20 ) );
+		Set( "Sequencer.SetAutoKeyModeNone", new IMAGE_BRUSH("Sequencer/Main_Icons/Icon_Sequencer_Disable_Auto_Key_24x", Icon40x40));
+		Set( "Sequencer.SetAutoKeyModeNone.Small", new IMAGE_BRUSH("Sequencer/Main_Icons/Icon_Sequencer_Disable_Auto_Key_24x", Icon20x20 ) );		
+		Set( "Sequencer.KeyAllEnabled", new IMAGE_BRUSH( "Sequencer/Main_Icons/Icon_Sequencer_Key_All_24x", Icon40x40 ) );
+		Set( "Sequencer.KeyAllEnabled.Small", new IMAGE_BRUSH( "Sequencer/Main_Icons/Icon_Sequencer_Key_All_24x", Icon20x20 ) );
+		Set( "Sequencer.KeyAllDisabled", new IMAGE_BRUSH( "Sequencer/Main_Icons/Icon_Sequencer_Key_Part_24x", Icon40x40 ) );
+		Set( "Sequencer.KeyAllDisabled.Small", new IMAGE_BRUSH( "Sequencer/Main_Icons/Icon_Sequencer_Key_Part_24x", Icon20x20 ) );
+		Set( "Sequencer.ToggleIsSnapEnabled", new IMAGE_BRUSH( "Sequencer/Main_Icons/Icon_Sequencer_Snap_24x", Icon40x40 ) );
+		Set( "Sequencer.ToggleIsSnapEnabled.Small", new IMAGE_BRUSH( "Sequencer/Main_Icons/Icon_Sequencer_Snap_24x", Icon20x20 ) );
 		Set( "Sequencer.ToggleShowCurveEditor", new IMAGE_BRUSH("Icons/SequencerIcons/icon_Sequencer_CurveEditor_24x", Icon40x40) );
 		Set( "Sequencer.ToggleShowCurveEditor.Small", new IMAGE_BRUSH("Icons/SequencerIcons/icon_Sequencer_CurveEditor_24x", Icon20x20) );
 		Set( "Sequencer.ToggleAutoScroll", new IMAGE_BRUSH( "Icons/icon_Sequencer_ToggleAutoScroll_40x", Icon40x40 ) );
 		Set( "Sequencer.ToggleAutoScroll.Small", new IMAGE_BRUSH( "Icons/icon_Sequencer_ToggleAutoScroll_16x", Icon16x16 ) );
 		Set( "Sequencer.MoveTool.Small", new IMAGE_BRUSH( "Icons/SequencerIcons/icon_Sequencer_Move_24x", Icon16x16 ) );
 		Set( "Sequencer.MarqueeTool.Small", new IMAGE_BRUSH( "Icons/SequencerIcons/icon_Sequencer_Marquee_24x", Icon16x16 ) );
-		Set( "Sequencer.RenderMovie.Small", new IMAGE_BRUSH( "Sequencer/RenderMovie_16px", Icon16x16 ) );
-		Set( "Sequencer.LockCamera", new IMAGE_BRUSH( "Icons/icon_ToggleActorPilotCameraView_16x", Icon16x16 ) );
-		Set( "Sequencer.UnlockCamera", new IMAGE_BRUSH( "Icons/icon_ToggleActorPilotCameraView_16x", Icon16x16, FLinearColor(1.f, 1.f, 1.f, 0.5f) ) );
+		Set( "Sequencer.RenderMovie.Small", new IMAGE_BRUSH( "Sequencer/Main_Icons/Icon_Sequencer_Create_Movie_24x", Icon16x16 ) );
+		Set( "Sequencer.LockCamera", new IMAGE_BRUSH( "Sequencer/Main_Icons/Icon_Sequencer_Look_Thru_24x", Icon16x16 ) );
+		Set( "Sequencer.UnlockCamera", new IMAGE_BRUSH( "Sequencer/Main_Icons/Icon_Sequencer_Look_Thru_24x", Icon16x16, FLinearColor(1.f, 1.f, 1.f, 0.5f) ) );
+		Set( "Sequencer.Thumbnail.SectionHandle", new IMAGE_BRUSH( "Old/White", Icon16x16, FLinearColor::Black ) );
 		Set( "Sequencer.TrackHoverHighlight_Top", new IMAGE_BRUSH( TEXT("Sequencer/TrackHoverHighlight_Top"), FVector2D(4, 4) ) );
 		Set( "Sequencer.TrackHoverHighlight_Bottom", new IMAGE_BRUSH( TEXT("Sequencer/TrackHoverHighlight_Bottom"), FVector2D(4, 4) ) );
 		Set( "Sequencer.SpawnableIconOverlay", new IMAGE_BRUSH( TEXT("Sequencer/SpawnableIconOverlay"), FVector2D(13, 13) ) );
-		Set( "Sequencer.DiscardChanges.Small", new IMAGE_BRUSH("Icons/icon_Delete_40x", Icon16x16));
 
-		Set( "Sequencer.General", new IMAGE_BRUSH( "Icons/SequencerIcons/icon_Sequencer_View_24x", Icon24x24 ) );
-		Set( "Sequencer.General.Small", new IMAGE_BRUSH( "Icons/SequencerIcons/icon_Sequencer_View_24x", Icon20x20 ) );
-		Set( "Sequencer.Time", new IMAGE_BRUSH( "Icons/SequencerIcons/icon_Sequencer_Time_24x", Icon24x24 ) );
-		Set( "Sequencer.Time.Small", new IMAGE_BRUSH( "Icons/SequencerIcons/icon_Sequencer_Time_24x", Icon20x20 ) );
-		Set( "Sequencer.Value", new IMAGE_BRUSH( "Icons/SequencerIcons/icon_Sequencer_Value_24x", Icon24x24 ) );
-		Set( "Sequencer.Value.Small", new IMAGE_BRUSH( "Icons/SequencerIcons/icon_Sequencer_Value_24x", Icon20x20 ) );
+		Set( "Sequencer.General", new IMAGE_BRUSH( "Sequencer/Main_Icons/Icon_Sequencer_General_Options_24x", Icon24x24 ) );
+		Set( "Sequencer.General.Small", new IMAGE_BRUSH( "Sequencer/Main_Icons/Icon_Sequencer_General_Options_24x", Icon20x20 ) );
+		Set( "Sequencer.Time", new IMAGE_BRUSH( "Sequencer/Main_Icons/Icon_Sequencer_Time_24x", Icon24x24 ) );
+		Set( "Sequencer.Time.Small", new IMAGE_BRUSH( "Sequencer/Main_Icons/Icon_Sequencer_Time_24x", Icon20x20 ) );
+		Set( "Sequencer.Value", new IMAGE_BRUSH( "Sequencer/Main_Icons/Icon_Sequencer_Value_24x", Icon24x24 ) );
+		Set( "Sequencer.Value.Small", new IMAGE_BRUSH( "Sequencer/Main_Icons/Icon_Sequencer_Value_24x", Icon20x20 ) );
 
 		Set( "Sequencer.TrackArea.LaneColor", FLinearColor(0.3f, 0.3f, 0.3f, 0.3f) );
 
-		Set( "Sequencer.Tracks.Audio", new IMAGE_BRUSH("Icons/AssetIcons/SoundCue_16x", Icon16x16));
-		Set( "Sequencer.Tracks.Event", new IMAGE_BRUSH("Icons/SequencerIcons/icon_Sequencer_Event_16x", Icon16x16));
-		Set( "Sequencer.Tracks.Fade", new IMAGE_BRUSH("Icons/SequencerIcons/icon_Sequencer_Fade_16x", Icon16x16));
-		Set( "Sequencer.Tracks.CameraCut", new IMAGE_BRUSH("Icons/SequencerIcons/icon_Sequencer_CameraCut_16x", Icon16x16));
-		Set( "Sequencer.Tracks.CinematicShot", new IMAGE_BRUSH("Icons/SequencerIcons/icon_Sequencer_CinematicShot_16x", Icon16x16));
-		Set( "Sequencer.Tracks.Slomo", new IMAGE_BRUSH("Icons/SequencerIcons/icon_Sequencer_Slomo_16x", Icon16x16));
-		Set( "Sequencer.Tracks.Sub", new IMAGE_BRUSH("Icons/SequencerIcons/icon_Sequencer_Sub_16x", Icon16x16));
-		Set( "Sequencer.Tracks.LevelVisibility", new IMAGE_BRUSH("Icons/SequencerIcons/icon_Sequencer_LevelVisibility_16x", Icon16x16));
+		Set( "Sequencer.Tracks.Audio", new IMAGE_BRUSH("Sequencer/Dropdown_Icons/Icon_Audio_Track_16x", Icon16x16));
+		Set( "Sequencer.Tracks.Event", new IMAGE_BRUSH("Sequencer/Dropdown_Icons/Icon_Event_Track_16x", Icon16x16));
+		Set( "Sequencer.Tracks.Fade", new IMAGE_BRUSH("Sequencer/Dropdown_Icons/Icon_Fade_Track_16x", Icon16x16));
+		Set( "Sequencer.Tracks.CameraCut", new IMAGE_BRUSH("Sequencer/Dropdown_Icons/Icon_Camera_Cut_Track_16x", Icon16x16));
+		Set( "Sequencer.Tracks.CinematicShot", new IMAGE_BRUSH("Sequencer/Dropdown_Icons/Icon_Shot_Track_16x", Icon16x16));
+		Set( "Sequencer.Tracks.Slomo", new IMAGE_BRUSH("Sequencer/Dropdown_Icons/Icon_Play_Rate_Track_16x", Icon16x16));
+		Set( "Sequencer.Tracks.Sub", new IMAGE_BRUSH("Sequencer/Dropdown_Icons/Icon_Sub_Track_16x", Icon16x16));
+		Set( "Sequencer.Tracks.LevelVisibility", new IMAGE_BRUSH("Sequencer/Dropdown_Icons/Icon_Level_Visibility_Track_16x", Icon16x16));
 
 		Set( "Sequencer.CursorDecorator_MarqueeAdd", new IMAGE_BRUSH( "Sequencer/CursorDecorator_MarqueeAdd", Icon16x16));
 		Set( "Sequencer.CursorDecorator_MarqueeSubtract", new IMAGE_BRUSH( "Sequencer/CursorDecorator_MarqueeSubtract", Icon16x16));
@@ -1735,6 +1813,20 @@ void FSlateEditorStyle::FStyle::SetupGeneralStyles()
 		Set( "Sequencer.FixedFont", TTF_FONT( "Fonts/DroidSansMono", 9 ) );
 	}
 
+
+	// Sequence recorder standalone UI
+	{
+		Set( "SequenceRecorder.TabIcon", new IMAGE_BRUSH( "SequenceRecorder/icon_tab_SequenceRecorder_16x", Icon16x16 ) );
+		Set( "SequenceRecorder.Common.RecordAll.Small", new IMAGE_BRUSH( "SequenceRecorder/icon_RecordAll_40x", Icon20x20 ) );
+		Set( "SequenceRecorder.Common.RecordAll", new IMAGE_BRUSH( "SequenceRecorder/icon_RecordAll_40x", Icon40x40 ) );
+		Set( "SequenceRecorder.Common.StopAll.Small", new IMAGE_BRUSH( "SequenceRecorder/icon_StopAll_40x", Icon20x20 ) );
+		Set( "SequenceRecorder.Common.StopAll", new IMAGE_BRUSH( "SequenceRecorder/icon_StopAll_40x", Icon40x40 ) );
+		Set( "SequenceRecorder.Common.AddRecording.Small", new IMAGE_BRUSH( "SequenceRecorder/icon_AddRecording_40x", Icon20x20 ) );
+		Set( "SequenceRecorder.Common.AddRecording", new IMAGE_BRUSH( "SequenceRecorder/icon_AddRecording_40x", Icon40x40 ) );
+		Set( "SequenceRecorder.Common.RemoveRecording.Small", new IMAGE_BRUSH( "SequenceRecorder/icon_RemoveRecording_40x", Icon20x20 ) );
+		Set( "SequenceRecorder.Common.RemoveRecording", new IMAGE_BRUSH( "SequenceRecorder/icon_RemoveRecording_40x", Icon40x40 ) );
+	}
+
 	// Foliage Edit Mode
 	{	
 		FLinearColor DimBackground = FLinearColor(FColor(64, 64, 64));
diff --git a/Engine/Source/Editor/EditorWidgets/Private/STransportControl.cpp b/Engine/Source/Editor/EditorWidgets/Private/STransportControl.cpp
--- a/Engine/Source/Editor/EditorWidgets/Private/STransportControl.cpp
+++ b/Engine/Source/Editor/EditorWidgets/Private/STransportControl.cpp
@@ -6,51 +6,26 @@
 #define LOCTEXT_NAMESPACE "STransportControl"
 
 BEGIN_SLATE_FUNCTION_BUILD_OPTIMIZATION
-void STransportControl::Construct( const STransportControl::FArguments& InArgs )
+TSharedPtr<SWidget> STransportControl::MakeTransportControlWidget(ETransportControlWidgetType WidgetType, const FOnMakeTransportWidget& MakeCustomWidgetDelegate)
 {
-	TransportControlArgs = InArgs._TransportArgs;
-	
-	ChildSlot
-	[
-		SNew(SHorizontalBox)
-		+SHorizontalBox::Slot()
-		.AutoWidth()
-		.HAlign(HAlign_Center) 
-		.VAlign(VAlign_Center) 
-		[
-			SNew(SButton)
+	switch(WidgetType)
+	{
+	case ETransportControlWidgetType::BackwardEnd:
+		return SNew(SButton)
+			. ButtonStyle(&FEditorStyle::Get().GetWidgetStyle<FButtonStyle>("Animation.Backward_End"))
 			. OnClicked(TransportControlArgs.OnBackwardEnd)
 			. Visibility(TransportControlArgs.OnBackwardEnd.IsBound() ? EVisibility::Visible : EVisibility::Collapsed)
 			. ToolTipText( LOCTEXT("ToFront", "To Front") )
-			. ButtonStyle( FEditorStyle::Get(), "NoBorder" )
-			. ContentPadding(2.0f)
-			[
-				SNew(SImage)
-				. Image( FEditorStyle::GetBrush("Animation.Backward_End") )
-			]
-		]
-		+SHorizontalBox::Slot()
-		.AutoWidth()
-		.HAlign(HAlign_Center) 
-		.VAlign(VAlign_Center) 
-		[
-			SNew(SButton)
+			. ContentPadding(2.0f);
+	case ETransportControlWidgetType::BackwardStep:
+		return SNew(SButton)
+			. ButtonStyle(&FEditorStyle::Get().GetWidgetStyle<FButtonStyle>("Animation.Backward_Step"))
 			. OnClicked(TransportControlArgs.OnBackwardStep)
 			. Visibility(TransportControlArgs.OnBackwardStep.IsBound() ? EVisibility::Visible : EVisibility::Collapsed)
 			. ToolTipText( LOCTEXT("ToPrevious", "To Previous") )
-			. ButtonStyle( FEditorStyle::Get(), "NoBorder" )
-			. ContentPadding(2.0f)
-			[
-				SNew(SImage)
-				. Image( FEditorStyle::GetBrush("Animation.Backward_Step") )
-			]
-		]
-		+SHorizontalBox::Slot()
-		.AutoWidth()
-		.HAlign(HAlign_Center) 
-		.VAlign(VAlign_Center) 
-		[
-			SNew(SButton)
+			. ContentPadding(2.0f);
+	case ETransportControlWidgetType::BackwardPlay:
+		return SAssignNew(BackwardPlayButton, SButton)
 			. OnClicked(TransportControlArgs.OnBackwardPlay)
 			. Visibility(TransportControlArgs.OnBackwardPlay.IsBound() ? EVisibility::Visible : EVisibility::Collapsed)
 			. ToolTipText( LOCTEXT("Reverse", "Reverse") )
@@ -59,30 +34,16 @@ void STransportControl::Construct( const STransportControl::FArguments& InArgs )
 			[
 				SNew(SImage)
 				. Image( this, &STransportControl::GetBackwardStatusIcon )
-			]
-		]
-		+SHorizontalBox::Slot()
-		.AutoWidth()
-		.HAlign(HAlign_Center) 
-		.VAlign(VAlign_Center) 
-		[
-			SNew(SButton)
+			];
+	case ETransportControlWidgetType::Record:
+		return SNew(SButton)
+			. ButtonStyle(&FEditorStyle::Get().GetWidgetStyle<FButtonStyle>("Animation.Record"))
 			. OnClicked(TransportControlArgs.OnRecord)
 			. Visibility(TransportControlArgs.OnRecord.IsBound() ? EVisibility::Visible : EVisibility::Collapsed)
 			. ToolTipText( this, &STransportControl::GetRecordStatusTooltip )
-			. ButtonStyle( FEditorStyle::Get(), "NoBorder" )
-			. ContentPadding(2.0f)
-			[
-				SNew(SImage)
-				. Image( FEditorStyle::GetBrush("Animation.Record") )
-			]
-		]
-		+SHorizontalBox::Slot()
-		.AutoWidth()
-		.HAlign(HAlign_Center) 
-		.VAlign(VAlign_Center) 
-		[
-			SNew(SButton)
+			. ContentPadding(2.0f);
+	case ETransportControlWidgetType::ForwardPlay:
+		return SAssignNew(ForwardPlayButton, SButton)
 			. OnClicked(TransportControlArgs.OnForwardPlay)
 			. Visibility(TransportControlArgs.OnForwardPlay.IsBound() ? EVisibility::Visible : EVisibility::Collapsed)
 			. ToolTipText( this, &STransportControl::GetForwardStatusTooltip )
@@ -91,46 +52,23 @@ void STransportControl::Construct( const STransportControl::FArguments& InArgs )
 			[
 				SNew(SImage)
 				. Image( this, &STransportControl::GetForwardStatusIcon )
-			]
-		]
-		+SHorizontalBox::Slot()
-		.AutoWidth()
-		.HAlign(HAlign_Center) 
-		.VAlign(VAlign_Center) 
-		[
-			SNew(SButton)
+			];
+	case ETransportControlWidgetType::ForwardStep:
+		return SNew(SButton)
+			. ButtonStyle(&FEditorStyle::Get().GetWidgetStyle<FButtonStyle>("Animation.Forward_Step"))
 			. OnClicked(TransportControlArgs.OnForwardStep)
 			. Visibility(TransportControlArgs.OnForwardStep.IsBound() ? EVisibility::Visible : EVisibility::Collapsed)
 			. ToolTipText( LOCTEXT("ToNext", "To Next") )
-			. ButtonStyle( FEditorStyle::Get(), "NoBorder" )
-			. ContentPadding(2.0f)
-			[
-				SNew(SImage)
-				. Image( FEditorStyle::GetBrush("Animation.Forward_Step") )
-			]
-		]
-		+SHorizontalBox::Slot()
-		.AutoWidth()
-		.HAlign(HAlign_Center) 
-		.VAlign(VAlign_Center) 
-		[
-			SNew(SButton)
+			. ContentPadding(2.0f);
+	case ETransportControlWidgetType::ForwardEnd:
+		return SNew(SButton)
+			. ButtonStyle(&FEditorStyle::Get().GetWidgetStyle<FButtonStyle>("Animation.Forward_End"))
 			. OnClicked(TransportControlArgs.OnForwardEnd)
 			. Visibility(TransportControlArgs.OnForwardEnd.IsBound() ? EVisibility::Visible : EVisibility::Collapsed)
 			. ToolTipText( LOCTEXT("ToEnd", "To End") )
-			. ButtonStyle( FEditorStyle::Get(), "NoBorder" )
-			. ContentPadding(2.0f)
-			[
-				SNew(SImage)
-				. Image( FEditorStyle::GetBrush("Animation.Forward_End") )
-			]
-		]
-		+SHorizontalBox::Slot()
-		.AutoWidth()
-		.HAlign(HAlign_Center) 
-		.VAlign(VAlign_Center) 
-		[
-			SNew(SButton)
+			. ContentPadding(2.0f);
+	case ETransportControlWidgetType::Loop:
+		return SNew(SButton)
 			. OnClicked(TransportControlArgs.OnToggleLooping)
 			. Visibility(TransportControlArgs.OnGetLooping.IsBound() ? EVisibility::Visible : EVisibility::Collapsed)
 			. ToolTipText( LOCTEXT("Loop", "Loop") )
@@ -138,10 +76,62 @@ void STransportControl::Construct( const STransportControl::FArguments& InArgs )
 			. ContentPadding(2.0f)
 			[
 				SNew(SImage)
-				.Image( FEditorStyle::GetBrush("Animation.Loop") )
-				.ColorAndOpacity( this, &STransportControl::GetLoopStatusColor )
-			]
-		]
+				. Image( this, &STransportControl::GetLoopStatusIcon )
+			];
+	case ETransportControlWidgetType::Custom:
+		if(MakeCustomWidgetDelegate.IsBound())
+		{
+			return MakeCustomWidgetDelegate.Execute();
+		}
+		break;
+	}
+
+	return nullptr;
+}
+
+void STransportControl::Construct( const STransportControl::FArguments& InArgs )
+{
+	TransportControlArgs = InArgs._TransportArgs;
+	
+	TSharedRef<SHorizontalBox> HorizontalBox = SNew(SHorizontalBox);
+	if(TransportControlArgs.WidgetsToCreate.Num() > 0)
+	{
+		for(FTransportControlWidget WidgetDesc : TransportControlArgs.WidgetsToCreate)
+		{
+			TSharedPtr<SWidget> Widget = MakeTransportControlWidget(WidgetDesc.WidgetType, WidgetDesc.MakeCustomWidgetDelegate);
+			if(Widget.IsValid())
+			{
+				HorizontalBox->AddSlot()
+				.AutoWidth()
+				.HAlign(HAlign_Center)
+				.VAlign(VAlign_Center)
+				[
+					Widget.ToSharedRef()
+				];
+			}
+		}	
+	}
+	else
+	{
+		for(ETransportControlWidgetType WidgetType : TEnumRange<ETransportControlWidgetType>())
+		{
+			TSharedPtr<SWidget> Widget = MakeTransportControlWidget(WidgetType);
+			if(Widget.IsValid())
+			{
+				HorizontalBox->AddSlot()
+				.AutoWidth()
+				.HAlign(HAlign_Center)
+				.VAlign(VAlign_Center)
+				[
+					Widget.ToSharedRef()
+				];
+			}
+		}
+	}
+
+	ChildSlot
+	[
+		HorizontalBox
 	];
 }
 END_SLATE_FUNCTION_BUILD_OPTIMIZATION
@@ -173,10 +163,14 @@ const FSlateBrush* STransportControl::GetForwardStatusIcon() const
 	if (TransportControlArgs.OnGetPlaybackMode.IsBound() &&
 		 ( PlaybackMode == EPlaybackMode::PlayingForward || PlaybackMode == EPlaybackMode::Recording ) )
 	{
-		return FEditorStyle::GetBrush("Animation.Pause");
+		return ForwardPlayButton.IsValid() && ForwardPlayButton->IsPressed() ? 
+			&FEditorStyle::Get().GetWidgetStyle<FButtonStyle>("Animation.Pause").Pressed : 
+			&FEditorStyle::Get().GetWidgetStyle<FButtonStyle>("Animation.Pause").Normal;
 	}
 
-	return FEditorStyle::GetBrush("Animation.Forward");
+	return ForwardPlayButton.IsValid() && ForwardPlayButton->IsPressed() ? 
+		&FEditorStyle::Get().GetWidgetStyle<FButtonStyle>("Animation.Forward").Pressed :
+		&FEditorStyle::Get().GetWidgetStyle<FButtonStyle>("Animation.Forward").Normal;
 }
 
 FText STransportControl::GetForwardStatusTooltip() const
@@ -206,21 +200,29 @@ const FSlateBrush* STransportControl::GetBackwardStatusIcon() const
 	if (TransportControlArgs.OnGetPlaybackMode.IsBound() &&
 		TransportControlArgs.OnGetPlaybackMode.Execute() == EPlaybackMode::PlayingReverse)
 	{
-		return FEditorStyle::GetBrush("Animation.Pause");
+		return BackwardPlayButton.IsValid() && BackwardPlayButton->IsPressed() ? 
+			&FEditorStyle::Get().GetWidgetStyle<FButtonStyle>("Animation.Pause").Pressed : 
+			&FEditorStyle::Get().GetWidgetStyle<FButtonStyle>("Animation.Pause").Normal;
 	}
 
-	return FEditorStyle::GetBrush("Animation.Backward");
+	return BackwardPlayButton.IsValid() && BackwardPlayButton->IsPressed() ? 
+		&FEditorStyle::Get().GetWidgetStyle<FButtonStyle>("Animation.Backward").Pressed : 
+		&FEditorStyle::Get().GetWidgetStyle<FButtonStyle>("Animation.Backward").Normal;
 }
 
-FSlateColor STransportControl::GetLoopStatusColor() const
+const FSlateBrush* STransportControl::GetLoopStatusIcon() const
 {
 	if (TransportControlArgs.OnGetLooping.IsBound() &&
 		TransportControlArgs.OnGetLooping.Execute())
 	{
-		return FSlateColor(FLinearColor(1.f, 0.f, 0.f));
+		return BackwardPlayButton.IsValid() && BackwardPlayButton->IsPressed() ? 
+			&FEditorStyle::Get().GetWidgetStyle<FButtonStyle>("Animation.Loop.Enabled").Pressed : 
+			&FEditorStyle::Get().GetWidgetStyle<FButtonStyle>("Animation.Loop.Enabled").Normal;
 	}
 
-	return FSlateColor(FLinearColor(1.f, 1.f, 1.f));
+	return BackwardPlayButton.IsValid() && BackwardPlayButton->IsPressed() ? 
+		&FEditorStyle::Get().GetWidgetStyle<FButtonStyle>("Animation.Loop.Disabled").Pressed : 
+		&FEditorStyle::Get().GetWidgetStyle<FButtonStyle>("Animation.Loop.Disabled").Normal;
 }
 
 EActiveTimerReturnType STransportControl::TickPlayback( double InCurrentTime, float InDeltaTime )
diff --git a/Engine/Source/Editor/EditorWidgets/Private/STransportControl.h b/Engine/Source/Editor/EditorWidgets/Private/STransportControl.h
--- a/Engine/Source/Editor/EditorWidgets/Private/STransportControl.h
+++ b/Engine/Source/Editor/EditorWidgets/Private/STransportControl.h
@@ -33,13 +33,16 @@ private:
 	FText GetForwardStatusTooltip() const;
 	FText GetRecordStatusTooltip() const;
 	const FSlateBrush* GetBackwardStatusIcon() const;
-	FSlateColor GetLoopStatusColor() const;
+	const FSlateBrush* GetLoopStatusIcon() const;
 
 	/** Executes the OnTickPlayback delegate */
 	EActiveTimerReturnType TickPlayback( double InCurrentTime, float InDeltaTime );
 
 	FReply OnToggleLooping();
 
+	/** Make default transport control widgets */
+	TSharedPtr<SWidget> MakeTransportControlWidget(ETransportControlWidgetType WidgetType, const FOnMakeTransportWidget& MakeCustomWidgetDelegate = FOnMakeTransportWidget());
+
 private:
 	/** The handle to the active timer */
 	TWeakPtr<FActiveTimerHandle> ActiveTimerHandle;
@@ -48,4 +51,7 @@ private:
 	bool bIsActiveTimerRegistered;
 
 	FTransportControlArgs TransportControlArgs;
+
+	TSharedPtr<SButton> ForwardPlayButton;
+	TSharedPtr<SButton> BackwardPlayButton;
 };
diff --git a/Engine/Source/Editor/EditorWidgets/Public/ITransportControl.h b/Engine/Source/Editor/EditorWidgets/Public/ITransportControl.h
--- a/Engine/Source/Editor/EditorWidgets/Public/ITransportControl.h
+++ b/Engine/Source/Editor/EditorWidgets/Public/ITransportControl.h
@@ -16,6 +16,41 @@ namespace EPlaybackMode
 DECLARE_DELEGATE_RetVal(bool, FOnGetLooping)
 DECLARE_DELEGATE_RetVal(EPlaybackMode::Type, FOnGetPlaybackMode)
 DECLARE_DELEGATE_TwoParams(FOnTickPlayback, double /*InCurrentTime*/, float /*InDeltaTime*/)
+DECLARE_DELEGATE_RetVal(TSharedRef<SWidget>, FOnMakeTransportWidget)
+
+enum class ETransportControlWidgetType : int32
+{
+	BackwardEnd,
+	BackwardStep,
+	BackwardPlay,
+	Record,
+	ForwardPlay,
+	ForwardStep,
+	ForwardEnd,
+	Loop,
+	Custom
+};
+
+ENUM_RANGE_BY_COUNT(ETransportControlWidgetType, ETransportControlWidgetType::Custom);
+
+/** Descriptor for a transport control widget */
+struct FTransportControlWidget
+{
+	FTransportControlWidget(ETransportControlWidgetType InWidgetType)
+		: WidgetType(InWidgetType)
+	{}
+
+	FTransportControlWidget(FOnMakeTransportWidget InMakeCustomWidgetDelegate)
+		: WidgetType(ETransportControlWidgetType::Custom)
+		, MakeCustomWidgetDelegate(InMakeCustomWidgetDelegate)
+	{}
+
+	/** Basic widget type */
+	ETransportControlWidgetType WidgetType;
+
+	/** Delegate used for making custom widgets */
+	FOnMakeTransportWidget MakeCustomWidgetDelegate;
+};
 
 struct FTransportControlArgs
 {
@@ -44,6 +79,12 @@ struct FTransportControlArgs
 	FOnGetLooping OnGetLooping;
 	FOnGetPlaybackMode OnGetPlaybackMode;
 	FOnTickPlayback OnTickPlayback;
+
+	/** 
+	 * Array of custom widgets to create - if this array is used the default widget ordering will be
+	 * ignored in favor of this set of widgets
+	 */
+	TArray<FTransportControlWidget> WidgetsToCreate;
 };
 
 /**
diff --git a/Engine/Source/Editor/LevelEditor/Private/LevelEditor.cpp b/Engine/Source/Editor/LevelEditor/Private/LevelEditor.cpp
--- a/Engine/Source/Editor/LevelEditor/Private/LevelEditor.cpp
+++ b/Engine/Source/Editor/LevelEditor/Private/LevelEditor.cpp
@@ -21,6 +21,8 @@
 #include "TargetPlatform.h"
 #include "IIntroTutorials.h"
 #include "IProjectManager.h"
+#include "LevelViewportLayout.h"
+#include "LevelViewportLayoutEntity.h"
 
 // @todo Editor: remove this circular dependency
 #include "Editor/MainFrame/Public/Interfaces/IMainFrameModule.h"
@@ -231,6 +233,8 @@ void FLevelEditorModule::StartupModule()
 	// Bind level editor commands shared across an instance
 	BindGlobalLevelEditorCommands();
 
+	RegisterViewportType("Default", FViewportTypeDefinition::FromType<FLevelViewportLayoutEntity>());
+
 	const IWorkspaceMenuStructure& MenuStructure = WorkspaceMenu::GetMenuStructure();
 
 	FGlobalTabmanager::Get()->RegisterTabSpawner("LevelEditor", FOnSpawnTab::CreateRaw( this, &FLevelEditorModule::SpawnLevelEditor ) )
@@ -534,6 +538,17 @@ TSharedPtr<SDockTab> FLevelEditorModule::GetLevelEditorTab() const
 	return LevelEditorInstanceTabPtr.Pin();
 }
 
+TSharedRef<IViewportLayoutEntity> FLevelEditorModule::FactoryViewport(FName InTypeName, const FViewportConstructionArgs& ConstructionArgs) const
+{
+	const FViewportTypeDefinition* Definition = CustomViewports.Find(InTypeName);
+	if (Definition)
+	{
+		return Definition->FactoryFunction(ConstructionArgs);
+	}
+
+	return FViewportTypeDefinition::FromType<FLevelViewportLayoutEntity>().FactoryFunction(ConstructionArgs);
+}
+
 void FLevelEditorModule::BindGlobalLevelEditorCommands()
 {
 	check( !GlobalLevelEditorActions.IsValid() );
diff --git a/Engine/Source/Editor/LevelEditor/Private/LevelViewportActions.cpp b/Engine/Source/Editor/LevelEditor/Private/LevelViewportActions.cpp
--- a/Engine/Source/Editor/LevelEditor/Private/LevelViewportActions.cpp
+++ b/Engine/Source/Editor/LevelEditor/Private/LevelViewportActions.cpp
@@ -50,6 +50,8 @@ void FLevelViewportCommands::RegisterCommands()
 	UI_COMMAND( ViewportConfig_FourPanesBottom, "Layout Four Panes (one bottom, three top)", "Changes the viewport arrangement to four panes, one on the bottom, three on the top", EUserInterfaceActionType::ToggleButton, FInputChord() );
 	UI_COMMAND( ViewportConfig_FourPanes2x2, "Layout Four Panes (2x2)", "Changes the viewport arrangement to four panes, in a 2x2 grid", EUserInterfaceActionType::ToggleButton, FInputChord() );
 
+	UI_COMMAND( ToggleViewportToolbar, "Show Toolbar", "Defines whether a toolbar should be displayed on this viewport", EUserInterfaceActionType::ToggleButton, FInputChord(EModifierKey::Control | EModifierKey::Shift, EKeys::T) );
+
 	UI_COMMAND( ApplyMaterialToActor, "Apply Material", "Attempts to apply a dropped material to this object", EUserInterfaceActionType::Button, FInputChord() );
 
 	UI_COMMAND( ToggleCinematicPreview, "Cinematic Preview", "If enabled, allows Matinee or Sequencer previews to play in this viewport", EUserInterfaceActionType::ToggleButton, FInputChord() );
diff --git a/Engine/Source/Editor/LevelEditor/Private/LevelViewportLayout.cpp b/Engine/Source/Editor/LevelEditor/Private/LevelViewportLayout.cpp
--- a/Engine/Source/Editor/LevelEditor/Private/LevelViewportLayout.cpp
+++ b/Engine/Source/Editor/LevelEditor/Private/LevelViewportLayout.cpp
@@ -129,48 +129,6 @@ TSharedPtr<FLevelViewportTabContent> SViewportsOverlay::GetLevelViewportTab() co
 	return LevelViewportTab;
 }
 
-FLevelViewportLayoutEntity::FLevelViewportLayoutEntity(TSharedRef<SLevelViewport> InLevelViewport)
-	: LevelViewport(InLevelViewport)
-{}
-
-TSharedRef<SWidget> FLevelViewportLayoutEntity::AsWidget() const
-{
-	return LevelViewport;
-}
-
-TSharedPtr<SLevelViewport> FLevelViewportLayoutEntity::AsLevelViewport() const
-{
-	return LevelViewport;
-}
-
-bool FLevelViewportLayoutEntity::IsPlayInEditorViewportActive() const
-{
-	return LevelViewport->IsPlayInEditorViewportActive();
-}
-
-void FLevelViewportLayoutEntity::RegisterGameViewportIfPIE()
-{
-	return LevelViewport->RegisterGameViewportIfPIE();
-}
-
-void FLevelViewportLayoutEntity::SetKeyboardFocus()
-{
-	return LevelViewport->SetKeyboardFocusToThisViewport();
-}
-
-void FLevelViewportLayoutEntity::OnLayoutDestroyed()
-{
-	if (LevelViewport->IsPlayInEditorViewportActive() || LevelViewport->GetLevelViewportClient().IsSimulateInEditorViewport() )
-	{
-		GUnrealEd->EndPlayMap();
-	}
-}
-
-void FLevelViewportLayoutEntity::SaveConfig(const FString& ConfigSection)
-{
-	LevelViewport->SaveConfig(ConfigSection);
-}
-
 // FLevelViewportLayout /////////////////////////////
 
 FLevelViewportLayout::FLevelViewportLayout()
@@ -208,7 +166,7 @@ FLevelViewportLayout::~FLevelViewportLayout()
 }
 
 
-TSharedRef<SWidget> FLevelViewportLayout::BuildViewportLayout( TSharedPtr<SDockTab> InParentDockTab, TSharedPtr<FLevelViewportTabContent> InParentTab, const FString& LayoutString, TWeakPtr<SLevelEditor> InParentLevelEditor )
+TSharedRef<SWidget> FLevelViewportLayout::BuildViewportLayout( TSharedPtr<SDockTab> InParentDockTab, TSharedPtr<FLevelViewportTabContent> InParentTab, const FString& LayoutString, TWeakPtr<ILevelEditor> InParentLevelEditor )
 {
 	// We don't support reconfiguring an existing layout object, as this makes handling of transitions
 	// particularly difficult.  Instead just destroy the old layout and create a new layout object.
@@ -296,7 +254,11 @@ void FLevelViewportLayout::SaveCommonLayoutString( const FString& SpecificLayout
 	// Save all our data using the additional layout config
 	for (auto& Pair : Viewports)
 	{
-		Pair.Value->SaveConfig(SpecificLayoutString + Pair.Key.ToString());
+		FString ConfigName = SpecificLayoutString + Pair.Key.ToString();
+
+		Pair.Value->SaveConfig(ConfigName);
+
+		GConfig->SetString( *IniSection, *( ConfigName + TEXT(".TypeWithinLayout") ), *Pair.Value->GetType().ToString(), GEditorPerProjectIni );
 	}
 
 	// We don't bother saving that we were in immersive mode, because we never want to start back up directly in immersive mode
diff --git a/Engine/Source/Editor/LevelEditor/Private/LevelViewportLayout2x2.cpp b/Engine/Source/Editor/LevelEditor/Private/LevelViewportLayout2x2.cpp
--- a/Engine/Source/Editor/LevelEditor/Private/LevelViewportLayout2x2.cpp
+++ b/Engine/Source/Editor/LevelEditor/Private/LevelViewportLayout2x2.cpp
@@ -16,12 +16,10 @@ TSharedRef<SWidget> FLevelViewportLayout2x2::MakeViewportLayout(const FString& L
 {
 	FString SpecificLayoutString = GetTypeSpecificLayoutString(LayoutString);
 
-	TSharedPtr< SLevelViewport > ViewportWidgetTL;
-	TSharedPtr< SLevelViewport > ViewportWidgetTR;
-	TSharedPtr< SLevelViewport > ViewportWidgetBL;
-	TSharedPtr< SLevelViewport > ViewportWidgetBR;
-
 	FString TopLeftKey, BottomLeftKey, TopRightKey, BottomRightKey;
+
+	FString TopLeftType = TEXT("Default"), BottomLeftType = TEXT("Default"), TopRightType = TEXT("Default"), BottomRightType = TEXT("Default");
+
 	TArray<FVector2D> SplitterPercentages;
 	
 	if (!SpecificLayoutString.IsEmpty())
@@ -34,6 +32,11 @@ TSharedRef<SWidget> FLevelViewportLayout2x2::MakeViewportLayout(const FString& L
 		TopRightKey = SpecificLayoutString + TEXT(".Viewport2");
 		BottomRightKey = SpecificLayoutString + TEXT(".Viewport3");
 
+		GConfig->GetString(*IniSection, *(TopLeftKey + TEXT(".TypeWithinLayout")), TopLeftType, GEditorPerProjectIni);
+		GConfig->GetString(*IniSection, *(TopRightKey + TEXT(".TypeWithinLayout")), TopRightType, GEditorPerProjectIni);
+		GConfig->GetString(*IniSection, *(BottomLeftKey + TEXT(".TypeWithinLayout")), BottomLeftType, GEditorPerProjectIni);
+		GConfig->GetString(*IniSection, *(BottomRightKey + TEXT(".TypeWithinLayout")), BottomRightType, GEditorPerProjectIni);
+
 		for (int32 i = 0; i < 4; ++i)
 		{
 			FString PercentageString;
@@ -46,56 +49,66 @@ TSharedRef<SWidget> FLevelViewportLayout2x2::MakeViewportLayout(const FString& L
 		}
 	}
 
+	FLevelEditorModule& LevelEditor = FModuleManager::LoadModuleChecked<FLevelEditorModule>("LevelEditor");
+
+	// Set up the viewports
+	FViewportConstructionArgs Args;
+	Args.ParentLayout = AsShared();
+	Args.ParentLevelEditor = ParentLevelEditor;
+	Args.IsEnabled = FSlateApplication::Get().GetNormalExecutionAttribute();
+
+	// Left viewport
+	Args.bRealtime = false;
+	Args.ConfigKey = TopLeftKey;
+	Args.ViewportType = LVT_OrthoYZ;
+	TSharedPtr< IViewportLayoutEntity > ViewportTL = LevelEditor.FactoryViewport(*TopLeftType, Args);
+
+	// Persp viewport
+	Args.bRealtime = true;
+	Args.ConfigKey = BottomLeftKey;
+	Args.ViewportType = LVT_Perspective;
+	TSharedPtr< IViewportLayoutEntity > ViewportBL = LevelEditor.FactoryViewport(*BottomLeftType, Args);
+
+	// Front viewport
+	Args.bRealtime = false;
+	Args.ConfigKey = TopRightKey;
+	Args.ViewportType = LVT_OrthoXZ;
+	TSharedPtr< IViewportLayoutEntity > ViewportTR = LevelEditor.FactoryViewport(*TopRightType, Args);
+
+	// Top Viewport
+	Args.bRealtime = false;
+	Args.ConfigKey = BottomRightKey;
+	Args.ViewportType = LVT_OrthoXY;
+	TSharedPtr< IViewportLayoutEntity > ViewportBR = LevelEditor.FactoryViewport(*BottomRightType, Args);
+
+	Viewports.Add( *TopLeftKey, ViewportTL );
+	Viewports.Add( *BottomLeftKey, ViewportBL );
+	Viewports.Add( *TopRightKey, ViewportTR );
+	Viewports.Add( *BottomRightKey, ViewportBR );
+
+	// Set up the splitter
 	SplitterWidget = 
 	SNew( SSplitter2x2 )
 	.TopLeft()
 	[
-		SAssignNew( ViewportWidgetTL, SLevelViewport )
-		.ViewportType( LVT_OrthoYZ )		// Left
-		.ParentLayout( AsShared() )
-		.ParentLevelEditor( ParentLevelEditor )
-		.ConfigKey( TopLeftKey )
-		.IsEnabled(FSlateApplication::Get().GetNormalExecutionAttribute())
+		ViewportTL->AsWidget()
 	]
 	.BottomLeft()
 	[
-		SAssignNew( ViewportWidgetBL, SLevelViewport )
-		.ViewportType( LVT_Perspective ) // Persp
-		.Realtime( true )
-		.ParentLayout( AsShared() )
-		.ParentLevelEditor( ParentLevelEditor )
-		.ConfigKey( BottomLeftKey )
-		.IsEnabled(FSlateApplication::Get().GetNormalExecutionAttribute())
+		ViewportBL->AsWidget()
 	]
 	.TopRight()
 	[
-		SAssignNew( ViewportWidgetTR, SLevelViewport )
-		.ViewportType( LVT_OrthoXZ )		// Front
-		.ParentLayout( AsShared() )
-		.ParentLevelEditor( ParentLevelEditor )
-		.ConfigKey( TopRightKey )
-		.IsEnabled(FSlateApplication::Get().GetNormalExecutionAttribute())
+		ViewportTR->AsWidget()
 	]
 	.BottomRight()
 	[
-		SAssignNew( ViewportWidgetBR, SLevelViewport )
-		.ViewportType( LVT_OrthoXY )		// Top
-		.ParentLayout( AsShared() )
-		.ParentLevelEditor( ParentLevelEditor )
-		.ConfigKey( BottomRightKey )
-		.IsEnabled(FSlateApplication::Get().GetNormalExecutionAttribute())
+		ViewportBR->AsWidget()
 	];
-
-	Viewports.Add( *TopLeftKey, MakeShareable( new FLevelViewportLayoutEntity(ViewportWidgetTL.ToSharedRef()) ) );
-	Viewports.Add( *BottomLeftKey, MakeShareable( new FLevelViewportLayoutEntity(ViewportWidgetBL.ToSharedRef()) ) );
-	Viewports.Add( *TopRightKey, MakeShareable( new FLevelViewportLayoutEntity(ViewportWidgetTR.ToSharedRef()) ) );
-	Viewports.Add( *BottomRightKey, MakeShareable( new FLevelViewportLayoutEntity(ViewportWidgetBR.ToSharedRef()) ) );
-
 	
 	// Make newly-created perspective viewports active by default
-	GCurrentLevelEditingViewportClient = &ViewportWidgetBL->GetLevelViewportClient();
+	GCurrentLevelEditingViewportClient = &ViewportBL->GetLevelViewportClient();
 
-	
 	if (SplitterPercentages.Num() > 0)
 	{
 		SplitterWidget->SetSplitterPercentages(SplitterPercentages);
diff --git a/Engine/Source/Editor/LevelEditor/Private/LevelViewportLayoutEntity.cpp b/Engine/Source/Editor/LevelEditor/Private/LevelViewportLayoutEntity.cpp
new file mode 100644
--- /dev/null
+++ b/Engine/Source/Editor/LevelEditor/Private/LevelViewportLayoutEntity.cpp
@@ -0,0 +1,68 @@
+// Copyright 1998-2016 Epic Games, Inc. All Rights Reserved.
+
+#pragma once
+
+#include "LevelEditor.h"
+#include "LevelViewportLayoutEntity.h"
+#include "SLevelViewport.h"
+
+FLevelViewportLayoutEntity::FLevelViewportLayoutEntity(const FViewportConstructionArgs& ConstructionArgs)
+	: LevelViewport(
+		SNew(SLevelViewport)
+		.Realtime(ConstructionArgs.bRealtime)
+		.ViewportType( ConstructionArgs.ViewportType )
+		.ParentLayout( ConstructionArgs.ParentLayout )
+		.ParentLevelEditor( ConstructionArgs.ParentLevelEditor )
+		.ConfigKey( ConstructionArgs.ConfigKey )
+		.IsEnabled( ConstructionArgs.IsEnabled )
+	)
+{}
+
+TSharedRef<SWidget> FLevelViewportLayoutEntity::AsWidget() const
+{
+	return LevelViewport;
+}
+
+TSharedPtr<SLevelViewport> FLevelViewportLayoutEntity::AsLevelViewport() const
+{
+	return LevelViewport;
+}
+
+FName FLevelViewportLayoutEntity::GetType() const
+{
+	static FName DefaultName("Default");
+	return DefaultName;
+}
+
+FLevelEditorViewportClient& FLevelViewportLayoutEntity::GetLevelViewportClient() const
+{
+	return LevelViewport->GetLevelViewportClient();
+}
+
+bool FLevelViewportLayoutEntity::IsPlayInEditorViewportActive() const
+{
+	return LevelViewport->IsPlayInEditorViewportActive();
+}
+
+void FLevelViewportLayoutEntity::RegisterGameViewportIfPIE()
+{
+	return LevelViewport->RegisterGameViewportIfPIE();
+}
+
+void FLevelViewportLayoutEntity::SetKeyboardFocus()
+{
+	return LevelViewport->SetKeyboardFocusToThisViewport();
+}
+
+void FLevelViewportLayoutEntity::OnLayoutDestroyed()
+{
+	if (LevelViewport->IsPlayInEditorViewportActive() || LevelViewport->GetLevelViewportClient().IsSimulateInEditorViewport() )
+	{
+		GUnrealEd->EndPlayMap();
+	}
+}
+
+void FLevelViewportLayoutEntity::SaveConfig(const FString& ConfigSection)
+{
+	LevelViewport->SaveConfig(ConfigSection);
+}
\ No newline at end of file
diff --git a/Engine/Source/Editor/LevelEditor/Private/LevelViewportLayoutEntity.h b/Engine/Source/Editor/LevelEditor/Private/LevelViewportLayoutEntity.h
new file mode 100644
--- /dev/null
+++ b/Engine/Source/Editor/LevelEditor/Private/LevelViewportLayoutEntity.h
@@ -0,0 +1,28 @@
+// Copyright 1998-2016 Epic Games, Inc. All Rights Reserved.
+
+#pragma once
+
+#include "LevelEditor.h"
+#include "LevelViewportLayout.h"
+
+class LEVELEDITOR_API FLevelViewportLayoutEntity : public IViewportLayoutEntity
+{
+public:
+
+	FLevelViewportLayoutEntity(const FViewportConstructionArgs& ConstructionArgs);
+	virtual TSharedRef<SWidget> AsWidget() const override;
+	virtual TSharedPtr<SLevelViewport> AsLevelViewport() const override;
+
+	bool IsPlayInEditorViewportActive() const;
+	void RegisterGameViewportIfPIE();
+	void SetKeyboardFocus();
+	void OnLayoutDestroyed();
+	void SaveConfig(const FString& ConfigSection);
+	FLevelEditorViewportClient& GetLevelViewportClient() const;
+	FName GetType() const;
+
+protected:
+
+	/** This entity's level viewport */
+	TSharedRef<SLevelViewport> LevelViewport;
+};
\ No newline at end of file
diff --git a/Engine/Source/Editor/LevelEditor/Private/LevelViewportLayoutFourPanes.cpp b/Engine/Source/Editor/LevelEditor/Private/LevelViewportLayoutFourPanes.cpp
--- a/Engine/Source/Editor/LevelEditor/Private/LevelViewportLayoutFourPanes.cpp
+++ b/Engine/Source/Editor/LevelEditor/Private/LevelViewportLayoutFourPanes.cpp
@@ -21,6 +21,8 @@ TSharedRef<SWidget> FLevelViewportLayoutFourPanes::MakeViewportLayout(const FStr
 	FString SpecificLayoutString = GetTypeSpecificLayoutString(LayoutString);
 
 	FString ViewportKey0, ViewportKey1, ViewportKey2, ViewportKey3;
+	FString ViewportType0 = TEXT("Default"), ViewportType1 = TEXT("Default"), ViewportType2 = TEXT("Default"), ViewportType3 = TEXT("Default");
+
 	float PrimarySplitterPercentage = ViewportLayoutFourPanesDefs::DefaultPrimarySplitterPercentage;
 	float SecondarySplitterPercentage0 = ViewportLayoutFourPanesDefs::DefaultSecondarySplitterPercentage;
 	float SecondarySplitterPercentage1 = ViewportLayoutFourPanesDefs::DefaultSecondarySplitterPercentage;
@@ -35,6 +37,11 @@ TSharedRef<SWidget> FLevelViewportLayoutFourPanes::MakeViewportLayout(const FStr
 		ViewportKey2 = SpecificLayoutString + TEXT(".Viewport2");
 		ViewportKey3 = SpecificLayoutString + TEXT(".Viewport3");
 
+		GConfig->GetString(*IniSection, *(ViewportKey0 + TEXT(".TypeWithinLayout")), ViewportType0, GEditorPerProjectIni);
+		GConfig->GetString(*IniSection, *(ViewportKey1 + TEXT(".TypeWithinLayout")), ViewportType1, GEditorPerProjectIni);
+		GConfig->GetString(*IniSection, *(ViewportKey2 + TEXT(".TypeWithinLayout")), ViewportType2, GEditorPerProjectIni);
+		GConfig->GetString(*IniSection, *(ViewportKey3 + TEXT(".TypeWithinLayout")), ViewportType3, GEditorPerProjectIni);
+
 		FString PercentageString;
 		if(GConfig->GetString(*IniSection, *(SpecificLayoutString + TEXT(".Percentage0")), PercentageString, GEditorPerProjectIni))
 		{
@@ -47,12 +54,50 @@ TSharedRef<SWidget> FLevelViewportLayoutFourPanes::MakeViewportLayout(const FStr
 		if(GConfig->GetString(*IniSection, *(SpecificLayoutString + TEXT(".Percentage2")), PercentageString, GEditorPerProjectIni))
 		{
 			TTypeFromString<float>::FromString(SecondarySplitterPercentage1, *PercentageString);
-		}	
+		}
 	}
 
+	FLevelEditorModule& LevelEditor = FModuleManager::LoadModuleChecked<FLevelEditorModule>("LevelEditor");
+
+	// Set up the viewports
+	FViewportConstructionArgs Args;
+	Args.ParentLayout = AsShared();
+	Args.ParentLevelEditor = ParentLevelEditor;
+	Args.IsEnabled = FSlateApplication::Get().GetNormalExecutionAttribute();
+
+	Args.bRealtime = true;
+	Args.ConfigKey = ViewportKey0;
+	Args.ViewportType = LVT_Perspective;
+	TSharedRef<IViewportLayoutEntity> Viewport0 = LevelEditor.FactoryViewport(*ViewportType0, Args);
+
+	Args.bRealtime = false;
+	Args.ConfigKey = ViewportKey1;
+	Args.ViewportType = LVT_OrthoXY;
+	TSharedRef<IViewportLayoutEntity> Viewport1 = LevelEditor.FactoryViewport(*ViewportType1, Args);
+
+	// Front viewport
+	Args.bRealtime = false;
+	Args.ConfigKey = ViewportKey2;
+	Args.ViewportType = LVT_OrthoXZ;
+	TSharedRef<IViewportLayoutEntity> Viewport2 = LevelEditor.FactoryViewport(*ViewportType2, Args);
+
+	// Top Viewport
+	Args.bRealtime = false;
+	Args.ConfigKey = ViewportKey3;
+	Args.ViewportType = LVT_OrthoYZ;
+	TSharedRef<IViewportLayoutEntity> Viewport3 = LevelEditor.FactoryViewport(*ViewportType2, Args);
+
+	Viewports.Add( *ViewportKey0, Viewport0 );
+	Viewports.Add( *ViewportKey1, Viewport1 );
+	Viewports.Add( *ViewportKey2, Viewport2 );
+	Viewports.Add( *ViewportKey3, Viewport3 );
+
+	// Make newly-created perspective viewports active by default
+	GCurrentLevelEditingViewportClient = &Viewport0->GetLevelViewportClient();
+
 	TSharedRef<SWidget> LayoutWidget = MakeFourPanelWidget(
 		Viewports,
-		ViewportKey0, ViewportKey1, ViewportKey2, ViewportKey3,
+		Viewport0->AsWidget(), Viewport1->AsWidget(), Viewport2->AsWidget(), Viewport3->AsWidget(),
 		PrimarySplitterPercentage, SecondarySplitterPercentage0, SecondarySplitterPercentage1);
 
 	InitCommonLayoutFromString(SpecificLayoutString, *ViewportKey1);
@@ -121,26 +166,16 @@ void FLevelViewportLayoutFourPanes::ReplaceWidget( TSharedRef< SWidget > Source,
 
 TSharedRef<SWidget> FLevelViewportLayoutFourPanesLeft::MakeFourPanelWidget(
 	TMap<FName, TSharedPtr< IViewportLayoutEntity >>& ViewportWidgets,
-	const FString& ViewportKey0, const FString& ViewportKey1, const FString& ViewportKey2, const FString& ViewportKey3,
+	TSharedRef<SWidget> Viewport0, TSharedRef<SWidget> Viewport1, TSharedRef<SWidget> Viewport2, TSharedRef<SWidget> Viewport3,
 	float PrimarySplitterPercentage, float SecondarySplitterPercentage0, float SecondarySplitterPercentage1)
 {
-	TSharedPtr< SLevelViewport > ViewportWidget0;
-	TSharedPtr< SLevelViewport > ViewportWidget1;
-	TSharedPtr< SLevelViewport > ViewportWidget2;
-	TSharedPtr< SLevelViewport > ViewportWidget3;
-
 	TSharedRef<SWidget> Widget = 
 		SAssignNew( PrimarySplitterWidget, SSplitter )
 		.Orientation(EOrientation::Orient_Horizontal)
 		+SSplitter::Slot()
 		.Value(PrimarySplitterPercentage)
 		[
-			SAssignNew( ViewportWidget0, SLevelViewport )
-			.ViewportType( LVT_Perspective )	
-			.Realtime( true )
-			.ParentLayout( AsShared() )
-			.ParentLevelEditor( ParentLevelEditor )
-			.ConfigKey( ViewportKey0 )
+			Viewport0
 		]
 		+SSplitter::Slot()
 		.Value(1.0f - PrimarySplitterPercentage)
@@ -150,40 +185,20 @@ TSharedRef<SWidget> FLevelViewportLayoutFourPanesLeft::MakeFourPanelWidget(
 			+SSplitter::Slot()
 			.Value(SecondarySplitterPercentage0)
 			[
-				SAssignNew( ViewportWidget1, SLevelViewport )
-				.ViewportType( LVT_OrthoXY )
-				.ParentLayout( AsShared() )
-				.ParentLevelEditor( ParentLevelEditor )
-				.ConfigKey( ViewportKey1 )
+				Viewport1
 			]
 			+SSplitter::Slot()
 			.Value(SecondarySplitterPercentage1)
 			[
-				SAssignNew( ViewportWidget2, SLevelViewport )
-				.ViewportType( LVT_OrthoXZ )
-				.ParentLayout( AsShared() )
-				.ParentLevelEditor( ParentLevelEditor )
-				.ConfigKey( ViewportKey2 )
+				Viewport2
 			]
 			+SSplitter::Slot()
 			.Value(1.0f - SecondarySplitterPercentage0 - SecondarySplitterPercentage1)
 			[
-				SAssignNew( ViewportWidget3, SLevelViewport )
-				.ViewportType( LVT_OrthoYZ )
-				.ParentLayout( AsShared() )
-				.ParentLevelEditor( ParentLevelEditor )
-				.ConfigKey( ViewportKey3 )
+				Viewport3
 			]
 		];
 
-	ViewportWidgets.Add( *ViewportKey0, MakeShareable( new FLevelViewportLayoutEntity(ViewportWidget0.ToSharedRef()) ) );
-	ViewportWidgets.Add( *ViewportKey1, MakeShareable( new FLevelViewportLayoutEntity(ViewportWidget1.ToSharedRef()) ) );
-	ViewportWidgets.Add( *ViewportKey2, MakeShareable( new FLevelViewportLayoutEntity(ViewportWidget2.ToSharedRef()) ) );
-	ViewportWidgets.Add( *ViewportKey3, MakeShareable( new FLevelViewportLayoutEntity(ViewportWidget3.ToSharedRef()) ) );
-
-	// Make newly-created perspective viewports active by default
-	GCurrentLevelEditingViewportClient = &ViewportWidget0->GetLevelViewportClient();
-
 	return Widget;
 }
 
@@ -192,14 +207,9 @@ TSharedRef<SWidget> FLevelViewportLayoutFourPanesLeft::MakeFourPanelWidget(
 
 TSharedRef<SWidget> FLevelViewportLayoutFourPanesRight::MakeFourPanelWidget(
 	TMap<FName, TSharedPtr< IViewportLayoutEntity >>& ViewportWidgets,
-	const FString& ViewportKey0, const FString& ViewportKey1, const FString& ViewportKey2, const FString& ViewportKey3,
+	TSharedRef<SWidget> Viewport0, TSharedRef<SWidget> Viewport1, TSharedRef<SWidget> Viewport2, TSharedRef<SWidget> Viewport3,
 	float PrimarySplitterPercentage, float SecondarySplitterPercentage0, float SecondarySplitterPercentage1)
 {
-	TSharedPtr< SLevelViewport > ViewportWidget0;
-	TSharedPtr< SLevelViewport > ViewportWidget1;
-	TSharedPtr< SLevelViewport > ViewportWidget2;
-	TSharedPtr< SLevelViewport > ViewportWidget3;
-
 	TSharedRef<SWidget> Widget = 
 		SAssignNew( PrimarySplitterWidget, SSplitter )
 		.Orientation(EOrientation::Orient_Horizontal)
@@ -211,50 +221,25 @@ TSharedRef<SWidget> FLevelViewportLayoutFourPanesRight::MakeFourPanelWidget(
 			+SSplitter::Slot()
 			.Value(SecondarySplitterPercentage0)
 			[
-				SAssignNew( ViewportWidget1, SLevelViewport )
-				.ViewportType( LVT_OrthoXY )
-				.ParentLayout( AsShared() )
-				.ParentLevelEditor( ParentLevelEditor )
-				.ConfigKey( ViewportKey1 )
+				Viewport1
 			]
 			+SSplitter::Slot()
 			.Value(SecondarySplitterPercentage1)
 			[
-				SAssignNew( ViewportWidget2, SLevelViewport )
-				.ViewportType( LVT_OrthoXZ )
-				.ParentLayout( AsShared() )
-				.ParentLevelEditor( ParentLevelEditor )
-				.ConfigKey( ViewportKey2 )
+				Viewport2
 			]
 			+SSplitter::Slot()
 			.Value(1.0f - SecondarySplitterPercentage0 - SecondarySplitterPercentage1)
 			[
-				SAssignNew( ViewportWidget3, SLevelViewport )
-				.ViewportType( LVT_OrthoYZ )
-				.ParentLayout( AsShared() )
-				.ParentLevelEditor( ParentLevelEditor )
-				.ConfigKey( ViewportKey3 )
+				Viewport3
 			]
 		]
 		+SSplitter::Slot()
 		.Value(1.0f - PrimarySplitterPercentage)
 		[
-			SAssignNew( ViewportWidget0, SLevelViewport )
-			.ViewportType( LVT_Perspective )	
-			.Realtime( true )
-			.ParentLayout( AsShared() )
-			.ParentLevelEditor( ParentLevelEditor )
-			.ConfigKey( ViewportKey0 )
+			Viewport0
 		];
 
-	ViewportWidgets.Add( *ViewportKey0, MakeShareable( new FLevelViewportLayoutEntity(ViewportWidget0.ToSharedRef()) ) );
-	ViewportWidgets.Add( *ViewportKey1, MakeShareable( new FLevelViewportLayoutEntity(ViewportWidget1.ToSharedRef()) ) );
-	ViewportWidgets.Add( *ViewportKey2, MakeShareable( new FLevelViewportLayoutEntity(ViewportWidget2.ToSharedRef()) ) );
-	ViewportWidgets.Add( *ViewportKey3, MakeShareable( new FLevelViewportLayoutEntity(ViewportWidget3.ToSharedRef()) ) );
-
-	// Make newly-created perspective viewports active by default
-	GCurrentLevelEditingViewportClient = &ViewportWidget0->GetLevelViewportClient();
-
 	return Widget;
 }
 
@@ -263,26 +248,16 @@ TSharedRef<SWidget> FLevelViewportLayoutFourPanesRight::MakeFourPanelWidget(
 
 TSharedRef<SWidget> FLevelViewportLayoutFourPanesTop::MakeFourPanelWidget(
 	TMap<FName, TSharedPtr< IViewportLayoutEntity >>& ViewportWidgets,
-	const FString& ViewportKey0, const FString& ViewportKey1, const FString& ViewportKey2, const FString& ViewportKey3,
+	TSharedRef<SWidget> Viewport0, TSharedRef<SWidget> Viewport1, TSharedRef<SWidget> Viewport2, TSharedRef<SWidget> Viewport3,
 	float PrimarySplitterPercentage, float SecondarySplitterPercentage0, float SecondarySplitterPercentage1)
 {
-	TSharedPtr< SLevelViewport > ViewportWidget0;
-	TSharedPtr< SLevelViewport > ViewportWidget1;
-	TSharedPtr< SLevelViewport > ViewportWidget2;
-	TSharedPtr< SLevelViewport > ViewportWidget3;
-
 	TSharedRef<SWidget> Widget = 
 		SAssignNew( PrimarySplitterWidget, SSplitter )
 		.Orientation(EOrientation::Orient_Vertical)
 		+SSplitter::Slot()
 		.Value(PrimarySplitterPercentage)
 		[
-			SAssignNew( ViewportWidget0, SLevelViewport )
-			.ViewportType( LVT_Perspective )	
-			.Realtime( true )
-			.ParentLayout( AsShared() )
-			.ParentLevelEditor( ParentLevelEditor )
-			.ConfigKey( ViewportKey0 )
+			Viewport0
 		]
 		+SSplitter::Slot()
 		.Value(1.0f - PrimarySplitterPercentage)
@@ -292,40 +267,20 @@ TSharedRef<SWidget> FLevelViewportLayoutFourPanesTop::MakeFourPanelWidget(
 			+SSplitter::Slot()
 			.Value(SecondarySplitterPercentage0)
 			[
-				SAssignNew( ViewportWidget1, SLevelViewport )
-				.ViewportType( LVT_OrthoXY )
-				.ParentLayout( AsShared() )
-				.ParentLevelEditor( ParentLevelEditor )
-				.ConfigKey( ViewportKey1 )
+				Viewport1
 			]
 			+SSplitter::Slot()
 			.Value(SecondarySplitterPercentage1)
 			[
-				SAssignNew( ViewportWidget2, SLevelViewport )
-				.ViewportType( LVT_OrthoXZ )
-				.ParentLayout( AsShared() )
-				.ParentLevelEditor( ParentLevelEditor )
-				.ConfigKey( ViewportKey2 )
+				Viewport2
 			]
 			+SSplitter::Slot()
 			.Value(1.0f - SecondarySplitterPercentage0 - SecondarySplitterPercentage1)
 			[
-				SAssignNew( ViewportWidget3, SLevelViewport )
-				.ViewportType( LVT_OrthoYZ )
-				.ParentLayout( AsShared() )
-				.ParentLevelEditor( ParentLevelEditor )
-				.ConfigKey( ViewportKey3 )
+				Viewport3
 			]
 		];
 
-	ViewportWidgets.Add( *ViewportKey0, MakeShareable( new FLevelViewportLayoutEntity(ViewportWidget0.ToSharedRef()) ) );
-	ViewportWidgets.Add( *ViewportKey1, MakeShareable( new FLevelViewportLayoutEntity(ViewportWidget1.ToSharedRef()) ) );
-	ViewportWidgets.Add( *ViewportKey2, MakeShareable( new FLevelViewportLayoutEntity(ViewportWidget2.ToSharedRef()) ) );
-	ViewportWidgets.Add( *ViewportKey3, MakeShareable( new FLevelViewportLayoutEntity(ViewportWidget3.ToSharedRef()) ) );
-
-	// Make newly-created perspective viewports active by default
-	GCurrentLevelEditingViewportClient = &ViewportWidget0->GetLevelViewportClient();
-
 	return Widget;
 }
 
@@ -334,14 +289,9 @@ TSharedRef<SWidget> FLevelViewportLayoutFourPanesTop::MakeFourPanelWidget(
 
 TSharedRef<SWidget> FLevelViewportLayoutFourPanesBottom::MakeFourPanelWidget(
 	TMap<FName, TSharedPtr< IViewportLayoutEntity >>& ViewportWidgets,
-	const FString& ViewportKey0, const FString& ViewportKey1, const FString& ViewportKey2, const FString& ViewportKey3,
+	TSharedRef<SWidget> Viewport0, TSharedRef<SWidget> Viewport1, TSharedRef<SWidget> Viewport2, TSharedRef<SWidget> Viewport3,
 	float PrimarySplitterPercentage, float SecondarySplitterPercentage0, float SecondarySplitterPercentage1)
 {
-	TSharedPtr< SLevelViewport > ViewportWidget0;
-	TSharedPtr< SLevelViewport > ViewportWidget1;
-	TSharedPtr< SLevelViewport > ViewportWidget2;
-	TSharedPtr< SLevelViewport > ViewportWidget3;
-
 	TSharedRef<SWidget> Widget = 
 		SAssignNew( PrimarySplitterWidget, SSplitter )
 		.Orientation(EOrientation::Orient_Vertical)
@@ -353,49 +303,24 @@ TSharedRef<SWidget> FLevelViewportLayoutFourPanesBottom::MakeFourPanelWidget(
 			+SSplitter::Slot()
 			.Value(SecondarySplitterPercentage0)
 			[
-				SAssignNew( ViewportWidget1, SLevelViewport )
-				.ViewportType( LVT_OrthoXY )
-				.ParentLayout( AsShared() )
-				.ParentLevelEditor( ParentLevelEditor )
-				.ConfigKey( ViewportKey1 )
+				Viewport1
 			]
 			+SSplitter::Slot()
 			.Value(SecondarySplitterPercentage1)
 			[
-				SAssignNew( ViewportWidget2, SLevelViewport )
-				.ViewportType( LVT_OrthoXZ )
-				.ParentLayout( AsShared() )
-				.ParentLevelEditor( ParentLevelEditor )
-				.ConfigKey( ViewportKey2 )
+				Viewport2
 			]
 			+SSplitter::Slot()
 			.Value(1.0f - SecondarySplitterPercentage0 - SecondarySplitterPercentage1)
 			[
-				SAssignNew( ViewportWidget3, SLevelViewport )
-				.ViewportType( LVT_OrthoYZ )
-				.ParentLayout( AsShared() )
-				.ParentLevelEditor( ParentLevelEditor )
-				.ConfigKey( ViewportKey3 )
+				Viewport3
 			]
 		]
 		+SSplitter::Slot()
 		.Value(1.0f - PrimarySplitterPercentage)
 		[
-			SAssignNew( ViewportWidget0, SLevelViewport )
-			.ViewportType( LVT_Perspective )	
-			.Realtime( true )
-			.ParentLayout( AsShared() )
-			.ParentLevelEditor( ParentLevelEditor )
-			.ConfigKey( ViewportKey0 )
+			Viewport0
 		];
 
-	ViewportWidgets.Add( *ViewportKey0, MakeShareable( new FLevelViewportLayoutEntity(ViewportWidget0.ToSharedRef()) ) );
-	ViewportWidgets.Add( *ViewportKey1, MakeShareable( new FLevelViewportLayoutEntity(ViewportWidget1.ToSharedRef()) ) );
-	ViewportWidgets.Add( *ViewportKey2, MakeShareable( new FLevelViewportLayoutEntity(ViewportWidget2.ToSharedRef()) ) );
-	ViewportWidgets.Add( *ViewportKey3, MakeShareable( new FLevelViewportLayoutEntity(ViewportWidget3.ToSharedRef()) ) );
-
-	// Make newly-created perspective viewports active by default
-	GCurrentLevelEditingViewportClient = &ViewportWidget0->GetLevelViewportClient();
-
 	return Widget;
 }
diff --git a/Engine/Source/Editor/LevelEditor/Private/LevelViewportLayoutFourPanes.h b/Engine/Source/Editor/LevelEditor/Private/LevelViewportLayoutFourPanes.h
--- a/Engine/Source/Editor/LevelEditor/Private/LevelViewportLayoutFourPanes.h
+++ b/Engine/Source/Editor/LevelEditor/Private/LevelViewportLayoutFourPanes.h
@@ -23,7 +23,7 @@ protected:
 
 	virtual TSharedRef<SWidget> MakeFourPanelWidget(
 		TMap<FName, TSharedPtr< IViewportLayoutEntity >>& ViewportWidgets,
-		const FString& ViewportKey0, const FString& ViewportKey1, const FString& ViewportKey2, const FString& ViewportKey3,
+		TSharedRef<SWidget> Viewport0, TSharedRef<SWidget> Viewport1, TSharedRef<SWidget> Viewport2, TSharedRef<SWidget> Viewport3,
 		float PrimarySplitterPercentage, float SecondarySplitterPercentage0, float SecondarySplitterPercentage1) = 0;
 
 	/** Overridden from FLevelViewportLayout */
@@ -46,7 +46,7 @@ public:
 
 	virtual TSharedRef<SWidget> MakeFourPanelWidget(
 		TMap<FName, TSharedPtr< IViewportLayoutEntity >>& ViewportWidgets,
-		const FString& ViewportKey0, const FString& ViewportKey1, const FString& ViewportKey2, const FString& ViewportKey3,
+		TSharedRef<SWidget> Viewport0, TSharedRef<SWidget> Viewport1, TSharedRef<SWidget> Viewport2, TSharedRef<SWidget> Viewport3,
 		float PrimarySplitterPercentage, float SecondarySplitterPercentage0, float SecondarySplitterPercentage1) override;
 };
 
@@ -60,7 +60,7 @@ public:
 
 	virtual TSharedRef<SWidget> MakeFourPanelWidget(
 		TMap<FName, TSharedPtr< IViewportLayoutEntity >>& ViewportWidgets,
-		const FString& ViewportKey0, const FString& ViewportKey1, const FString& ViewportKey2, const FString& ViewportKey3,
+		TSharedRef<SWidget> Viewport0, TSharedRef<SWidget> Viewport1, TSharedRef<SWidget> Viewport2, TSharedRef<SWidget> Viewport3,
 		float PrimarySplitterPercentage, float SecondarySplitterPercentage0, float SecondarySplitterPercentage1) override;
 };
 
@@ -74,7 +74,7 @@ public:
 
 	virtual TSharedRef<SWidget> MakeFourPanelWidget(
 		TMap<FName, TSharedPtr< IViewportLayoutEntity >>& ViewportWidgets,
-		const FString& ViewportKey0, const FString& ViewportKey1, const FString& ViewportKey2, const FString& ViewportKey3,
+		TSharedRef<SWidget> Viewport0, TSharedRef<SWidget> Viewport1, TSharedRef<SWidget> Viewport2, TSharedRef<SWidget> Viewport3,
 		float PrimarySplitterPercentage, float SecondarySplitterPercentage0, float SecondarySplitterPercentage1) override;
 };
 
@@ -88,7 +88,7 @@ public:
 
 	virtual TSharedRef<SWidget> MakeFourPanelWidget(
 		TMap<FName, TSharedPtr< IViewportLayoutEntity >>& ViewportWidgets,
-		const FString& ViewportKey0, const FString& ViewportKey1, const FString& ViewportKey2, const FString& ViewportKey3,
+		TSharedRef<SWidget> Viewport0, TSharedRef<SWidget> Viewport1, TSharedRef<SWidget> Viewport2, TSharedRef<SWidget> Viewport3,
 		float PrimarySplitterPercentage, float SecondarySplitterPercentage0, float SecondarySplitterPercentage1) override;
 };
 
diff --git a/Engine/Source/Editor/LevelEditor/Private/LevelViewportLayoutOnePane.cpp b/Engine/Source/Editor/LevelEditor/Private/LevelViewportLayoutOnePane.cpp
--- a/Engine/Source/Editor/LevelEditor/Private/LevelViewportLayoutOnePane.cpp
+++ b/Engine/Source/Editor/LevelEditor/Private/LevelViewportLayoutOnePane.cpp
@@ -29,34 +29,42 @@ TSharedRef<SWidget> FLevelViewportLayoutOnePane::MakeViewportLayout(const FStrin
 	FEngineShowFlags PerspectiveShowFlags(ESFIM_Editor);
 	ApplyViewMode(VMI_Lit, true, PerspectiveShowFlags);
 
-	FString ViewportKey;
+	FString ViewportKey, ViewportType;
 	if (!SpecificLayoutString.IsEmpty())
 	{
+		const FString& IniSection = FLayoutSaveRestore::GetAdditionalLayoutConfigIni();
+
 		ViewportKey = SpecificLayoutString + TEXT(".Viewport0");
+		GConfig->GetString(*IniSection, *(ViewportKey + TEXT(".TypeWithinLayout")), ViewportType, GEditorPerProjectIni);
 	}
 
-	TSharedPtr<SLevelViewport> ViewportWidget;
+	FLevelEditorModule& LevelEditor = FModuleManager::LoadModuleChecked<FLevelEditorModule>("LevelEditor");
+
+	// Set up the viewport
+	FViewportConstructionArgs Args;
+	Args.ParentLayout = AsShared();
+	Args.ParentLevelEditor = ParentLevelEditor;
+	Args.IsEnabled = FSlateApplication::Get().GetNormalExecutionAttribute();
+	Args.bRealtime = true;
+	Args.ConfigKey = ViewportKey;
+	Args.ViewportType = LVT_Perspective;
+	TSharedRef<IViewportLayoutEntity> Viewport = LevelEditor.FactoryViewport(*ViewportType, Args);
 
 	ViewportBox =
 		SNew(SHorizontalBox)
 		+ SHorizontalBox::Slot()
 		[
-			SAssignNew(ViewportWidget, SLevelViewport)
-			.ViewportType(LVT_Perspective)
-			.Realtime(true)
-			.ParentLayout(AsShared())
-			.ParentLevelEditor(ParentLevelEditor)
-			.ConfigKey(ViewportKey)
+			Viewport->AsWidget()
 		];
 
-	Viewports.Add( *ViewportKey, MakeShareable(new FLevelViewportLayoutEntity(ViewportWidget.ToSharedRef())));
+	Viewports.Add( *ViewportKey, Viewport );
 
 	// Make newly-created perspective viewports active by default
-	GCurrentLevelEditingViewportClient = &ViewportWidget->GetLevelViewportClient();
+	GCurrentLevelEditingViewportClient = &Viewport->GetLevelViewportClient();
 
 	InitCommonLayoutFromString(SpecificLayoutString, NAME_None);
 
-	return StaticCastSharedRef<SWidget>(ViewportBox.ToSharedRef());
+	return ViewportBox.ToSharedRef();
 }
 
 void FLevelViewportLayoutOnePane::ReplaceWidget(TSharedRef<SWidget> Source, TSharedRef<SWidget> Replacement)
diff --git a/Engine/Source/Editor/LevelEditor/Private/LevelViewportLayoutThreePanes.cpp b/Engine/Source/Editor/LevelEditor/Private/LevelViewportLayoutThreePanes.cpp
--- a/Engine/Source/Editor/LevelEditor/Private/LevelViewportLayoutThreePanes.cpp
+++ b/Engine/Source/Editor/LevelEditor/Private/LevelViewportLayoutThreePanes.cpp
@@ -24,6 +24,7 @@ TSharedRef<SWidget> FLevelViewportLayoutThreePanes::MakeViewportLayout(const FSt
 	ApplyViewMode(VMI_Lit, true, PerspectiveShowFlags);
 
 	FString ViewportKey0, ViewportKey1, ViewportKey2;
+	FString ViewportType0 = TEXT("Default"), ViewportType1 = TEXT("Default"), ViewportType2 = TEXT("Default");
 	float PrimarySplitterPercentage = ViewportLayoutThreePanesDefs::DefaultSplitterPercentage;
 	float SecondarySplitterPercentage = ViewportLayoutThreePanesDefs::DefaultSplitterPercentage;
 	
@@ -36,6 +37,10 @@ TSharedRef<SWidget> FLevelViewportLayoutThreePanes::MakeViewportLayout(const FSt
 		ViewportKey1 = SpecificLayoutString + TEXT(".Viewport1");
 		ViewportKey2 = SpecificLayoutString + TEXT(".Viewport2");
 
+		GConfig->GetString(*IniSection, *(ViewportKey0 + TEXT(".TypeWithinLayout")), ViewportType0, GEditorPerProjectIni);
+		GConfig->GetString(*IniSection, *(ViewportKey1 + TEXT(".TypeWithinLayout")), ViewportType1, GEditorPerProjectIni);
+		GConfig->GetString(*IniSection, *(ViewportKey2 + TEXT(".TypeWithinLayout")), ViewportType2, GEditorPerProjectIni);
+
 		FString PercentageString;
 		if(GConfig->GetString(*IniSection, *(SpecificLayoutString + TEXT(".Percentage0")), PercentageString, GEditorPerProjectIni))
 		{
@@ -47,9 +52,37 @@ TSharedRef<SWidget> FLevelViewportLayoutThreePanes::MakeViewportLayout(const FSt
 		}	
 	}
 
+	FLevelEditorModule& LevelEditor = FModuleManager::LoadModuleChecked<FLevelEditorModule>("LevelEditor");
+
+	// Set up the viewports
+	FViewportConstructionArgs Args;
+	Args.ParentLayout = AsShared();
+	Args.ParentLevelEditor = ParentLevelEditor;
+	Args.IsEnabled = FSlateApplication::Get().GetNormalExecutionAttribute();
+
+	Args.bRealtime = true;
+	Args.ConfigKey = ViewportKey0;
+	Args.ViewportType = LVT_Perspective;
+	TSharedRef<IViewportLayoutEntity> Viewport0 = LevelEditor.FactoryViewport(*ViewportType0, Args);
+
+	Args.bRealtime = false;
+	Args.ConfigKey = ViewportKey1;
+	Args.ViewportType = LVT_OrthoXY;
+	TSharedRef<IViewportLayoutEntity> Viewport1 = LevelEditor.FactoryViewport(*ViewportType1, Args);
+
+	// Front viewport
+	Args.bRealtime = false;
+	Args.ConfigKey = ViewportKey2;
+	Args.ViewportType = LVT_OrthoXZ;
+	TSharedRef<IViewportLayoutEntity> Viewport2 = LevelEditor.FactoryViewport(*ViewportType2, Args);
+
+	Viewports.Add( *ViewportKey0, Viewport0 );
+	Viewports.Add( *ViewportKey1, Viewport1 );
+	Viewports.Add( *ViewportKey2, Viewport2 );
+
 	TSharedRef<SWidget> LayoutWidget = MakeThreePanelWidget(
 		Viewports,
-		ViewportKey0, ViewportKey1, ViewportKey2,
+		Viewport0->AsWidget(), Viewport1->AsWidget(), Viewport2->AsWidget(),
 		PrimarySplitterPercentage, SecondarySplitterPercentage);
 
 	InitCommonLayoutFromString(SpecificLayoutString, *ViewportKey1);
@@ -116,25 +149,16 @@ void FLevelViewportLayoutThreePanes::ReplaceWidget( TSharedRef< SWidget > Source
 
 TSharedRef<SWidget> FLevelViewportLayoutThreePanesLeft::MakeThreePanelWidget(
 	TMap< FName, TSharedPtr< IViewportLayoutEntity >>& ViewportWidgets,
-	const FString& ViewportKey0, const FString& ViewportKey1, const FString& ViewportKey2,
+	const TSharedRef<SWidget>& Viewport0, const TSharedRef<SWidget>& Viewport1, const TSharedRef<SWidget>& Viewport2,
 	float PrimarySplitterPercentage, float SecondarySplitterPercentage)
 {
-	TSharedPtr< SLevelViewport > ViewportWidget0;
-	TSharedPtr< SLevelViewport > ViewportWidget1;
-	TSharedPtr< SLevelViewport > ViewportWidget2;
-
 	TSharedRef<SWidget> Widget = 
 		SAssignNew( PrimarySplitterWidget, SSplitter )
 		.Orientation(EOrientation::Orient_Horizontal)
 		+SSplitter::Slot()
 		.Value(PrimarySplitterPercentage)
 		[
-			SAssignNew( ViewportWidget0, SLevelViewport )
-			.ViewportType( LVT_Perspective )
-			.Realtime( true )
-			.ParentLayout( AsShared() )
-			.ParentLevelEditor( ParentLevelEditor )
-			.ConfigKey( ViewportKey0 )
+			Viewport0
 		]
 		+SSplitter::Slot()
 		.Value(1.0f - PrimarySplitterPercentage)
@@ -144,30 +168,15 @@ TSharedRef<SWidget> FLevelViewportLayoutThreePanesLeft::MakeThreePanelWidget(
 			+SSplitter::Slot()
 			.Value(SecondarySplitterPercentage)
 			[
-				SAssignNew( ViewportWidget1, SLevelViewport )
-				.ViewportType( LVT_OrthoXY )
-				.ParentLayout( AsShared() )
-				.ParentLevelEditor( ParentLevelEditor )
-				.ConfigKey( ViewportKey1 )
+				Viewport1
 			]
 			+SSplitter::Slot()
 			.Value(1.0f - SecondarySplitterPercentage)
 			[
-				SAssignNew( ViewportWidget2, SLevelViewport )
-				.ViewportType( LVT_OrthoXZ )
-				.ParentLayout( AsShared() )
-				.ParentLevelEditor( ParentLevelEditor )
-				.ConfigKey( ViewportKey2 )
+				Viewport2
 			]
 		];
 
-	ViewportWidgets.Add( *ViewportKey0, MakeShareable( new FLevelViewportLayoutEntity(ViewportWidget0.ToSharedRef()) ) );
-	ViewportWidgets.Add( *ViewportKey1, MakeShareable( new FLevelViewportLayoutEntity(ViewportWidget1.ToSharedRef()) ) );
-	ViewportWidgets.Add( *ViewportKey2, MakeShareable( new FLevelViewportLayoutEntity(ViewportWidget2.ToSharedRef()) ) );
-
-	// Make newly-created perspective viewports active by default
-	GCurrentLevelEditingViewportClient = &ViewportWidget0->GetLevelViewportClient();
-
 	return Widget;
 }
 
@@ -176,13 +185,9 @@ TSharedRef<SWidget> FLevelViewportLayoutThreePanesLeft::MakeThreePanelWidget(
 
 TSharedRef<SWidget> FLevelViewportLayoutThreePanesRight::MakeThreePanelWidget(
 	TMap< FName, TSharedPtr< IViewportLayoutEntity >>& ViewportWidgets,
-	const FString& ViewportKey0, const FString& ViewportKey1, const FString& ViewportKey2,
+	const TSharedRef<SWidget>& Viewport0, const TSharedRef<SWidget>& Viewport1, const TSharedRef<SWidget>& Viewport2,
 	float PrimarySplitterPercentage, float SecondarySplitterPercentage)
 {
-	TSharedPtr< SLevelViewport > ViewportWidget0;
-	TSharedPtr< SLevelViewport > ViewportWidget1;
-	TSharedPtr< SLevelViewport > ViewportWidget2;
-
 	TSharedRef<SWidget> Widget = 
 		SAssignNew( PrimarySplitterWidget, SSplitter )
 		.Orientation(EOrientation::Orient_Horizontal)
@@ -194,40 +199,20 @@ TSharedRef<SWidget> FLevelViewportLayoutThreePanesRight::MakeThreePanelWidget(
 			+SSplitter::Slot()
 			.Value(SecondarySplitterPercentage)
 			[
-				SAssignNew( ViewportWidget1, SLevelViewport )
-				.ViewportType( LVT_OrthoXY )
-				.ParentLayout( AsShared() )
-				.ParentLevelEditor( ParentLevelEditor )
-				.ConfigKey( ViewportKey1 )
+				Viewport1
 			]
 			+SSplitter::Slot()
 			.Value(1.0f - SecondarySplitterPercentage)
 			[
-				SAssignNew( ViewportWidget2, SLevelViewport )
-				.ViewportType( LVT_OrthoXZ )
-				.ParentLayout( AsShared() )
-				.ParentLevelEditor( ParentLevelEditor )
-				.ConfigKey( ViewportKey2 )
+				Viewport2
 			]
 		]
 		+SSplitter::Slot()
 		.Value(1.0f - PrimarySplitterPercentage)
 		[
-			SAssignNew( ViewportWidget0, SLevelViewport )
-			.ViewportType( LVT_Perspective )	
-			.Realtime( true )
-			.ParentLayout( AsShared() )
-			.ParentLevelEditor( ParentLevelEditor )
-			.ConfigKey( ViewportKey0 )
+			Viewport0
 		];
 
-	ViewportWidgets.Add( *ViewportKey0, MakeShareable( new FLevelViewportLayoutEntity(ViewportWidget0.ToSharedRef()) ) );
-	ViewportWidgets.Add( *ViewportKey1, MakeShareable( new FLevelViewportLayoutEntity(ViewportWidget1.ToSharedRef()) ) );
-	ViewportWidgets.Add( *ViewportKey2, MakeShareable( new FLevelViewportLayoutEntity(ViewportWidget2.ToSharedRef()) ) );
-
-	// Make newly-created perspective viewports active by default
-	GCurrentLevelEditingViewportClient = &ViewportWidget0->GetLevelViewportClient();
-
 	return Widget;
 }
 
@@ -236,25 +221,16 @@ TSharedRef<SWidget> FLevelViewportLayoutThreePanesRight::MakeThreePanelWidget(
 
 TSharedRef<SWidget> FLevelViewportLayoutThreePanesTop::MakeThreePanelWidget(
 	TMap< FName, TSharedPtr< IViewportLayoutEntity >>& ViewportWidgets,
-	const FString& ViewportKey0, const FString& ViewportKey1, const FString& ViewportKey2,
+	const TSharedRef<SWidget>& Viewport0, const TSharedRef<SWidget>& Viewport1, const TSharedRef<SWidget>& Viewport2,
 	float PrimarySplitterPercentage, float SecondarySplitterPercentage)
 {
-	TSharedPtr< SLevelViewport > ViewportWidget0;
-	TSharedPtr< SLevelViewport > ViewportWidget1;
-	TSharedPtr< SLevelViewport > ViewportWidget2;
-
 	TSharedRef<SWidget> Widget = 
 		SAssignNew( PrimarySplitterWidget, SSplitter )
 		.Orientation(EOrientation::Orient_Vertical)
 		+SSplitter::Slot()
 		.Value(PrimarySplitterPercentage)
 		[
-			SAssignNew( ViewportWidget0, SLevelViewport )
-			.ViewportType( LVT_Perspective )	
-			.Realtime( true )
-			.ParentLayout( AsShared() )
-			.ParentLevelEditor( ParentLevelEditor )
-			.ConfigKey( ViewportKey0 )
+			Viewport0
 		]
 		+SSplitter::Slot()
 		.Value(1.0f - PrimarySplitterPercentage)
@@ -264,30 +240,15 @@ TSharedRef<SWidget> FLevelViewportLayoutThreePanesTop::MakeThreePanelWidget(
 			+SSplitter::Slot()
 			.Value(SecondarySplitterPercentage)
 			[
-				SAssignNew( ViewportWidget1, SLevelViewport )
-				.ViewportType( LVT_OrthoXY )
-				.ParentLayout( AsShared() )
-				.ParentLevelEditor( ParentLevelEditor )
-				.ConfigKey( ViewportKey1 )
+				Viewport1
 			]
 			+SSplitter::Slot()
 			.Value(1.0f - SecondarySplitterPercentage)
 			[
-				SAssignNew( ViewportWidget2, SLevelViewport )
-				.ViewportType( LVT_OrthoXZ )
-				.ParentLayout( AsShared() )
-				.ParentLevelEditor( ParentLevelEditor )
-				.ConfigKey( ViewportKey2 )
+				Viewport2
 			]
 		];
 
-	ViewportWidgets.Add( *ViewportKey0, MakeShareable( new FLevelViewportLayoutEntity(ViewportWidget0.ToSharedRef()) ) );
-	ViewportWidgets.Add( *ViewportKey1, MakeShareable( new FLevelViewportLayoutEntity(ViewportWidget1.ToSharedRef()) ) );
-	ViewportWidgets.Add( *ViewportKey2, MakeShareable( new FLevelViewportLayoutEntity(ViewportWidget2.ToSharedRef()) ) );
-
-	// Make newly-created perspective viewports active by default
-	GCurrentLevelEditingViewportClient = &ViewportWidget0->GetLevelViewportClient();
-
 	return Widget;
 }
 
@@ -296,13 +257,9 @@ TSharedRef<SWidget> FLevelViewportLayoutThreePanesTop::MakeThreePanelWidget(
 
 TSharedRef<SWidget> FLevelViewportLayoutThreePanesBottom::MakeThreePanelWidget(
 	TMap< FName, TSharedPtr< IViewportLayoutEntity >>& ViewportWidgets,
-	const FString& ViewportKey0, const FString& ViewportKey1, const FString& ViewportKey2,
+	const TSharedRef<SWidget>& Viewport0, const TSharedRef<SWidget>& Viewport1, const TSharedRef<SWidget>& Viewport2,
 	float PrimarySplitterPercentage, float SecondarySplitterPercentage)
 {
-	TSharedPtr< SLevelViewport > ViewportWidget0;
-	TSharedPtr< SLevelViewport > ViewportWidget1;
-	TSharedPtr< SLevelViewport > ViewportWidget2;
-
 	TSharedRef<SWidget> Widget = 
 		SAssignNew( PrimarySplitterWidget, SSplitter )
 		.Orientation(EOrientation::Orient_Vertical)
@@ -314,39 +271,19 @@ TSharedRef<SWidget> FLevelViewportLayoutThreePanesBottom::MakeThreePanelWidget(
 			+SSplitter::Slot()
 			.Value(SecondarySplitterPercentage)
 			[
-				SAssignNew( ViewportWidget1, SLevelViewport )
-				.ViewportType( LVT_OrthoXY )
-				.ParentLayout( AsShared() )
-				.ParentLevelEditor( ParentLevelEditor )
-				.ConfigKey( ViewportKey1 )
+				Viewport1
 			]
 			+SSplitter::Slot()
 			.Value(1.0f - SecondarySplitterPercentage)
 			[
-				SAssignNew( ViewportWidget2, SLevelViewport )
-				.ViewportType( LVT_OrthoXZ )
-				.ParentLayout( AsShared() )
-				.ParentLevelEditor( ParentLevelEditor )
-				.ConfigKey( ViewportKey2 )
+				Viewport2
 			]
 		]
 		+SSplitter::Slot()
 		.Value(1.0f - PrimarySplitterPercentage)
 		[
-			SAssignNew( ViewportWidget0, SLevelViewport )
-			.ViewportType( LVT_Perspective )
-			.Realtime( true )
-			.ParentLayout( AsShared() )
-			.ParentLevelEditor( ParentLevelEditor )
-			.ConfigKey( ViewportKey0 )
+			Viewport0
 		];
 
-	ViewportWidgets.Add( *ViewportKey0, MakeShareable( new FLevelViewportLayoutEntity(ViewportWidget0.ToSharedRef()) ) );
-	ViewportWidgets.Add( *ViewportKey1, MakeShareable( new FLevelViewportLayoutEntity(ViewportWidget1.ToSharedRef()) ) );
-	ViewportWidgets.Add( *ViewportKey2, MakeShareable( new FLevelViewportLayoutEntity(ViewportWidget2.ToSharedRef()) ) );
-
-	// Make newly-created perspective viewports active by default
-	GCurrentLevelEditingViewportClient = &ViewportWidget0->GetLevelViewportClient();
-
 	return Widget;
 }
diff --git a/Engine/Source/Editor/LevelEditor/Private/LevelViewportLayoutThreePanes.h b/Engine/Source/Editor/LevelEditor/Private/LevelViewportLayoutThreePanes.h
--- a/Engine/Source/Editor/LevelEditor/Private/LevelViewportLayoutThreePanes.h
+++ b/Engine/Source/Editor/LevelEditor/Private/LevelViewportLayoutThreePanes.h
@@ -23,7 +23,7 @@ protected:
 
 	virtual TSharedRef<SWidget> MakeThreePanelWidget(
 		TMap< FName, TSharedPtr< IViewportLayoutEntity >>& ViewportWidgets,
-		const FString& ViewportKey0, const FString& ViewportKey1, const FString& ViewportKey2,
+		const TSharedRef<SWidget>& ViewportKey0, const TSharedRef<SWidget>& ViewportKey1, const TSharedRef<SWidget>& ViewportKey2,
 		float PrimarySplitterPercentage, float SecondarySplitterPercentage) = 0;
 
 	/** Overridden from FLevelViewportLayout */
@@ -45,7 +45,7 @@ public:
 
 	virtual TSharedRef<SWidget> MakeThreePanelWidget(
 		TMap< FName, TSharedPtr< IViewportLayoutEntity >>& ViewportWidgets,
-		const FString& ViewportKey0, const FString& ViewportKey1, const FString& ViewportKey2,
+		const TSharedRef<SWidget>& ViewportKey0, const TSharedRef<SWidget>& ViewportKey1, const TSharedRef<SWidget>& ViewportKey2,
 		float PrimarySplitterPercentage, float SecondarySplitterPercentage) override;
 };
 
@@ -59,7 +59,7 @@ public:
 
 	virtual TSharedRef<SWidget> MakeThreePanelWidget(
 		TMap< FName, TSharedPtr< IViewportLayoutEntity >>& ViewportWidgets,
-		const FString& ViewportKey0, const FString& ViewportKey1, const FString& ViewportKey2,
+		const TSharedRef<SWidget>& ViewportKey0, const TSharedRef<SWidget>& ViewportKey1, const TSharedRef<SWidget>& ViewportKey2,
 		float PrimarySplitterPercentage, float SecondarySplitterPercentage) override;
 };
 
@@ -73,7 +73,7 @@ public:
 
 	virtual TSharedRef<SWidget> MakeThreePanelWidget(
 		TMap< FName, TSharedPtr< IViewportLayoutEntity >>& ViewportWidgets,
-		const FString& ViewportKey0, const FString& ViewportKey1, const FString& ViewportKey2,
+		const TSharedRef<SWidget>& ViewportKey0, const TSharedRef<SWidget>& ViewportKey1, const TSharedRef<SWidget>& ViewportKey2,
 		float PrimarySplitterPercentage, float SecondarySplitterPercentage) override;
 };
 
@@ -87,7 +87,7 @@ public:
 
 	virtual TSharedRef<SWidget> MakeThreePanelWidget(
 		TMap< FName, TSharedPtr< IViewportLayoutEntity >>& ViewportWidgets,
-		const FString& ViewportKey0, const FString& ViewportKey1, const FString& ViewportKey2,
+		const TSharedRef<SWidget>& ViewportKey0, const TSharedRef<SWidget>& ViewportKey1, const TSharedRef<SWidget>& ViewportKey2,
 		float PrimarySplitterPercentage, float SecondarySplitterPercentage) override;
 };
 
diff --git a/Engine/Source/Editor/LevelEditor/Private/LevelViewportLayoutTwoPanes.cpp b/Engine/Source/Editor/LevelEditor/Private/LevelViewportLayoutTwoPanes.cpp
--- a/Engine/Source/Editor/LevelEditor/Private/LevelViewportLayoutTwoPanes.cpp
+++ b/Engine/Source/Editor/LevelEditor/Private/LevelViewportLayoutTwoPanes.cpp
@@ -18,16 +18,8 @@ TSharedRef<SWidget> TLevelViewportLayoutTwoPanes<TOrientation>::MakeViewportLayo
 {
 	FString SpecificLayoutString = GetTypeSpecificLayoutString(LayoutString);
 
-	TSharedPtr< SLevelViewport > ViewportWidget0;
-	TSharedPtr< SLevelViewport > ViewportWidget1;
-
-	FEngineShowFlags OrthoShowFlags(ESFIM_Editor);	
-	ApplyViewMode(VMI_BrushWireframe, false, OrthoShowFlags);
-
-	FEngineShowFlags PerspectiveShowFlags(ESFIM_Editor);	
-	ApplyViewMode(VMI_Lit, true, PerspectiveShowFlags);
-
 	FString ViewportKey0, ViewportKey1;
+	FString ViewportType0, ViewportType1;
 	float SplitterPercentage = ViewportLayoutTwoPanesDefs::DefaultSplitterPercentage;
 
 	if (!SpecificLayoutString.IsEmpty())
@@ -38,43 +30,54 @@ TSharedRef<SWidget> TLevelViewportLayoutTwoPanes<TOrientation>::MakeViewportLayo
 		ViewportKey0 = SpecificLayoutString + TEXT(".Viewport0");
 		ViewportKey1 = SpecificLayoutString + TEXT(".Viewport1");
 
+		GConfig->GetString(*IniSection, *(ViewportKey0 + TEXT(".TypeWithinLayout")), ViewportType0, GEditorPerProjectIni);
+		GConfig->GetString(*IniSection, *(ViewportKey1 + TEXT(".TypeWithinLayout")), ViewportType1, GEditorPerProjectIni);
 
 		FString PercentageString;
 		if(GConfig->GetString(*IniSection, *(SpecificLayoutString + TEXT(".Percentage")), PercentageString, GEditorPerProjectIni))
 		{
 			TTypeFromString<float>::FromString(SplitterPercentage, *PercentageString);
-		}						
+		}
 	}
 
+	FLevelEditorModule& LevelEditor = FModuleManager::LoadModuleChecked<FLevelEditorModule>("LevelEditor");
+
+	// Set up the viewports
+	FViewportConstructionArgs Args;
+	Args.ParentLayout = AsShared();
+	Args.ParentLevelEditor = ParentLevelEditor;
+	Args.IsEnabled = FSlateApplication::Get().GetNormalExecutionAttribute();
+
+	Args.bRealtime = false;
+	Args.ConfigKey = ViewportKey0;
+	Args.ViewportType = LVT_OrthoXY;
+	TSharedRef<IViewportLayoutEntity> Viewport0 = LevelEditor.FactoryViewport(*ViewportType0, Args);
+
+	Args.bRealtime = true;
+	Args.ConfigKey = ViewportKey1;
+	Args.ViewportType = LVT_Perspective;
+	TSharedRef<IViewportLayoutEntity> Viewport1 = LevelEditor.FactoryViewport(*ViewportType1, Args);
+
+	Viewports.Add( *ViewportKey0, Viewport0 );
+	Viewports.Add( *ViewportKey1, Viewport1 );
+
+	// Make newly-created perspective viewports active by default
+	GCurrentLevelEditingViewportClient = &Viewport0->GetLevelViewportClient();
+
 	SplitterWidget = 
 		SNew( SSplitter )
 		.Orientation(TOrientation)
 		+SSplitter::Slot()
 		.Value(SplitterPercentage)
 		[
-			SAssignNew( ViewportWidget0, SLevelViewport )
-			.ViewportType( LVT_OrthoXY )	
-			.ParentLayout( AsShared() )
-			.ParentLevelEditor( ParentLevelEditor )
-			.ConfigKey( ViewportKey0 )
+			Viewport0->AsWidget()
 		]
 	+SSplitter::Slot()
 		.Value(1.0f - SplitterPercentage)
 		[
-			SAssignNew( ViewportWidget1, SLevelViewport )
-			.ViewportType( LVT_Perspective )
-			.Realtime( true )
-			.ParentLayout( AsShared() )
-			.ParentLevelEditor( ParentLevelEditor )
-			.ConfigKey( ViewportKey1 )
+			Viewport1->AsWidget()
 		];
 
-	Viewports.Add( *ViewportKey0, MakeShareable( new FLevelViewportLayoutEntity(ViewportWidget0.ToSharedRef()) ) );
-	Viewports.Add( *ViewportKey1, MakeShareable( new FLevelViewportLayoutEntity(ViewportWidget1.ToSharedRef()) ) );
-
-	// Make newly-created perspective viewports active by default
-	GCurrentLevelEditingViewportClient = &ViewportWidget1->GetLevelViewportClient();
-
 	InitCommonLayoutFromString(SpecificLayoutString, *ViewportKey1);
 
 	return SplitterWidget.ToSharedRef();
diff --git a/Engine/Source/Editor/LevelEditor/Private/LevelViewportTabContent.cpp b/Engine/Source/Editor/LevelEditor/Private/LevelViewportTabContent.cpp
--- a/Engine/Source/Editor/LevelEditor/Private/LevelViewportTabContent.cpp
+++ b/Engine/Source/Editor/LevelEditor/Private/LevelViewportTabContent.cpp
@@ -30,11 +30,6 @@ TSharedPtr< class FLevelViewportLayout > FLevelViewportTabContent::ConstructView
 	else if (TypeName == LevelViewportConfigurationNames::FourPanesBottom) ViewportLayout = MakeShareable(new FLevelViewportLayoutFourPanesBottom);
 	else if (TypeName == LevelViewportConfigurationNames::FourPanesTop) ViewportLayout = MakeShareable(new FLevelViewportLayoutFourPanesTop);
 	else if (TypeName == LevelViewportConfigurationNames::OnePane) ViewportLayout = MakeShareable(new FLevelViewportLayoutOnePane);
-	else
-	{
-		FLevelEditorModule& LevelEditorModule = FModuleManager::Get().GetModuleChecked<FLevelEditorModule>("LevelEditor");
-		ViewportLayout = LevelEditorModule.FactoryCustomViewportLayout(TypeName);
-	}
 
 	if (!ensure(ViewportLayout.IsValid()))
 	{
@@ -44,7 +39,7 @@ TSharedPtr< class FLevelViewportLayout > FLevelViewportTabContent::ConstructView
 	return ViewportLayout;
 }
 
-void FLevelViewportTabContent::Initialize(TSharedPtr<class SLevelEditor> InParentLevelEditor, TSharedPtr<SDockTab> InParentTab, const FString& InLayoutString)
+void FLevelViewportTabContent::Initialize(TSharedPtr<ILevelEditor> InParentLevelEditor, TSharedPtr<SDockTab> InParentTab, const FString& InLayoutString)
 {
 	ParentTab = InParentTab;
 	ParentLevelEditor = InParentLevelEditor;
diff --git a/Engine/Source/Editor/LevelEditor/Private/SLevelViewport.cpp b/Engine/Source/Editor/LevelEditor/Private/SLevelViewport.cpp
--- a/Engine/Source/Editor/LevelEditor/Private/SLevelViewport.cpp
+++ b/Engine/Source/Editor/LevelEditor/Private/SLevelViewport.cpp
@@ -288,7 +288,14 @@ void SLevelViewport::ConstructViewportOverlayContent()
 
 void SLevelViewport::ConstructLevelEditorViewportClient( const FArguments& InArgs )
 {
-	LevelViewportClient = MakeShareable( new FLevelEditorViewportClient(SharedThis(this)) );
+	if (InArgs._LevelEditorViewportClient.IsValid())
+	{
+		LevelViewportClient = InArgs._LevelEditorViewportClient;
+	}
+	else
+	{
+		LevelViewportClient = MakeShareable( new FLevelEditorViewportClient(SharedThis(this)) );
+	}
 
 	// Default level viewport client values for settings that could appear in layout config ini
 	FLevelEditorViewportInstanceSettings ViewportInstanceSettings;
@@ -386,6 +393,8 @@ void SLevelViewport::ConstructLevelEditorViewportClient( const FArguments& InArg
 	LevelViewportClient->EngineShowFlags.SetCompositeEditorPrimitives(true);
 
 	LevelViewportClient->SetViewModes(ViewportInstanceSettings.PerspViewModeIndex, ViewportInstanceSettings.OrthoViewModeIndex );
+
+	bShowFullToolbar = ViewportInstanceSettings.bShowFullToolbar;
 }
 
 const FSceneViewport* SLevelViewport::GetGameSceneViewport() const
@@ -417,13 +426,14 @@ FReply SLevelViewport::OnKeyDown( const FGeometry& MyGeometry, const FKeyEvent&
 		//  the SLevelEditor is not a direct parent of the viewport.  
 		if ( this->IsImmersive() && !Reply.IsEventHandled() )
 		{
-			TSharedPtr<SLevelEditor> ParentLevelEditorSharedPtr = ParentLevelEditor.Pin();
+			TSharedPtr<ILevelEditor> ParentLevelEditorSharedPtr = ParentLevelEditor.Pin();
 			if( ParentLevelEditorSharedPtr.IsValid() )
 			{
 				Reply = ParentLevelEditorSharedPtr->OnKeyDownInViewport( MyGeometry, InKeyEvent );
 			}
 		}
 	}
+
 	return Reply;
 }
 
@@ -1188,7 +1198,14 @@ void SLevelViewport::BindOptionCommands( FUICommandList& CommandList )
 
 	CommandList.MapAction(
 		ViewportActions.ClearAllBookMarks,
-		FExecuteAction::CreateSP( this, &SLevelViewport::OnClearAllBookMarks )		
+		FExecuteAction::CreateSP( this, &SLevelViewport::OnClearAllBookMarks )
+		);
+
+	CommandList.MapAction(
+		ViewportActions.ToggleViewportToolbar,
+		FExecuteAction::CreateSP( this, &SLevelViewport::OnToggleShowFullToolbar ),
+		FCanExecuteAction(),
+		FIsActionChecked::CreateSP( this, &SLevelViewport::ShouldShowFullToolbar )
 		);
 }
 
@@ -1286,6 +1303,18 @@ void SLevelViewport::BindViewCommands( FUICommandList& CommandList )
 		FCanExecuteAction(),
 		FIsActionChecked::CreateSP( this, &SLevelViewport::IsViewportConfigurationSet, LevelViewportConfigurationNames::FourPanes2x2 ));
 
+	auto ProcessViewportTypeActions = [&](FName InViewportTypeName, const FViewportTypeDefinition& InDefinition){
+		if (InDefinition.ToggleCommand.IsValid())
+		{
+			CommandList.MapAction(InDefinition.ToggleCommand, FUIAction(
+				FExecuteAction::CreateSP(this, &SLevelViewport::ToggleViewportTypeActivationWithinLayout, InViewportTypeName),
+				FCanExecuteAction(),
+				FIsActionChecked::CreateSP(this, &SLevelViewport::IsViewportTypeWithinLayoutEqual, InViewportTypeName)
+			));
+		}
+	};
+	FLevelEditorModule& LevelEditorModule = FModuleManager::GetModuleChecked<FLevelEditorModule>(TEXT("LevelEditor"));
+	LevelEditorModule.IterateViewportTypes(ProcessViewportTypeActions);
 
 	// Map Buffer visualization mode actions
 	for (FLevelViewportCommands::TBufferVisualizationModeCommandMap::TConstIterator It = ViewportActions.BufferVisualizationModeCommands.CreateConstIterator(); It; ++It)
@@ -1983,6 +2012,7 @@ void SLevelViewport::SaveConfig(const FString& ConfigName)
 	ViewportInstanceSettings.bIsRealtime = LevelViewportClient->IsRealtime();
 	ViewportInstanceSettings.bShowStats = LevelViewportClient->ShouldShowStats();
 	ViewportInstanceSettings.FarViewPlane = LevelViewportClient->GetFarClipPlaneOverride();
+	ViewportInstanceSettings.bShowFullToolbar = bShowFullToolbar;
 
 	if (GetDefault<ULevelEditorViewportSettings>()->bSaveEngineStats)
 	{
@@ -3314,6 +3344,50 @@ bool SLevelViewport::IsViewportConfigurationSet(FName ConfigurationName) const
 	return false;
 }
 
+FName SLevelViewport::GetViewportTypeWithinLayout() const
+{
+	TSharedPtr<FLevelViewportLayout> LayoutPinned = ParentLayout.Pin();
+	if (LayoutPinned.IsValid() && !ConfigKey.IsEmpty())
+	{
+		TSharedPtr<IViewportLayoutEntity> Entity = LayoutPinned->GetViewports().FindRef(*ConfigKey);
+		if (Entity.IsValid())
+		{
+			return Entity->GetType();
+		}
+	}
+	return "Default";
+}
+
+void SLevelViewport::SetViewportTypeWithinLayout(FName InLayoutType)
+{
+	TSharedPtr<FLevelViewportLayout> LayoutPinned = ParentLayout.Pin();
+	if (LayoutPinned.IsValid() && !ConfigKey.IsEmpty())
+	{
+		const FString& IniSection = FLayoutSaveRestore::GetAdditionalLayoutConfigIni();
+		GConfig->SetString( *IniSection, *( ConfigKey + TEXT(".TypeWithinLayout") ), *InLayoutType.ToString(), GEditorPerProjectIni );
+
+		// Force a refresh of the tab content
+		OnSetViewportConfiguration(LayoutPinned->GetLayoutTypeName());
+	}
+}
+
+void SLevelViewport::ToggleViewportTypeActivationWithinLayout(FName InLayoutType)
+{
+	if (GetViewportTypeWithinLayout() == InLayoutType)
+	{
+		SetViewportTypeWithinLayout("Default");
+	}
+	else
+	{
+		SetViewportTypeWithinLayout(InLayoutType);
+	}
+}
+
+bool SLevelViewport::IsViewportTypeWithinLayoutEqual(FName InLayoutType)
+{
+	return GetViewportTypeWithinLayout() == InLayoutType;
+}
+
 void SLevelViewport::StartPlayInEditorSession(UGameViewportClient* PlayClient, const bool bInSimulateInEditor)
 {
 	check( !HasPlayInEditorViewport() );
