- Fixed missing project menus

diff --git a/Engine/Source/Editor/EditorStyle/Private/EditorStyleClasses.cpp b/Engine/Source/Editor/EditorStyle/Private/EditorStyleClasses.cpp
--- a/Engine/Source/Editor/EditorStyle/Private/EditorStyleClasses.cpp
+++ b/Engine/Source/Editor/EditorStyle/Private/EditorStyleClasses.cpp
@@ -16,6 +16,9 @@ UEditorStyleSettings::UEditorStyleSettings( const FObjectInitializer& ObjectInit
 
 	bShowFriendlyNames = true;
 	LogTimestampMode = ELogTimes::None;
+	
+	bShowProjectMenus = true;
+	bShowLaunchMenus = true;
 }
 
 FLinearColor UEditorStyleSettings::GetSubduedSelectionColor() const
