- Backed out the first bandaid fix for missing menu items (this will make it easier to verify that any changes to my second, better, band-aid are working)

diff --git a/Engine/Source/Editor/EditorStyle/Private/EditorStyleClasses.cpp b/Engine/Source/Editor/EditorStyle/Private/EditorStyleClasses.cpp
--- a/Engine/Source/Editor/EditorStyle/Private/EditorStyleClasses.cpp
+++ b/Engine/Source/Editor/EditorStyle/Private/EditorStyleClasses.cpp
@@ -16,9 +16,6 @@ UEditorStyleSettings::UEditorStyleSettings( const FObjectInitializer& ObjectInit
 
 	bShowFriendlyNames = true;
 	LogTimestampMode = ELogTimes::None;
-	
-	bShowProjectMenus = true;
-	bShowLaunchMenus = true;
 }
 
 FLinearColor UEditorStyleSettings::GetSubduedSelectionColor() const
