Style files include what they use (resolves Mac NonUnity CIS issues caused by [REDACTED_PERSON] including property customzation headers)

diff --git a/Engine/Source/Editor/EditorStyle/Public/EditorStyleSet.h b/Engine/Source/Editor/EditorStyle/Public/EditorStyleSet.h
--- a/Engine/Source/Editor/EditorStyle/Public/EditorStyleSet.h
+++ b/Engine/Source/Editor/EditorStyle/Public/EditorStyleSet.h
@@ -1,6 +1,11 @@
 // Copyright 1998-2016 Epic Games, Inc. All Rights Reserved.
 #pragma once
 
+#include "ISlateStyle.h"
+#include "SlateColor.h"
+#include "SlateDynamicImageBrush.h"
+#include "SlateFontInfo.h"
+
 /**
  * A collection of named properties that guide the appearance of Slate.
  */
diff --git a/Engine/Source/Editor/UnrealEd/Public/Toolkits/AssetEditorToolkit.h b/Engine/Source/Editor/UnrealEd/Public/Toolkits/AssetEditorToolkit.h
--- a/Engine/Source/Editor/UnrealEd/Public/Toolkits/AssetEditorToolkit.h
+++ b/Engine/Source/Editor/UnrealEd/Public/Toolkits/AssetEditorToolkit.h
@@ -3,6 +3,7 @@
 #pragma once
 
 #include "BaseToolkit.h"
+#include "LayoutService.h"
 #include "TabManager.h"
 #include "Toolkits/AssetEditorManager.h"		// For IAssetEditorInstance derive
 
diff --git a/Engine/Source/Runtime/Slate/Public/Framework/Docking/LayoutService.h b/Engine/Source/Runtime/Slate/Public/Framework/Docking/LayoutService.h
--- a/Engine/Source/Runtime/Slate/Public/Framework/Docking/LayoutService.h
+++ b/Engine/Source/Runtime/Slate/Public/Framework/Docking/LayoutService.h
@@ -2,6 +2,7 @@
 
 #pragma once
 
+#include "TabManager.h"
 
 struct SLATE_API FLayoutSaveRestore
 {
diff --git a/Engine/Source/Runtime/Slate/Public/Framework/Docking/TabManager.h b/Engine/Source/Runtime/Slate/Public/Framework/Docking/TabManager.h
--- a/Engine/Source/Runtime/Slate/Public/Framework/Docking/TabManager.h
+++ b/Engine/Source/Runtime/Slate/Public/Framework/Docking/TabManager.h
@@ -4,6 +4,7 @@
 
 #include "SlateIcon.h"
 #include "WorkspaceItem.h"
+#include "SlateEnums.h"
 
 class FMenuBuilder;
 class FMultiBox;
