Fix ddc warnings

diff --git a/Engine/Source/Editor/EditorStyle/Private/SlateEditorStyle.cpp b/Engine/Source/Editor/EditorStyle/Private/SlateEditorStyle.cpp
--- a/Engine/Source/Editor/EditorStyle/Private/SlateEditorStyle.cpp
+++ b/Engine/Source/Editor/EditorStyle/Private/SlateEditorStyle.cpp
@@ -6863,7 +6863,7 @@ void FSlateEditorStyle::FStyle::SetupAutomationStyles()
 
 		Set("SettingsEditor.CheckoutWarningBorder", new BOX_BRUSH( "Common/GroupBorderLight", FMargin(4.0f/16.0f) ) );
 
-		Set("SettingsEditor.CatgoryAndSectionFont", TTF_CORE_FONT("Slate/Fonts/Roboto-Regular", 18));
+		Set("SettingsEditor.CatgoryAndSectionFont", TTF_CORE_FONT("Fonts/Roboto-Regular", 18));
 		Set("SettingsEditor.TopLevelObjectFontStyle", TTF_CORE_FONT("Fonts/Roboto-Bold", 12));
 	}
 
