VR Editor: Fixed blurry toolbar icon when "Use Small Icons" preference was enabled (UE-32208)
- Also, added a higher resolution reference image for this icon (not used yet)

diff --git a/Engine/Source/Editor/EditorStyle/Private/SlateEditorStyle.cpp b/Engine/Source/Editor/EditorStyle/Private/SlateEditorStyle.cpp
--- a/Engine/Source/Editor/EditorStyle/Private/SlateEditorStyle.cpp
+++ b/Engine/Source/Editor/EditorStyle/Private/SlateEditorStyle.cpp
@@ -4603,7 +4603,7 @@ void FSlateEditorStyle::FStyle::SetupLevelEditorStyle()
 		Set( "LevelEditor.EditMatinee.Small", new IMAGE_BRUSH( "Icons/icon_matinee_40x", Icon20x20 ) );
 
 		Set( "LevelEditor.ToggleVR", new IMAGE_BRUSH( "Icons/icon_ToggleVREditor_40x", Icon40x40 ) );
-		Set( "LevelEditor.ToggleVR.Small", new IMAGE_BRUSH( "Icons/icon_ToggleVREditor_16x", Icon40x40 ) );
+		Set( "LevelEditor.ToggleVR.Small", new IMAGE_BRUSH( "Icons/icon_ToggleVREditor_40x", Icon20x20 ) );
 
 		Set( "MergeActors.MeshMergingTool", new IMAGE_BRUSH( "Icons/icon_MergeActors_MeshMerging_40x", Icon40x40 ) );
 		Set( "MergeActors.MeshProxyTool", new IMAGE_BRUSH( "Icons/icon_MergeActors_MeshProxy_40x", Icon40x40 ) );
