Cleaning up slate styles

diff --git a/Engine/Source/Editor/EditorStyle/Private/SlateEditorStyle.cpp b/Engine/Source/Editor/EditorStyle/Private/SlateEditorStyle.cpp
--- a/Engine/Source/Editor/EditorStyle/Private/SlateEditorStyle.cpp
+++ b/Engine/Source/Editor/EditorStyle/Private/SlateEditorStyle.cpp
@@ -4624,10 +4624,10 @@ void FSlateEditorStyle::FStyle::SetupLevelEditorStyle()
 		Set( "LevelEditor.FoliageMode.Selected", new IMAGE_BRUSH( "Icons/icon_Mode_Foliage_selected_40x", Icon40x40 ) );
 		Set( "LevelEditor.FoliageMode.Selected.Small", new IMAGE_BRUSH( "Icons/icon_Mode_Foliage_selected_40x", Icon20x20 ) );
 
-		Set( "LevelEditor.BspMode", new IMAGE_BRUSH( "Icons/icon_Mode_MeshEditor_40px", Icon40x40 ) );
-		Set( "LevelEditor.BspMode.Small", new IMAGE_BRUSH( "Icons/icon_Mode_MeshEditor_40px", Icon20x20 ) );
-		Set( "LevelEditor.BspMode.Selected", new IMAGE_BRUSH( "Icons/icon_Mode_MeshEditor_40px-a_40px", Icon40x40 ) );
-		Set( "LevelEditor.BspMode.Selected.Small", new IMAGE_BRUSH( "Icons/icon_Mode_MeshEditor_40px-a_40px", Icon20x20 ) );
+		Set("LevelEditor.BspMode", new IMAGE_BRUSH("Icons/icon_Mode_GeoEdit_40px", Icon40x40));
+		Set("LevelEditor.BspMode.Small", new IMAGE_BRUSH("Icons/icon_Mode_GeoEdit_40px", Icon20x20));
+		Set("LevelEditor.BspMode.Selected", new IMAGE_BRUSH("Icons/icon_Mode_GeoEdit-a_40px", Icon40x40));
+		Set("LevelEditor.BspMode.Selected.Small", new IMAGE_BRUSH("Icons/icon_Mode_GeoEdit-a_40px", Icon20x20));
 
 		Set( "LevelEditor.WorldProperties", new IMAGE_BRUSH( "Icons/icon_worldscript_40x", Icon40x40 ) );
 		Set( "LevelEditor.WorldProperties.Small", new IMAGE_BRUSH( "Icons/icon_worldscript_40x", Icon20x20 ) );
diff --git a/Engine/Source/Editor/VREditor/VREditorStyle.cpp b/Engine/Source/Editor/VREditor/VREditorStyle.cpp
--- a/Engine/Source/Editor/VREditor/VREditorStyle.cpp
+++ b/Engine/Source/Editor/VREditor/VREditorStyle.cpp
@@ -177,10 +177,6 @@ TSharedRef< FSlateStyleSet > FVREditorStyle::Create()
 	Style->Set("VREditorStyle.FoliageMode", new IMAGE_BRUSH("Icons/VREditor/T_Radial_Foliage", Icon512x512));
 	Style->Set("VREditorStyle.LandscapeMode", new IMAGE_BRUSH("Icons/VREditor/T_Radial_Landscape", Icon512x512));
 	Style->Set("VREditorStyle.MeshPaintMode", new IMAGE_BRUSH("Icons/VREditor/T_Radial_Mesh_Paint", Icon512x512));
-	Style->Set("VREditorStyle.MeshEditMode", new IMAGE_BRUSH("Icons/VREditor/T_Radial_Mesh_All", Icon512x512));
-	Style->Set("VREditorStyle.PolygonEditMode", new IMAGE_BRUSH("Icons/VREditor/T_Radial_Mesh_Poly", Icon512x512));
-	Style->Set("VREditorStyle.EdgeEditMode", new IMAGE_BRUSH("Icons/VREditor/T_Radial_Mesh_Edge", Icon512x512));
-	Style->Set("VREditorStyle.VertexEditMode", new IMAGE_BRUSH("Icons/VREditor/T_Radial_Mesh_Vertex", Icon512x512));
 
 	Style->Set("VREditorStyle.Copy", new IMAGE_BRUSH("Icons/VREditor/T_Radial_Copy", Icon512x512));
 	Style->Set("VREditorStyle.Cut", new IMAGE_BRUSH("Icons/VREditor/T_Radial_Cut", Icon512x512));
