Copying //UE4/Dev-AnimPhys to //UE4/Dev-Main (Source: //UE4/Dev-AnimPhys @ 3624379)

diff --git a/Engine/Source/Editor/EditorStyle/Private/SlateEditorStyle.cpp b/Engine/Source/Editor/EditorStyle/Private/SlateEditorStyle.cpp
--- a/Engine/Source/Editor/EditorStyle/Private/SlateEditorStyle.cpp
+++ b/Engine/Source/Editor/EditorStyle/Private/SlateEditorStyle.cpp
@@ -1999,6 +1999,7 @@ void FSlateEditorStyle::FStyle::SetupGeneralStyles()
 		Set("FoliageEditToolBar.SToolBarButtonBlock.Padding", FMargin(0.f));
 		Set("FoliageEditToolBar.SToolBarCheckComboButtonBlock.Padding", FMargin(4.0f));
 		Set("FoliageEditToolBar.SToolBarButtonBlock.CheckBox.Padding", FMargin(10.0f, 6.f));
+		Set("FoliageEditToolBar.SToolBarButtonBlock.Button.Padding", FMargin(0.0f));
 		Set("FoliageEditToolBar.SToolBarComboButtonBlock.ComboButton.Color", DefaultForeground);
 
 		Set("FoliageEditToolBar.Block.IndentedPadding", FMargin(18.0f, 2.0f, 4.0f, 4.0f));
@@ -2248,6 +2249,7 @@ void FSlateEditorStyle::FStyle::SetupGeneralStyles()
 		Set( "ToolBar.SToolBarButtonBlock.Padding", FMargin(4.0f));
 		Set( "ToolBar.SToolBarCheckComboButtonBlock.Padding", FMargin(4.0f));
 		Set( "ToolBar.SToolBarButtonBlock.CheckBox.Padding", FMargin(0.0f) );
+		Set( "ToolBar.SToolBarButtonBlock.Button.Padding", FMargin(0.0f) );
 		Set( "ToolBar.SToolBarComboButtonBlock.ComboButton.Color", DefaultForeground );
 
 		Set( "ToolBar.Block.IndentedPadding", FMargin( 18.0f, 2.0f, 4.0f, 4.0f ) );
@@ -2370,6 +2372,7 @@ void FSlateEditorStyle::FStyle::SetupGeneralStyles()
 		Set( "Menu.SToolBarButtonBlock.Padding", FMargin(4.0f));
 		Set( "Menu.SToolBarCheckComboButtonBlock.Padding", FMargin(4.0f));
 		Set( "Menu.SToolBarButtonBlock.CheckBox.Padding", FMargin(0.0f) );
+		Set( "Menu.SToolBarButtonBlock.Button.Padding", FMargin(0.0f));
 		Set( "Menu.SToolBarComboButtonBlock.ComboButton.Color", DefaultForeground );
 
 		Set( "Menu.Block.IndentedPadding", FMargin( 18.0f, 2.0f, 4.0f, 4.0f ) );
@@ -2379,12 +2382,14 @@ void FSlateEditorStyle::FStyle::SetupGeneralStyles()
 		Set( "Menu.Separator.Padding", FMargin( 0.5f ) );
 
 		Set( "Menu.Label", FTextBlockStyle(NormalText) .SetFont( TTF_CORE_FONT( "Fonts/Roboto-Regular", 9 ) ) );
+		Set( "Menu.Label.Padding", FMargin(0.0f, 0.0f, 0.0f, 0.0f) );
+		Set( "Menu.Label.ContentPadding", FMargin(10.0f, 2.0f) );
 		Set( "Menu.EditableText", FEditableTextBoxStyle(NormalEditableTextBoxStyle) .SetFont( TTF_CORE_FONT( "Fonts/Roboto-Regular", 9 ) ) );
 		Set( "Menu.Keybinding", FTextBlockStyle(NormalText) .SetFont( TTF_CORE_FONT( "Fonts/Roboto-Regular", 8 ) ) );
 
 		Set( "Menu.Heading", FTextBlockStyle(NormalText)
-			.SetFont( TTF_CORE_FONT( "Fonts/Roboto-Regular", 8 ) )
-			.SetColorAndOpacity( FLinearColor( 0.4f, 0.4, 0.4f, 1.0f ) ) );
+			.SetFont(TTF_CORE_FONT("Fonts/Roboto-Regular", 8))
+			.SetColorAndOpacity(FLinearColor(0.4f, 0.4, 0.4f, 1.0f)));
 
 		/* Set images for various SCheckBox states associated with menu check box items... */
 		const FCheckBoxStyle BasicMenuCheckBoxStyle = FCheckBoxStyle()
@@ -2459,6 +2464,9 @@ void FSlateEditorStyle::FStyle::SetupGeneralStyles()
 
 		/* The style of a menu bar button when it has a sub menu open */
 		Set( "Menu.Button.SubMenuOpen", new BORDER_BRUSH( "Common/Selection", FMargin(4.f/16.f), FLinearColor(0.10f, 0.10f, 0.10f) ) );
+
+		/* The style of a menu bar button when it has a sub menu closed */
+		Set("Menu.Button.SubMenuClosed", new FSlateNoResource());
 	}
 
 	// ViewportLayoutToolbar
@@ -2484,6 +2492,7 @@ void FSlateEditorStyle::FStyle::SetupGeneralStyles()
 
 		Set( "ViewportLayoutToolbar.SToolBarButtonBlock.Padding", FMargin(4.0f) );
 		Set( "ViewportLayoutToolbar.SToolBarButtonBlock.CheckBox.Padding", FMargin(0.0f) );
+		Set( "ViewportLayoutToolbar.SToolBarButtonBlock.Button.Padding", FMargin(0.0f) );
 		Set( "ViewportLayoutToolbar.SToolBarComboButtonBlock.ComboButton.Color", DefaultForeground );
 		}
 
@@ -2566,6 +2575,7 @@ void FSlateEditorStyle::FStyle::SetupGeneralStyles()
 		Set( "NotificationBar.Button.Checked_Pressed", new BOX_BRUSH( "Common/RoundedSelection_16x",  4.0f/16.0f, SelectionColor ) );
 
 		Set( "NotificationBar.SToolBarButtonBlock.CheckBox.Padding", FMargin(4.0f) );
+		Set( "NotificationBar.SToolBarButtonBlock.Button.Padding", FMargin(0.0f) );
 		Set( "NotificationBar.SToolBarComboButtonBlock.ComboButton.Color", DefaultForeground );
 	}
 
@@ -2579,12 +2589,17 @@ void FSlateEditorStyle::FStyle::SetupGeneralStyles()
 		Set( "ViewportMenu.SToolBarButtonBlock.Padding", FMargin(0));
 		Set( "ViewportMenu.SToolBarCheckComboButtonBlock.Padding", FMargin(0));
 		Set( "ViewportMenu.SToolBarButtonBlock.CheckBox.Padding", FMargin(4.0f) );
+		Set( "ViewportMenu.SToolBarButtonBlock.Button.Padding", FMargin(3.0f) );
 		Set( "ViewportMenu.SToolBarComboButtonBlock.ComboButton.Color", FLinearColor(0.f,0.f,0.f,0.75f) );
 
 		Set( "ViewportMenu.Separator", new BOX_BRUSH( "Old/Button", 8.0f/32.0f, FLinearColor::Transparent ) );
 		Set( "ViewportMenu.Separator.Padding", FMargin( 100.0f ) );
 
-		Set( "ViewportMenu.Label", FTextBlockStyle(NormalText) .SetFont( TTF_CORE_FONT( "Fonts/Roboto-Bold", 9) ) );
+		Set( "ViewportMenu.Label", FTextBlockStyle(NormalText)
+			.SetFont(TTF_CORE_FONT("Fonts/Roboto-Bold", 9))
+			.SetColorAndOpacity(FLinearColor(0.0f, 0.0f, 0.0f, 1.0f)));
+		Set( "ViewportMenu.Label.Padding", FMargin(0.0f, 0.0f, 3.0f, 0.0f) );
+		Set( "ViewportMenu.Label.ContentPadding", FMargin(5.0f, 2.0f) );
 		Set( "ViewportMenu.EditableText", FEditableTextBoxStyle(NormalEditableTextBoxStyle) .SetFont( TTF_CORE_FONT( "Fonts/Roboto-Regular", 9) ) );
 		Set( "ViewportMenu.Keybinding", FTextBlockStyle(NormalText) .SetFont( TTF_CORE_FONT( "Fonts/Roboto-Regular", 8) ) );
 
@@ -2593,7 +2608,7 @@ void FSlateEditorStyle::FStyle::SetupGeneralStyles()
 
 		Set( "ViewportMenu.Heading.Font", TTF_CORE_FONT( "Fonts/Roboto-Regular", 8 ) );
 		Set( "ViewportMenu.Heading.ColorAndOpacity", FLinearColor( 0.4f, 0.4, 0.4f, 1.0f ) );
-	
+
 		const FCheckBoxStyle ViewportMenuCheckBoxCheckBoxStyle = FCheckBoxStyle()
 			.SetUncheckedImage(			IMAGE_BRUSH( "Common/SmallCheckBox", Icon14x14 ) )
 			.SetUncheckedPressedImage(	IMAGE_BRUSH( "Common/SmallCheckBox_Hovered", Icon14x14, FLinearColor( 0.5f, 0.5f, 0.5f ) ) )
@@ -2678,15 +2693,10 @@ void FSlateEditorStyle::FStyle::SetupGeneralStyles()
 		/* ... and add new style */
 		Set( "ViewportMenu.ToggleButton.End", ViewportMenuToggleEndButtonStyle );
 
-		const FMargin NormalPadding = FMargin(4.0f,4.0f,4.0f,4.0f);
-		const FMargin PressedPadding = FMargin(4.0f,4.0f,4.0f,4.0f);
-
 		const FButtonStyle ViewportMenuButton = FButtonStyle(Button)
 			.SetNormal ( BOX_BRUSH( *SmallRoundedButton, 7.0f/16.0f, NormalColor))
 			.SetPressed( BOX_BRUSH( *SmallRoundedButton, 7.0f/16.0f, PressedColor ) )
-			.SetHovered( BOX_BRUSH( *SmallRoundedButton, 7.0f/16.0f, PressedColor ) )
-			.SetPressedPadding(PressedPadding)
-			.SetNormalPadding(NormalPadding);
+			.SetHovered( BOX_BRUSH( *SmallRoundedButton, 7.0f/16.0f, PressedColor ) );
 
 		Set( "ViewportMenu.Button", ViewportMenuButton );
 
@@ -2707,6 +2717,12 @@ void FSlateEditorStyle::FStyle::SetupGeneralStyles()
 			.SetPressed( BOX_BRUSH( *SmallRoundedButtonEnd, 7.0f/16.0f, PressedColor ) )
 			.SetHovered( BOX_BRUSH( *SmallRoundedButtonEnd, 7.0f/16.0f, PressedColor ) )
 			);
+
+		/* The style of a menu bar button when it has a sub menu open */
+		Set("ViewportMenu.Button.SubMenuOpen", new FSlateNoResource());
+
+		/* The style of a menu bar button when it has a sub menu closed */
+		Set("ViewportMenu.Button.SubMenuClosed", new FSlateNoResource());
 	}
 
 	// Viewport actor preview's pin/unpin buttons
@@ -4961,7 +4977,7 @@ void FSlateEditorStyle::FStyle::SetupLevelEditorStyle()
 		/* ... and set new style */
 		Set( "LevelViewportToolBar.CheckBoxButton", EditorViewportToolBarButton );
 
-		Set( "EditorViewportToolBar.MenuDropdown", new IMAGE_BRUSH( "Common/ComboArrow", Icon8x8 ) );
+		Set( "EditorViewportToolBar.MenuDropdown", new IMAGE_BRUSH( "Common/ComboArrow", Icon8x8, FLinearColor(0.0f, 0.0f, 0.0f) ) );
 		Set( "LevelViewportToolBar.Maximize.Normal", new IMAGE_BRUSH( "Old/LevelViewportToolBar/Maximized_Unchecked", Icon16x16 ) );
 		Set( "LevelViewportToolBar.Maximize.Checked", new IMAGE_BRUSH( "Old/LevelViewportToolBar/Maximized_Checked", Icon16x16 ) );
 		Set( "LevelViewportToolBar.RestoreFromImmersive.Normal", new IMAGE_BRUSH( "Icons/icon_RestoreFromImmersive_16px", Icon16x16 ) );
@@ -5055,6 +5071,7 @@ void FSlateEditorStyle::FStyle::SetupLevelEditorStyle()
 		Set( "EditorModesToolbar.SToolBarComboButtonBlock.ComboButton.Color", DefaultForeground );
 		Set( "EditorModesToolbar.SToolBarButtonBlock.Padding", FMargin( 1.0f, 0.0f, 0.0f, 0 ) );
 		Set( "EditorModesToolbar.SToolBarButtonBlock.CheckBox.Padding", FMargin( 6.0f, 4.0f, 6.0f, 6.0f ) );
+		Set( "EditorModesToolbar.SToolBarButtonBlock.Button.Padding", FMargin( 0.0f ) );
 		Set( "EditorModesToolbar.SToolBarCheckComboButtonBlock.Padding", FMargin( 0 ) );
 
 		Set( "EditorModesToolbar.Block.IndentedPadding", FMargin( 0 ) );
@@ -6631,85 +6648,195 @@ void FSlateEditorStyle::FStyle::SetupToolkitStyles()
 		Set("RichCurveEditor.StraightenTangents.Small", new IMAGE_BRUSH("Icons/icon_CurveEditor_Straighten_40x", Icon20x20));
 	}
 
-	// PhAT
-	{
-		Set( "PhAT.Tabs.Properties", new IMAGE_BRUSH( "/Icons/icon_tab_SelectionDetails_16x", Icon16x16 ) );
-		Set( "PhAT.Tabs.Hierarchy", new IMAGE_BRUSH( "/Icons/levels_16x", Icon16x16 ) );
-
-		Set( "PhAT.EditingMode_Body", new IMAGE_BRUSH( "/Icons/icon_PHatMode_Body_40x", Icon40x40) );
-		Set( "PhAT.EditingMode_Constraint", new IMAGE_BRUSH( "/Icons/icon_PHatMode_Joint_40x", Icon40x40) );
-
-		Set( "PhAT.EditingMode_Body.Small", new IMAGE_BRUSH( "/Icons/icon_PHatMode_Body_40x", Icon20x20) );
-		Set( "PhAT.EditingMode_Constraint.Small", new IMAGE_BRUSH( "/Icons/icon_PHatMode_Joint_40x", Icon20x20) );
-
-		Set( "PhAT.SimulationNormal", new IMAGE_BRUSH( "Icons/icon_PhAT_PlaySim_40x", Icon40x40 ) );
-		Set( "PhAT.SimulationNoGravity", new IMAGE_BRUSH( "Icons/icon_PhAT_PlaySimNoGravity_40x", Icon40x40 ) );
-		Set( "PhAT.ToggleSelectedSimulation", new IMAGE_BRUSH( "Icons/icon_PhAT_PlaySimSelected_40x", Icon40x40 ) );
-		Set( "PhAT.Undo", new IMAGE_BRUSH( "Icons/icon_Generic_Undo_40x", Icon40x40 ) );
-		Set( "PhAT.Redo", new IMAGE_BRUSH( "Icons/icon_Generic_Redo_40x", Icon40x40 ) );
-		Set( "PhAT.ChangeDefaultMesh", new IMAGE_BRUSH( "Icons/icon_PhAT_Mesh_40x", Icon40x40 ) );
-		Set( "PhAT.ResetEntireAsset", new IMAGE_BRUSH( "Icons/icon_PhAT_ResetAsset_40x", Icon40x40 ) );
-		Set( "PhAT.RestetBoneCollision", new IMAGE_BRUSH( "Icons/icon_PhAT_ResetCollision_40x", Icon40x40 ) );
-		Set( "PhAT.ApplyPhysicalMaterial", new IMAGE_BRUSH( "Icons/icon_PhAT_PhysMat_40x", Icon40x40 ) );
-		Set( "PhAT.CopyJointSettings", new IMAGE_BRUSH( "Icons/icon_PhAT_CopyJoints_40x", Icon40x40 ) );
-		Set( "PhAT.PlayAnimation", new IMAGE_BRUSH( "Icons/icon_PhAT_Play_40x", Icon40x40 ) );
-		Set( "PhAT.PhATTranslationMode", new IMAGE_BRUSH( "Icons/icon_translate_40x", Icon40x40 ) );
-		Set( "PhAT.PhATRotationMode", new IMAGE_BRUSH( "Icons/icon_rotate_40x", Icon40x40 ) );
-		Set( "PhAT.PhATScaleMode", new IMAGE_BRUSH( "Icons/icon_scale_40x", Icon40x40 ) );
-		Set( "PhAT.Snap", new IMAGE_BRUSH( "Icons/icon_PhAT_Snap_40x", Icon40x40 ) );
-		Set( "PhAT.CopyProperties", new IMAGE_BRUSH( "Icons/icon_PhAT_CopyProperties_40x", Icon40x40 ) );
-		Set( "PhAT.DisableCollision", new IMAGE_BRUSH( "Icons/icon_PhAT_DisableCollision_40x", Icon40x40 ) );
-		Set( "PhAT.EnableCollision", new IMAGE_BRUSH( "Icons/icon_PhAT_EnableCollision_40x", Icon40x40 ) );
-		Set( "PhAT.WeldToBody", new IMAGE_BRUSH( "Icons/icon_PhAT_Weld_40x", Icon40x40 ) );
-		Set( "PhAT.AddNewBody", new IMAGE_BRUSH( "Icons/icon_PhAT_NewBody_40x", Icon40x40 ) );
-		Set( "PhAT.AddSphere", new IMAGE_BRUSH( "Icons/icon_PhAT_Sphere_40x", Icon40x40 ) );
-		Set( "PhAT.AddSphyl", new IMAGE_BRUSH( "Icons/icon_PhAT_Sphyl_40x", Icon40x40 ) );
-		Set( "PhAT.AddBox", new IMAGE_BRUSH( "Icons/icon_PhAT_Box_40x", Icon40x40 ) );
-		Set( "PhAT.DeletePrimitive", new IMAGE_BRUSH( "Icons/icon_PhAT_DeletePrimitive_40x", Icon40x40 ) );
-		Set( "PhAT.DuplicatePrimitive", new IMAGE_BRUSH( "Icons/icon_PhAT_DupePrim_40x", Icon40x40 ) );
-		Set( "PhAT.ResetConstraint", new IMAGE_BRUSH( "Icons/icon_PhAT_ResetConstraint_40x", Icon40x40 ) );
-		Set( "PhAT.SnapConstraint", new IMAGE_BRUSH( "Icons/icon_PhAT_SnapConstraint_40x", Icon40x40 ) );
-		Set( "PhAT.SnapAllConstraints", new IMAGE_BRUSH( "Icons/icon_PhAT_SnapAll_40x", Icon40x40 ) );
-		Set( "PhAT.ConvertToBallAndSocket", new IMAGE_BRUSH( "Icons/icon_PhAT_Ball_40x", Icon40x40 ) );
-		Set( "PhAT.ConvertToHinge", new IMAGE_BRUSH( "Icons/icon_PhAT_Hinge_40x", Icon40x40 ) );
-		Set( "PhAT.ConvertToPrismatic", new IMAGE_BRUSH( "Icons/icon_PhAT_Prismatic_40x", Icon40x40 ) );
-		Set( "PhAT.ConvertToSkeletal", new IMAGE_BRUSH( "Icons/icon_PhAT_Skeletal_40x", Icon40x40 ) );
-		Set( "PhAT.DeleteConstraint", new IMAGE_BRUSH( "Icons/icon_PhAT_DeleteConstraint_40x", Icon40x40 ) );
-
-		Set("PhAT.SimulationNormal.Small", new IMAGE_BRUSH("Icons/icon_PhAT_PlaySim_40x", Icon20x20));
-		Set("PhAT.SimulationNoGravity.Small", new IMAGE_BRUSH("Icons/icon_PhAT_PlaySimNoGravity_40x", Icon20x20));
-		Set("PhAT.ToggleSelectedSimulation.Small", new IMAGE_BRUSH("Icons/icon_PhAT_PlaySimSelected_40x", Icon20x20));
-		Set( "PhAT.ToggleSimulation.Small", new IMAGE_BRUSH( "Icons/icon_PhAT_PlaySim_40x", Icon20x20 ) );
-		Set( "PhAT.Undo.Small", new IMAGE_BRUSH( "Icons/icon_Generic_Undo_40x", Icon20x20 ) );
-		Set( "PhAT.Redo.Small", new IMAGE_BRUSH( "Icons/icon_Generic_Redo_40x", Icon20x20 ) );
-		Set( "PhAT.ChangeDefaultMesh.Small", new IMAGE_BRUSH( "Icons/icon_PhAT_Mesh_40x", Icon20x20 ) );
-		Set( "PhAT.ResetEntireAsset.Small", new IMAGE_BRUSH( "Icons/icon_PhAT_ResetAsset_40x", Icon20x20 ) );
-		Set( "PhAT.RestetBoneCollision.Small", new IMAGE_BRUSH( "Icons/icon_PhAT_ResetCollision_40x", Icon20x20 ) );
-		Set( "PhAT.ApplyPhysicalMaterial.Small", new IMAGE_BRUSH( "Icons/icon_PhAT_PhysMat_40x", Icon20x20 ) );
-		Set( "PhAT.CopyJointSettings.Small", new IMAGE_BRUSH( "Icons/icon_PhAT_CopyJoints_40x", Icon20x20 ) );
-		Set( "PhAT.PlayAnimation.Small", new IMAGE_BRUSH( "Icons/icon_PhAT_Play_40x", Icon20x20 ) );
-		Set( "PhAT.PhATTranslationMode.Small", new IMAGE_BRUSH( "Icons/icon_translate_40x", Icon20x20 ) );
-		Set( "PhAT.PhATRotationMode.Small", new IMAGE_BRUSH( "Icons/icon_rotate_40x", Icon20x20 ) );
-		Set( "PhAT.PhATScaleMode.Small", new IMAGE_BRUSH( "Icons/icon_scale_40x", Icon20x20 ) );
-		Set( "PhAT.Snap.Small", new IMAGE_BRUSH( "Icons/icon_PhAT_Snap_40x", Icon20x20 ) );
-		Set( "PhAT.CopyProperties.Small", new IMAGE_BRUSH( "Icons/icon_PhAT_CopyProperties_40x", Icon20x20 ) );
-		Set( "PhAT.DisableCollision.Small", new IMAGE_BRUSH( "Icons/icon_PhAT_DisableCollision_40x", Icon20x20 ) );
-		Set( "PhAT.EnableCollision.Small", new IMAGE_BRUSH( "Icons/icon_PhAT_EnableCollision_40x", Icon20x20 ) );
-		Set( "PhAT.WeldToBody.Small", new IMAGE_BRUSH( "Icons/icon_PhAT_Weld_40x", Icon20x20 ) );
-		Set( "PhAT.AddNewBody.Small", new IMAGE_BRUSH( "Icons/icon_PhAT_NewBody_40x", Icon20x20 ) );
-		Set( "PhAT.AddSphere.Small", new IMAGE_BRUSH( "Icons/icon_PhAT_Sphere_40x", Icon20x20 ) );
-		Set( "PhAT.AddSphyl.Small", new IMAGE_BRUSH( "Icons/icon_PhAT_Sphyl_40x", Icon20x20 ) );
-		Set( "PhAT.AddBox.Small", new IMAGE_BRUSH( "Icons/icon_PhAT_Box_40x", Icon20x20 ) );
-		Set( "PhAT.DeletePrimitive.Small", new IMAGE_BRUSH( "Icons/icon_PhAT_DeletePrimitive_40x", Icon20x20 ) );
-		Set( "PhAT.DuplicatePrimitive.Small", new IMAGE_BRUSH( "Icons/icon_PhAT_DupePrim_40x", Icon20x20 ) );
-		Set( "PhAT.ResetConstraint.Small", new IMAGE_BRUSH( "Icons/icon_PhAT_ResetConstraint_40x", Icon20x20 ) );
-		Set( "PhAT.SnapConstraint.Small", new IMAGE_BRUSH( "Icons/icon_PhAT_SnapConstraint_40x", Icon20x20 ) );
-		Set( "PhAT.SnapAllConstraints.Small", new IMAGE_BRUSH( "Icons/icon_PhAT_SnapAll_40x", Icon20x20 ) );
-		Set( "PhAT.ConvertToBallAndSocket.Small", new IMAGE_BRUSH( "Icons/icon_PhAT_Ball_40x", Icon20x20 ) );
-		Set( "PhAT.ConvertToHinge.Small", new IMAGE_BRUSH( "Icons/icon_PhAT_Hinge_40x", Icon20x20 ) );
-		Set( "PhAT.ConvertToPrismatic.Small", new IMAGE_BRUSH( "Icons/icon_PhAT_Prismatic_40x", Icon20x20 ) );
-		Set( "PhAT.ConvertToSkeletal.Small", new IMAGE_BRUSH( "Icons/icon_PhAT_Skeletal_40x", Icon20x20 ) );
-		Set( "PhAT.DeleteConstraint.Small", new IMAGE_BRUSH( "Icons/icon_PhAT_DeleteConstraint_40x", Icon20x20 ) );
+	// PhysicsAssetEditor
+	{
+		Set( "PhysicsAssetEditor.Tabs.Properties", new IMAGE_BRUSH( "/Icons/icon_tab_SelectionDetails_16x", Icon16x16 ) );
+		Set( "PhysicsAssetEditor.Tabs.Hierarchy", new IMAGE_BRUSH( "/Icons/levels_16x", Icon16x16 ) );
+		Set( "PhysicsAssetEditor.Tabs.Profiles", new IMAGE_BRUSH( "/PhysicsAssetEditor/icon_ProfilesTab_16x", Icon16x16 ) );
+		Set( "PhysicsAssetEditor.Tabs.Graph", new IMAGE_BRUSH( "/PhysicsAssetEditor/icon_GraphTab_16x", Icon16x16 ) );
+		Set( "PhysicsAssetEditor.Tabs.Tools", new IMAGE_BRUSH( "/PhysicsAssetEditor/icon_ToolsTab_16x", Icon16x16 ) );
+
+		Set( "PhysicsAssetEditor.EditingMode_Body", new IMAGE_BRUSH( "/PhysicsAssetEditor/icon_PHatMode_Body_40x", Icon40x40) );
+		Set( "PhysicsAssetEditor.EditingMode_Constraint", new IMAGE_BRUSH( "/PhysicsAssetEditor/icon_PHatMode_Joint_40x", Icon40x40) );
+
+		Set( "PhysicsAssetEditor.EditingMode_Body.Small", new IMAGE_BRUSH( "/PhysicsAssetEditor/icon_PHatMode_Body_40x", Icon20x20) );
+		Set( "PhysicsAssetEditor.EditingMode_Constraint.Small", new IMAGE_BRUSH( "/PhysicsAssetEditor/icon_PHatMode_Joint_40x", Icon20x20) );
+
+		Set( "PhysicsAssetEditor.SimulationNoGravity", new IMAGE_BRUSH( "PhysicsAssetEditor/icon_PhAT_PlaySimNoGravity_40x", Icon40x40 ) );
+		Set( "PhysicsAssetEditor.SelectedSimulation", new IMAGE_BRUSH( "PhysicsAssetEditor/icon_PhAT_PlaySimSelected_40x", Icon40x40 ) );
+		Set( "PhysicsAssetEditor.SimulationAll", new IMAGE_BRUSH( "PhysicsAssetEditor/icon_PhAT_PlaySim_40x", Icon40x40 ) );
+		Set( "PhysicsAssetEditor.Undo", new IMAGE_BRUSH( "Icons/icon_Generic_Undo_40x", Icon40x40 ) );
+		Set( "PhysicsAssetEditor.Redo", new IMAGE_BRUSH( "Icons/icon_Generic_Redo_40x", Icon40x40 ) );
+		Set( "PhysicsAssetEditor.ChangeDefaultMesh", new IMAGE_BRUSH( "PhysicsAssetEditor/icon_PhAT_Mesh_40x", Icon40x40 ) );
+		Set( "PhysicsAssetEditor.ApplyPhysicalMaterial", new IMAGE_BRUSH( "PhysicsAssetEditor/icon_PhAT_PhysMat_40x", Icon40x40 ) );
+		Set( "PhysicsAssetEditor.CopyJointSettings", new IMAGE_BRUSH( "PhysicsAssetEditor/icon_PhAT_CopyJoints_40x", Icon40x40 ) );
+		Set( "PhysicsAssetEditor.PlayAnimation", new IMAGE_BRUSH( "PhysicsAssetEditor/icon_PhAT_Play_40x", Icon40x40 ) );
+		Set( "PhysicsAssetEditor.PhATTranslationMode", new IMAGE_BRUSH( "Icons/icon_translate_40x", Icon40x40 ) );
+		Set( "PhysicsAssetEditor.PhATRotationMode", new IMAGE_BRUSH( "Icons/icon_rotate_40x", Icon40x40 ) );
+		Set( "PhysicsAssetEditor.PhATScaleMode", new IMAGE_BRUSH( "Icons/icon_scale_40x", Icon40x40 ) );
+		Set( "PhysicsAssetEditor.Snap", new IMAGE_BRUSH( "PhysicsAssetEditor/icon_PhAT_Snap_40x", Icon40x40 ) );
+		Set( "PhysicsAssetEditor.CopyProperties", new IMAGE_BRUSH( "PhysicsAssetEditor/icon_PhAT_CopyProperties_40x", Icon40x40 ) );
+		Set( "PhysicsAssetEditor.DisableCollision", new IMAGE_BRUSH( "PhysicsAssetEditor/icon_PhAT_DisableCollision_40x", Icon40x40 ) );
+		Set( "PhysicsAssetEditor.EnableCollision", new IMAGE_BRUSH( "PhysicsAssetEditor/icon_PhAT_EnableCollision_40x", Icon40x40 ) );
+		Set( "PhysicsAssetEditor.DisableCollisionAll", new IMAGE_BRUSH( "PhysicsAssetEditor/icon_PhAT_DisableCollision_40x", Icon40x40 ) );
+		Set( "PhysicsAssetEditor.EnableCollisionAll", new IMAGE_BRUSH( "PhysicsAssetEditor/icon_PhAT_EnableCollision_40x", Icon40x40 ) );
+		Set( "PhysicsAssetEditor.WeldToBody", new IMAGE_BRUSH( "PhysicsAssetEditor/icon_PhAT_Weld_40x", Icon40x40 ) );
+		Set( "PhysicsAssetEditor.AddNewBody", new IMAGE_BRUSH( "PhysicsAssetEditor/icon_PhAT_NewBody_40x", Icon40x40 ) );
+		Set( "PhysicsAssetEditor.AddSphere", new IMAGE_BRUSH( "PhysicsAssetEditor/icon_PhAT_Sphere_40x", Icon40x40 ) );
+		Set( "PhysicsAssetEditor.AddSphyl", new IMAGE_BRUSH( "PhysicsAssetEditor/icon_PhAT_Sphyl_40x", Icon40x40 ) );
+		Set( "PhysicsAssetEditor.AddBox", new IMAGE_BRUSH( "PhysicsAssetEditor/icon_PhAT_Box_40x", Icon40x40 ) );
+		Set( "PhysicsAssetEditor.DeletePrimitive", new IMAGE_BRUSH( "PhysicsAssetEditor/icon_PhAT_DeletePrimitive_40x", Icon40x40 ) );
+		Set( "PhysicsAssetEditor.DuplicatePrimitive", new IMAGE_BRUSH( "PhysicsAssetEditor/icon_PhAT_DupePrim_40x", Icon40x40 ) );
+		Set( "PhysicsAssetEditor.ResetConstraint", new IMAGE_BRUSH( "PhysicsAssetEditor/icon_PhAT_ResetConstraint_40x", Icon40x40 ) );
+		Set( "PhysicsAssetEditor.SnapConstraint", new IMAGE_BRUSH( "PhysicsAssetEditor/icon_PhAT_SnapConstraint_40x", Icon40x40 ) );
+		Set( "PhysicsAssetEditor.SnapAllConstraints", new IMAGE_BRUSH( "PhysicsAssetEditor/icon_PhAT_SnapAll_40x", Icon40x40 ) );
+		Set( "PhysicsAssetEditor.ConvertToBallAndSocket", new IMAGE_BRUSH( "PhysicsAssetEditor/icon_PhAT_Ball_40x", Icon40x40 ) );
+		Set( "PhysicsAssetEditor.ConvertToHinge", new IMAGE_BRUSH( "PhysicsAssetEditor/icon_PhAT_Hinge_40x", Icon40x40 ) );
+		Set( "PhysicsAssetEditor.ConvertToPrismatic", new IMAGE_BRUSH( "PhysicsAssetEditor/icon_PhAT_Prismatic_40x", Icon40x40 ) );
+		Set( "PhysicsAssetEditor.ConvertToSkeletal", new IMAGE_BRUSH( "PhysicsAssetEditor/icon_PhAT_Skeletal_40x", Icon40x40 ) );
+		Set( "PhysicsAssetEditor.DeleteConstraint", new IMAGE_BRUSH( "PhysicsAssetEditor/icon_PhAT_DeleteConstraint_40x", Icon40x40 ) );
+
+		Set("PhysicsAssetEditor.SimulationNoGravity.Small", new IMAGE_BRUSH("PhysicsAssetEditor/icon_PhAT_PlaySimNoGravity_40x", Icon20x20));
+		Set("PhysicsAssetEditor.SelectedSimulation.Small", new IMAGE_BRUSH("PhysicsAssetEditor/icon_PhAT_PlaySimSelected_40x", Icon20x20));
+		Set( "PhysicsAssetEditor.SimulationAll.Small", new IMAGE_BRUSH( "PhysicsAssetEditor/icon_PhAT_PlaySim_40x", Icon20x20 ) );
+		Set( "PhysicsAssetEditor.Undo.Small", new IMAGE_BRUSH( "Icons/icon_Generic_Undo_40x", Icon20x20 ) );
+		Set( "PhysicsAssetEditor.Redo.Small", new IMAGE_BRUSH( "Icons/icon_Generic_Redo_40x", Icon20x20 ) );
+		Set( "PhysicsAssetEditor.ChangeDefaultMesh.Small", new IMAGE_BRUSH( "PhysicsAssetEditor/icon_PhAT_Mesh_40x", Icon20x20 ) );
+		Set( "PhysicsAssetEditor.ResetEntireAsset.Small", new IMAGE_BRUSH( "PhysicsAssetEditor/icon_PhAT_ResetAsset_40x", Icon20x20 ) );
+		Set( "PhysicsAssetEditor.ResetBoneCollision.Small", new IMAGE_BRUSH( "PhysicsAssetEditor/icon_PhAT_ResetCollision_40x", Icon20x20 ) );
+		Set( "PhysicsAssetEditor.ApplyPhysicalMaterial.Small", new IMAGE_BRUSH( "PhysicsAssetEditor/icon_PhAT_PhysMat_40x", Icon20x20 ) );
+		Set( "PhysicsAssetEditor.CopyJointSettings.Small", new IMAGE_BRUSH( "PhysicsAssetEditor/icon_PhAT_CopyJoints_40x", Icon20x20 ) );
+		Set( "PhysicsAssetEditor.PlayAnimation.Small", new IMAGE_BRUSH( "PhysicsAssetEditor/icon_PhAT_Play_40x", Icon20x20 ) );
+		Set( "PhysicsAssetEditor.PhATTranslationMode.Small", new IMAGE_BRUSH( "Icons/icon_translate_40x", Icon20x20 ) );
+		Set( "PhysicsAssetEditor.PhATRotationMode.Small", new IMAGE_BRUSH( "Icons/icon_rotate_40x", Icon20x20 ) );
+		Set( "PhysicsAssetEditor.PhATScaleMode.Small", new IMAGE_BRUSH( "Icons/icon_scale_40x", Icon20x20 ) );
+		Set( "PhysicsAssetEditor.Snap.Small", new IMAGE_BRUSH( "PhysicsAssetEditor/icon_PhAT_Snap_40x", Icon20x20 ) );
+		Set( "PhysicsAssetEditor.CopyProperties.Small", new IMAGE_BRUSH( "PhysicsAssetEditor/icon_PhAT_CopyProperties_40x", Icon20x20 ) );
+		Set( "PhysicsAssetEditor.DisableCollision.Small", new IMAGE_BRUSH( "PhysicsAssetEditor/icon_PhAT_DisableCollision_40x", Icon20x20 ) );
+		Set( "PhysicsAssetEditor.EnableCollision.Small", new IMAGE_BRUSH( "PhysicsAssetEditor/icon_PhAT_EnableCollision_40x", Icon20x20 ) );
+		Set( "PhysicsAssetEditor.DisableCollisionAll.Small", new IMAGE_BRUSH( "PhysicsAssetEditor/icon_PhAT_DisableCollision_40x", Icon20x20 ) );
+		Set( "PhysicsAssetEditor.EnableCollisionAll.Small", new IMAGE_BRUSH( "PhysicsAssetEditor/icon_PhAT_EnableCollision_40x", Icon20x20 ) );
+		Set( "PhysicsAssetEditor.WeldToBody.Small", new IMAGE_BRUSH( "PhysicsAssetEditor/icon_PhAT_Weld_40x", Icon20x20 ) );
+		Set( "PhysicsAssetEditor.AddNewBody.Small", new IMAGE_BRUSH( "PhysicsAssetEditor/icon_PhAT_NewBody_40x", Icon20x20 ) );
+		Set( "PhysicsAssetEditor.AddSphere.Small", new IMAGE_BRUSH( "PhysicsAssetEditor/icon_PhAT_Sphere_40x", Icon20x20 ) );
+		Set( "PhysicsAssetEditor.AddSphyl.Small", new IMAGE_BRUSH( "PhysicsAssetEditor/icon_PhAT_Sphyl_40x", Icon20x20 ) );
+		Set( "PhysicsAssetEditor.AddBox.Small", new IMAGE_BRUSH( "PhysicsAssetEditor/icon_PhAT_Box_40x", Icon20x20 ) );
+		Set( "PhysicsAssetEditor.DeletePrimitive.Small", new IMAGE_BRUSH( "PhysicsAssetEditor/icon_PhAT_DeletePrimitive_40x", Icon20x20 ) );
+		Set( "PhysicsAssetEditor.DuplicatePrimitive.Small", new IMAGE_BRUSH( "PhysicsAssetEditor/icon_PhAT_DupePrim_40x", Icon20x20 ) );
+		Set( "PhysicsAssetEditor.ResetConstraint.Small", new IMAGE_BRUSH( "PhysicsAssetEditor/icon_PhAT_ResetConstraint_40x", Icon20x20 ) );
+		Set( "PhysicsAssetEditor.SnapConstraint.Small", new IMAGE_BRUSH( "PhysicsAssetEditor/icon_PhAT_SnapConstraint_40x", Icon20x20 ) );
+		Set( "PhysicsAssetEditor.SnapAllConstraints.Small", new IMAGE_BRUSH( "PhysicsAssetEditor/icon_PhAT_SnapAll_40x", Icon20x20 ) );
+		Set( "PhysicsAssetEditor.ConvertToBallAndSocket.Small", new IMAGE_BRUSH( "PhysicsAssetEditor/icon_PhAT_Ball_40x", Icon20x20 ) );
+		Set( "PhysicsAssetEditor.ConvertToHinge.Small", new IMAGE_BRUSH( "PhysicsAssetEditor/icon_PhAT_Hinge_40x", Icon20x20 ) );
+		Set( "PhysicsAssetEditor.ConvertToPrismatic.Small", new IMAGE_BRUSH( "PhysicsAssetEditor/icon_PhAT_Prismatic_40x", Icon20x20 ) );
+		Set( "PhysicsAssetEditor.ConvertToSkeletal.Small", new IMAGE_BRUSH( "PhysicsAssetEditor/icon_PhAT_Skeletal_40x", Icon20x20 ) );
+		Set( "PhysicsAssetEditor.DeleteConstraint.Small", new IMAGE_BRUSH( "PhysicsAssetEditor/icon_PhAT_DeleteConstraint_40x", Icon20x20 ) );
+
+		Set("PhysicsAssetEditor.NewPhysicalAnimationProfile", new IMAGE_BRUSH("PhysicsAssetEditor/icon_PhAT_NewBody_40x", Icon20x20));
+		Set("PhysicsAssetEditor.DeleteCurrentPhysicalAnimationProfile", new IMAGE_BRUSH("PhysicsAssetEditor/icon_PhAT_DeletePrimitive_40x", Icon20x20));
+		Set("PhysicsAssetEditor.AddBodyToPhysicalAnimationProfile", new IMAGE_BRUSH("PhysicsAssetEditor/icon_PhAT_NewBody_40x", Icon20x20));
+		Set("PhysicsAssetEditor.RemoveBodyFromPhysicalAnimationProfile", new IMAGE_BRUSH("PhysicsAssetEditor/icon_PhAT_DeletePrimitive_40x", Icon20x20));
+		Set("PhysicsAssetEditor.NewConstraintProfile", new IMAGE_BRUSH("PhysicsAssetEditor/icon_PHatMode_Joint_40x", Icon20x20));
+		Set("PhysicsAssetEditor.DeleteCurrentConstraintProfile", new IMAGE_BRUSH("PhysicsAssetEditor/icon_PhAT_DeleteConstraint_40x", Icon20x20));
+		Set("PhysicsAssetEditor.AddConstraintToCurrentConstraintProfile", new IMAGE_BRUSH("PhysicsAssetEditor/icon_PHatMode_Joint_40x", Icon20x20));
+		Set("PhysicsAssetEditor.RemoveConstraintFromCurrentConstraintProfile", new IMAGE_BRUSH("PhysicsAssetEditor/icon_PhAT_DeleteConstraint_40x", Icon20x20));
+
+		Set("PhysicsAssetEditor.Tree.Body", new IMAGE_BRUSH("PhysicsAssetEditor/Body_16x", Icon16x16));
+		Set("PhysicsAssetEditor.Tree.KinematicBody", new IMAGE_BRUSH("PhysicsAssetEditor/KinematicBody_16x", Icon16x16));
+		Set("PhysicsAssetEditor.Tree.EmptyBody", new IMAGE_BRUSH("PhysicsAssetEditor/EmptyBody_16x", Icon16x16));
+		Set("PhysicsAssetEditor.Tree.Bone", new IMAGE_BRUSH("PhysicsAssetEditor/Bone_16x", Icon16x16));
+		Set("PhysicsAssetEditor.Tree.Sphere", new IMAGE_BRUSH("PhysicsAssetEditor/Sphere_16x", Icon16x16));
+		Set("PhysicsAssetEditor.Tree.Sphyl", new IMAGE_BRUSH("PhysicsAssetEditor/Sphyl_16x", Icon16x16));
+		Set("PhysicsAssetEditor.Tree.Box", new IMAGE_BRUSH("PhysicsAssetEditor/Box_16x", Icon16x16));
+		Set("PhysicsAssetEditor.Tree.Convex", new IMAGE_BRUSH("PhysicsAssetEditor/Convex_16x", Icon16x16));
+		Set("PhysicsAssetEditor.Tree.Constraint", new IMAGE_BRUSH("PhysicsAssetEditor/Constraint_16x", Icon16x16));
+
+		Set("PhysicsAssetEditor.Tree.Font", TTF_CORE_FONT("Fonts/Roboto-Regular", 10));
+
+		Set("PhysicsAssetEditor.Graph.TextStyle", FTextBlockStyle(NormalText)
+			.SetColorAndOpacity(FLinearColor(0.0f, 0.0f, 0.0f, 1.0f))
+			.SetFont(TTF_CORE_FONT("Fonts/Roboto-Regular", 8)));
+
+		Set("PhysicsAssetEditor.Graph.NodeBody", new BOX_BRUSH("PhysicsAssetEditor/NodeBody", FMargin(4.f / 64.f, 4.f / 64.f, 4.f / 64.f, 4.f / 64.f)));
+		Set("PhysicsAssetEditor.Graph.NodeIcon", new IMAGE_BRUSH("PhysicsAssetEditor/Bone_16x", Icon16x16));
+		Set("PhysicsAssetEditor.Graph.Pin.Background", new IMAGE_BRUSH("PhysicsAssetEditor/NodePin", Icon10x10));
+		Set("PhysicsAssetEditor.Graph.Pin.BackgroundHovered", new IMAGE_BRUSH("PhysicsAssetEditor/NodePinHoverCue", Icon10x10));
+		Set("PhysicsAssetEditor.Graph.Node.ShadowSelected", new BOX_BRUSH( "PhysicsAssetEditor/PhysicsNode_shadow_selected", FMargin(18.0f/64.0f) ) );
+		Set("PhysicsAssetEditor.Graph.Node.Shadow", new BOX_BRUSH( "Graph/RegularNode_shadow", FMargin(18.0f/64.0f) ) );
+
+		FEditableTextBoxStyle EditableTextBlock = FEditableTextBoxStyle()
+			.SetFont(NormalText.Font)
+			.SetBackgroundImageNormal(BOX_BRUSH("Common/TextBox", FMargin(4.0f / 16.0f)))
+			.SetBackgroundImageHovered(BOX_BRUSH("Common/TextBox_Hovered", FMargin(4.0f / 16.0f)))
+			.SetBackgroundImageFocused(BOX_BRUSH("Common/TextBox_Hovered", FMargin(4.0f / 16.0f)))
+			.SetBackgroundImageReadOnly(BOX_BRUSH("Common/TextBox_ReadOnly", FMargin(4.0f / 16.0f)));
+
+		Set("PhysicsAssetEditor.Profiles.EditableTextBoxStyle", EditableTextBlock);
+
+		Set("PhysicsAssetEditor.Profiles.Font", FTextBlockStyle(NormalText)
+			.SetFont(TTF_CORE_FONT("Fonts/Roboto-Bold", 11))
+			.SetColorAndOpacity(FLinearColor(1.0f, 1.0f, 1.0f))
+			.SetHighlightColor(FLinearColor(1.0f, 1.0f, 1.0f))
+			.SetShadowOffset(FVector2D(1, 1))
+			.SetShadowColorAndOpacity(FLinearColor(0, 0, 0, 0.9f)));
+
+		Set("PhysicsAssetEditor.Tools.Font", FTextBlockStyle(NormalText)
+			.SetFont(TTF_CORE_FONT("Fonts/Roboto-Bold", 11))
+			.SetColorAndOpacity(FLinearColor(1.0f, 1.0f, 1.0f))
+			.SetHighlightColor(FLinearColor(1.0f, 1.0f, 1.0f))
+			.SetShadowOffset(FVector2D(1, 1))
+			.SetShadowColorAndOpacity(FLinearColor(0, 0, 0, 0.9f)));
+
+		FLinearColor Red = FLinearColor::Red;
+		FLinearColor Red_Selected = FLinearColor::Red.Desaturate(0.75f);
+		FLinearColor Red_Pressed = FLinearColor::Red.Desaturate(0.5f);
+
+		const FCheckBoxStyle RedRadioButtonStyle = FCheckBoxStyle()
+			.SetUncheckedImage( IMAGE_BRUSH( "Common/RadioButton_Unselected_16x", Icon16x16, Red ) )
+			.SetUncheckedHoveredImage( IMAGE_BRUSH( "Common/RadioButton_Unselected_16x", Icon16x16, Red ) )
+			.SetUncheckedPressedImage( IMAGE_BRUSH( "Common/RadioButton_Unselected_16x", Icon16x16, Red ) )
+			.SetCheckedImage( IMAGE_BRUSH( "Common/RadioButton_Selected_16x", Icon16x16, Red ) )
+			.SetCheckedHoveredImage( IMAGE_BRUSH( "Common/RadioButton_Selected_16x", Icon16x16, Red_Selected ) )
+			.SetCheckedPressedImage( IMAGE_BRUSH( "Common/RadioButton_Unselected_16x", Icon16x16, Red_Pressed ) )
+			.SetUndeterminedImage( IMAGE_BRUSH( "Common/RadioButton_Unselected_16x", Icon16x16, Red ) )
+			.SetUndeterminedHoveredImage( IMAGE_BRUSH( "Common/RadioButton_Unselected_16x", Icon16x16, Red_Selected ) )
+			.SetUndeterminedPressedImage( IMAGE_BRUSH( "Common/RadioButton_Unselected_16x", Icon16x16, Red_Pressed ) );
+
+		Set( "PhysicsAssetEditor.RadioButtons.Red", RedRadioButtonStyle );
+
+		FLinearColor Green = FLinearColor::Green;
+		FLinearColor Green_Selected = FLinearColor::Green.Desaturate(0.75f);
+		FLinearColor Green_Pressed = FLinearColor::Green.Desaturate(0.5f);
+
+		const FCheckBoxStyle GreenRadioButtonStyle = FCheckBoxStyle()
+			.SetUncheckedImage( IMAGE_BRUSH( "Common/RadioButton_Unselected_16x", Icon16x16, Green ) )
+			.SetUncheckedHoveredImage( IMAGE_BRUSH( "Common/RadioButton_Unselected_16x", Icon16x16, Green ) )
+			.SetUncheckedPressedImage( IMAGE_BRUSH( "Common/RadioButton_Unselected_16x", Icon16x16, Green ) )
+			.SetCheckedImage( IMAGE_BRUSH( "Common/RadioButton_Selected_16x", Icon16x16, Green ) )
+			.SetCheckedHoveredImage( IMAGE_BRUSH( "Common/RadioButton_Selected_16x", Icon16x16, Green_Selected ) )
+			.SetCheckedPressedImage( IMAGE_BRUSH( "Common/RadioButton_Unselected_16x", Icon16x16, Green_Pressed ) )
+			.SetUndeterminedImage( IMAGE_BRUSH( "Common/RadioButton_Unselected_16x", Icon16x16, Green ) )
+			.SetUndeterminedHoveredImage( IMAGE_BRUSH( "Common/RadioButton_Unselected_16x", Icon16x16, Green_Selected ) )
+			.SetUndeterminedPressedImage( IMAGE_BRUSH( "Common/RadioButton_Unselected_16x", Icon16x16, Green_Pressed ) );
+
+		Set( "PhysicsAssetEditor.RadioButtons.Green", GreenRadioButtonStyle );
+
+		FLinearColor Blue = FLinearColor::Blue;
+		FLinearColor Blue_Selected = FLinearColor::Blue.Desaturate(0.75f);
+		FLinearColor Blue_Pressed = FLinearColor::Blue.Desaturate(0.5f);
+
+		const FCheckBoxStyle BlueRadioButtonStyle = FCheckBoxStyle()
+			.SetUncheckedImage( IMAGE_BRUSH( "Common/RadioButton_Unselected_16x", Icon16x16, Blue ) )
+			.SetUncheckedHoveredImage( IMAGE_BRUSH( "Common/RadioButton_Unselected_16x", Icon16x16, Blue ) )
+			.SetUncheckedPressedImage( IMAGE_BRUSH( "Common/RadioButton_Unselected_16x", Icon16x16, Blue ) )
+			.SetCheckedImage( IMAGE_BRUSH( "Common/RadioButton_Selected_16x", Icon16x16, Blue ) )
+			.SetCheckedHoveredImage( IMAGE_BRUSH( "Common/RadioButton_Selected_16x", Icon16x16, Blue_Selected ) )
+			.SetCheckedPressedImage( IMAGE_BRUSH( "Common/RadioButton_Unselected_16x", Icon16x16, Blue_Pressed ) )
+			.SetUndeterminedImage( IMAGE_BRUSH( "Common/RadioButton_Unselected_16x", Icon16x16, Blue ) )
+			.SetUndeterminedHoveredImage( IMAGE_BRUSH( "Common/RadioButton_Unselected_16x", Icon16x16, Blue_Selected ) )
+			.SetUndeterminedPressedImage( IMAGE_BRUSH( "Common/RadioButton_Unselected_16x", Icon16x16, Blue_Pressed ) );
+
+		Set( "PhysicsAssetEditor.RadioButtons.Blue", BlueRadioButtonStyle );
 	}
 #endif // WITH_EDITOR
 }
diff --git a/Engine/Source/Editor/KismetCompiler/Private/AnimBlueprintCompiler.cpp b/Engine/Source/Editor/KismetCompiler/Private/AnimBlueprintCompiler.cpp
--- a/Engine/Source/Editor/KismetCompiler/Private/AnimBlueprintCompiler.cpp
+++ b/Engine/Source/Editor/KismetCompiler/Private/AnimBlueprintCompiler.cpp
@@ -665,7 +665,6 @@ void FAnimBlueprintCompiler::ProcessSubInstance(UAnimGraphNode_SubInstance* SubI
 		UProperty* NewProperty = FKismetCompilerUtilities::CreatePropertyOnScope(NewAnimBlueprintClass, FName(*PrefixedName), Pin->PinType, NewAnimBlueprintClass, 0, GetSchema(), MessageLog);
 		if(NewProperty)
 		{
-			NewProperty->SetMetaData(TEXT("Category"), TEXT("SubInstance"));
 			FKismetCompilerUtilities::LinkAddedProperty(NewAnimBlueprintClass, NewProperty);
 
 			// Add mappings to the node
diff --git a/Engine/Source/Editor/LevelEditor/Private/SLevelEditor.cpp b/Engine/Source/Editor/LevelEditor/Private/SLevelEditor.cpp
--- a/Engine/Source/Editor/LevelEditor/Private/SLevelEditor.cpp
+++ b/Engine/Source/Editor/LevelEditor/Private/SLevelEditor.cpp
@@ -702,10 +702,13 @@ TSharedRef<SDockTab> SLevelEditor::SpawnLevelEditorTab( const FSpawnTabArgs& Arg
 	}
 	else if (TabIdentifier == TEXT("LevelEditorHierarchicalLODOutliner"))
 	{
+		FText Label = NSLOCTEXT("LevelEditor", "HLODOutlinerTabTitle", "Hierarchical LOD Outliner");
+
 		FHierarchicalLODOutlinerModule& HLODModule = FModuleManager::LoadModuleChecked<FHierarchicalLODOutlinerModule>("HierarchicalLODOutliner");
 		return SNew(SDockTab)
 			.Icon(FEditorStyle::GetBrush("LevelEditor.Tabs.HLOD"))
-			.Label(NSLOCTEXT("LevelEditor", "HLODTabTitle", "Hierarchical LOD Outliner"))
+			.Label(Label)
+			.ToolTip(IDocumentation::Get()->CreateToolTip(Label, nullptr, "Shared/Editor/HLOD", "main"))
 			[
 				SNew(SBorder)
 				.Padding(0)
diff --git a/Engine/Source/Editor/MeshPaint/Private/MeshPaintSkeletalMeshAdapter.cpp b/Engine/Source/Editor/MeshPaint/Private/MeshPaintSkeletalMeshAdapter.cpp
--- a/Engine/Source/Editor/MeshPaint/Private/MeshPaintSkeletalMeshAdapter.cpp
+++ b/Engine/Source/Editor/MeshPaint/Private/MeshPaintSkeletalMeshAdapter.cpp
@@ -198,7 +198,51 @@ void FMeshPaintGeometryAdapterForSkeletalMeshes::OnRemoved()
 
 bool FMeshPaintGeometryAdapterForSkeletalMeshes::LineTraceComponent(struct FHitResult& OutHit, const FVector Start, const FVector End, const struct FCollisionQueryParams& Params) const
 {
-	return SkeletalMeshComponent->LineTraceComponent(OutHit, Start, End, Params);
+	const bool bHitBounds = FMath::LineSphereIntersection(Start, End.GetSafeNormal(), (End - Start).SizeSquared(), SkeletalMeshComponent->Bounds.Origin, SkeletalMeshComponent->Bounds.SphereRadius);
+	const bool bHitPhysicsBodies = SkeletalMeshComponent->LineTraceComponent(OutHit, Start, End, Params);
+
+	bool bHitTriangle = false;
+	if (bHitBounds && !bHitPhysicsBodies)
+	{
+		const int32 NumTriangles = MeshIndices.Num() / 3;
+		const FTransform& ComponentTransform = SkeletalMeshComponent->GetComponentTransform();
+		const FTransform InverseComponentTransform = ComponentTransform.Inverse();
+		const FVector LocalStart = InverseComponentTransform.TransformPosition(Start);
+		const FVector LocalEnd = InverseComponentTransform.TransformPosition(End);
+
+		float MinDistance = FLT_MAX;
+		FVector Intersect;
+		FVector Normal;
+
+		for (int32 TriangleIndex = 0; TriangleIndex < NumTriangles; ++TriangleIndex)
+		{
+			FVector IntersectPoint;
+			FVector HitNormal;
+			bool bHit = FMath::SegmentTriangleIntersection(LocalStart, LocalEnd, MeshVertices[MeshIndices[(TriangleIndex * 3) + 0]], MeshVertices[MeshIndices[(TriangleIndex * 3) + 1]], MeshVertices[MeshIndices[(TriangleIndex * 3) + 2]], IntersectPoint, HitNormal);
+
+			if (bHit)
+			{
+				const float Distance = (LocalStart - IntersectPoint).SizeSquared();
+				if (Distance < MinDistance)
+				{
+					MinDistance = Distance;
+					Intersect = IntersectPoint;
+					Normal = HitNormal;
+				}
+			}
+		}
+
+		if (MinDistance != FLT_MAX)
+		{
+			OutHit.Component = SkeletalMeshComponent;
+			OutHit.Normal = Normal.GetSafeNormal();
+			OutHit.Location = ComponentTransform.TransformPosition(Intersect);
+			OutHit.bBlockingHit = true;
+			bHitTriangle = true;
+		}
+	}	
+
+	return bHitPhysicsBodies || bHitTriangle;
 }
 
 void FMeshPaintGeometryAdapterForSkeletalMeshes::QueryPaintableTextures(int32 MaterialIndex, int32& OutDefaultIndex, TArray<struct FPaintableTexture>& InOutTextureList)
diff --git a/Engine/Source/Editor/MeshPaint/Public/MeshPaintHelpers.h b/Engine/Source/Editor/MeshPaint/Public/MeshPaintHelpers.h
--- a/Engine/Source/Editor/MeshPaint/Public/MeshPaintHelpers.h
+++ b/Engine/Source/Editor/MeshPaint/Public/MeshPaintHelpers.h
@@ -142,7 +142,7 @@ public:
 	static bool IsPointInfluencedByBrush(const FVector2D& BrushSpacePosition, const float BrushRadiusSquared, float& OutInRangeValue);
 
 	template<typename T>
-	static void ApplyBrushToVertex(const FVector& VertexPosition, const FMatrix& InverseBrushMatrix, const float BrushRadius, const float BrushFalloff, const T& PaintValue, T& InOutValue);
+	static void ApplyBrushToVertex(const FVector& VertexPosition, const FMatrix& InverseBrushMatrix, const float BrushRadius, const float BrushFalloffAmount, const float BrushStrength, const T& PaintValue, T& InOutValue);
 
 public:
 	struct FPaintRay
@@ -169,23 +169,16 @@ protected:
 };
 
 template<typename T>
-void MeshPaintHelpers::ApplyBrushToVertex(const FVector& VertexPosition, const FMatrix& InverseBrushMatrix, const float BrushRadiusSquared, const float BrushFalloff, const T& PaintValue, T& InOutValue)
+void MeshPaintHelpers::ApplyBrushToVertex(const FVector& VertexPosition, const FMatrix& InverseBrushMatrix, const float BrushRadius, const float BrushFalloffAmount, const float BrushStrength, const T& PaintValue, T& InOutValue)
 {
 	const FVector BrushSpacePosition = InverseBrushMatrix.TransformPosition(VertexPosition);
 	const FVector2D BrushSpacePosition2D(BrushSpacePosition.X, BrushSpacePosition.Y);
 		
 	float InfluencedValue = 0.0f;
-	if (IsPointInfluencedByBrush(BrushSpacePosition2D, BrushRadiusSquared, InfluencedValue))
+	if (IsPointInfluencedByBrush(BrushSpacePosition2D, BrushRadius * BrushRadius, InfluencedValue))
 	{
-		float PaintStrength = 1.0f;
-
-		// Now know range from 0 - 1 how far vertex is inside of the sphere
-		// If this falls outside of the brush inner radius
-		const float InnerBrushPercentage = BrushFalloff;
-		if (InfluencedValue > InnerBrushPercentage)
-		{
-			PaintStrength *= (1.0f - InfluencedValue) / InnerBrushPercentage;
-		}
+		float InnerBrushRadius = BrushFalloffAmount * BrushRadius;
+		float PaintStrength = MeshPaintHelpers::ComputePaintMultiplier(BrushSpacePosition2D.SizeSquared(), BrushStrength, InnerBrushRadius, BrushRadius - InnerBrushRadius, 1.0f, 1.0f, 1.0f);
 
 		const T OldValue = InOutValue;
 		InOutValue = FMath::LerpStable(OldValue, PaintValue, PaintStrength);
diff --git a/Engine/Source/Editor/Persona/Private/AnimViewportShowCommands.cpp b/Engine/Source/Editor/Persona/Private/AnimViewportShowCommands.cpp
--- a/Engine/Source/Editor/Persona/Private/AnimViewportShowCommands.cpp
+++ b/Engine/Source/Editor/Persona/Private/AnimViewportShowCommands.cpp
@@ -16,7 +16,7 @@ void FAnimViewportShowCommands::RegisterCommands()
 
 	UI_COMMAND( ShowRetargetBasePose, "Retarget Base Pose", "Show retarget Base pose on preview mesh", EUserInterfaceActionType::ToggleButton, FInputChord() );
 	UI_COMMAND( ShowBound, "Bound", "Show bound on preview mesh", EUserInterfaceActionType::ToggleButton, FInputChord() );
-	UI_COMMAND( UseInGameBound, "In-game Bound", "Use in-game bound on preview mesh", EUserInterfaceActionType::ToggleButton, FInputChord() );
+	UI_COMMAND( UseInGameBound, "In-game Bound", "Use in-game bound on preview mesh when showing bounds. Otherwise bounds will always be calculated from bones alone.", EUserInterfaceActionType::ToggleButton, FInputChord() );
 	UI_COMMAND( ShowPreviewMesh, "Mesh", "Show the preview mesh", EUserInterfaceActionType::ToggleButton, FInputChord() );
 	UI_COMMAND( ShowMorphTargets, "MorphTargets", "Display Applied Morph Targets of the mesh", EUserInterfaceActionType::ToggleButton, FInputChord() );
 
diff --git a/Engine/Source/Editor/Persona/Private/AnimationEditorPreviewScene.cpp b/Engine/Source/Editor/Persona/Private/AnimationEditorPreviewScene.cpp
--- a/Engine/Source/Editor/Persona/Private/AnimationEditorPreviewScene.cpp
+++ b/Engine/Source/Editor/Persona/Private/AnimationEditorPreviewScene.cpp
@@ -32,6 +32,8 @@
 
 FAnimationEditorPreviewScene::FAnimationEditorPreviewScene(const ConstructionValues& CVS, const TSharedRef<IEditableSkeleton>& InEditableSkeleton, const TSharedRef<IPersonaToolkit>& InPersonaToolkit)
 	: IPersonaPreviewScene(CVS)
+	, Actor(nullptr)
+	, SkeletalMeshComponent(nullptr)
 	, EditableSkeletonPtr(InEditableSkeleton)
 	, PersonaToolkit(InPersonaToolkit)
 	, DefaultMode(EPreviewSceneDefaultAnimationMode::ReferencePose)
@@ -49,21 +51,11 @@ FAnimationEditorPreviewScene::FAnimationEditorPreviewScene(const ConstructionVal
 	
 	FloorBounds = FloorMeshComponent->CalcBounds(FloorMeshComponent->GetRelativeTransform());
 
-	Actor = GetWorld()->SpawnActor<AActor>(AActor::StaticClass(), FTransform::Identity);
-	Actor->AddToRoot();
-
-	// Create the preview component
-	SkeletalMeshComponent = NewObject<UDebugSkelMeshComponent>(Actor);
-	AddComponent(SkeletalMeshComponent, FTransform::Identity);
-
-	// set root component, so we can attach to it. 
-	Actor->SetRootComponent(SkeletalMeshComponent);
-
 	InEditableSkeleton->LoadAdditionalPreviewSkeletalMeshes();
 
 	// create the preview scene description
 	PreviewSceneDescription = NewObject<UPersonaPreviewSceneDescription>(GetTransientPackage());
-	PreviewSceneDescription->AddToRoot();
+	PreviewSceneDescription->SetFlags(RF_Transactional);
 
 	PreviewSceneDescription->AnimationMode = EPreviewAnimationMode::Default;
 	PreviewSceneDescription->Animation = InPersonaToolkit->GetAnimationAsset();
@@ -80,10 +72,9 @@ FAnimationEditorPreviewScene::FAnimationEditorPreviewScene(const ConstructionVal
 		PreviewSceneDescription->AdditionalMeshes = PreviewSceneDescription->DefaultAdditionalMeshes;
 	}
 
-	// Force validation of preview attached assets (catch case of never doing it if we dont have a valid preview mesh)
-	ValidatePreviewAttachedAssets(nullptr);
-
-	RefreshAdditionalMeshes();
+	// Disable killing actors outside of the world
+	AWorldSettings* WorldSettings = GetWorld()->GetWorldSettings(true);
+	WorldSettings->bEnableWorldBoundsChecks = false;
 }
 
 FAnimationEditorPreviewScene::~FAnimationEditorPreviewScene()
@@ -92,9 +83,6 @@ FAnimationEditorPreviewScene::~FAnimationEditorPreviewScene()
 	{
 		GEditor->UnregisterForUndo(this);
 	}
-
-	SkeletalMeshComponent->RemoveFromRoot();
-	PreviewSceneDescription->RemoveFromRoot();
 }
 
 void FAnimationEditorPreviewScene::SetPreviewMesh(USkeletalMesh* NewPreviewMesh)
@@ -148,7 +136,7 @@ void FAnimationEditorPreviewScene::SetPreviewMeshInternal(USkeletalMesh* NewPrev
 	USkeletalMeshComponent* DebuggedSkeletalMeshComponent = nullptr;
 	if(SkeletalMeshComponent->GetAnimInstance())
 	{
-		UAnimBlueprint* SourceBlueprint = Cast<UAnimBlueprint>(SkeletalMeshComponent->GetAnimInstance()->GetClass()->ClassGeneratedBy);
+		UAnimBlueprint* SourceBlueprint = PersonaToolkit.Pin()->GetAnimBlueprint();
 		if(SourceBlueprint)
 		{
 			UAnimInstance* DebuggedAnimInstance = Cast<UAnimInstance>(SourceBlueprint->GetObjectBeingDebugged());
@@ -223,7 +211,7 @@ void FAnimationEditorPreviewScene::SetPreviewMeshInternal(USkeletalMesh* NewPrev
 	// with the AnimBlueprint
 	if (DebuggedSkeletalMeshComponent)
 	{
-		UAnimBlueprint* SourceBlueprint = CastChecked<UAnimBlueprint>(SkeletalMeshComponent->GetAnimInstance()->GetClass()->ClassGeneratedBy);
+		UAnimBlueprint* SourceBlueprint = PersonaToolkit.Pin()->GetAnimBlueprint();
 		SourceBlueprint->SetObjectBeingDebugged(DebuggedSkeletalMeshComponent->GetAnimInstance());
 	}
 
@@ -940,6 +928,12 @@ void FAnimationEditorPreviewScene::TogglePlayback()
 	}
 }
 
+void FAnimationEditorPreviewScene::SetActor(AActor* InActor)
+{
+	check(Actor == nullptr || !Actor->IsRooted());
+	Actor = InActor;
+}
+
 AActor* FAnimationEditorPreviewScene::GetActor() const
 {
 	return Actor;
@@ -995,10 +989,22 @@ void FAnimationEditorPreviewScene::RemoveComponent(class UActorComponent* Compon
 
 void FAnimationEditorPreviewScene::PostUndo(bool bSuccess)
 {
-	// refresh skeletal mesh
-	if (SkeletalMeshComponent)
+	// refresh skeletal mesh & animation
+	if (PreviewSceneDescription)
 	{
-		SetPreviewMesh(SkeletalMeshComponent->SkeletalMesh);
+		SetPreviewMesh(PreviewSceneDescription->PreviewMesh.Get());
+		switch (PreviewSceneDescription->AnimationMode)
+		{
+		case EPreviewAnimationMode::Default:
+			ShowDefaultMode();
+			break;
+		case EPreviewAnimationMode::ReferencePose:
+			ShowReferencePose(true);
+			break;
+		case EPreviewAnimationMode::UseSpecificAnimation:
+			SetPreviewAnimationAsset(Cast<UAnimationAsset>(PreviewSceneDescription->Animation.LoadSynchronous()));
+			break;
+		}
 	}
 }
 
@@ -1007,4 +1013,14 @@ void FAnimationEditorPreviewScene::PostRedo(bool bSuccess)
 	PostUndo(bSuccess);
 }
 
+void FAnimationEditorPreviewScene::AddReferencedObjects( FReferenceCollector& Collector )
+{
+	IPersonaPreviewScene::AddReferencedObjects(Collector);
+
+	Collector.AddReferencedObject(PreviewSceneDescription);
+	Collector.AddReferencedObject(Actor);
+	Collector.AddReferencedObject(SkeletalMeshComponent);
+	Collector.AddReferencedObjects(AdditionalMeshes);
+}
+
 #undef LOCTEXT_NAMESPACE
diff --git a/Engine/Source/Editor/Persona/Private/AnimationEditorPreviewScene.h b/Engine/Source/Editor/Persona/Private/AnimationEditorPreviewScene.h
--- a/Engine/Source/Editor/Persona/Private/AnimationEditorPreviewScene.h
+++ b/Engine/Source/Editor/Persona/Private/AnimationEditorPreviewScene.h
@@ -31,6 +31,7 @@ public:
 	virtual void InvalidateViews() override;
 	virtual void FocusViews() override;
 	virtual UDebugSkelMeshComponent* GetPreviewMeshComponent() const override { return SkeletalMeshComponent; }
+	virtual void SetPreviewMeshComponent(UDebugSkelMeshComponent* InSkeletalMeshComponent) override { SkeletalMeshComponent = InSkeletalMeshComponent; }
 	virtual void SetAdditionalMeshes(class UDataAsset* InAdditionalMeshes) override;
 	virtual void RefreshAdditionalMeshes() override;
 	virtual void ShowReferencePose(bool bReferencePose) override;
@@ -116,6 +117,7 @@ public:
 	virtual int32 GetSelectedBoneIndex() const override;
 	virtual void TogglePlayback() override;
 	virtual AActor* GetActor() const override;
+	virtual void SetActor(AActor* InActor) override;
 	virtual bool AllowMeshHitProxies() const override;
 	virtual void SetAllowMeshHitProxies(bool bState) override;
 
@@ -146,6 +148,9 @@ public:
 	virtual void PostUndo(bool bSuccess) override;
 	virtual void PostRedo(bool bSuccess) override;
 
+	/** FGCObject interface */
+	virtual void AddReferencedObjects( FReferenceCollector& Collector ) override;
+
 	/** Validate preview attached assets on skeleton and supplied skeletal mesh, notifying user if any are removed */
 	void ValidatePreviewAttachedAssets(USkeletalMesh* PreviewSkeletalMesh);
 
diff --git a/Engine/Source/Editor/Persona/Private/AnimationEditorViewportClient.cpp b/Engine/Source/Editor/Persona/Private/AnimationEditorViewportClient.cpp
--- a/Engine/Source/Editor/Persona/Private/AnimationEditorViewportClient.cpp
+++ b/Engine/Source/Editor/Persona/Private/AnimationEditorViewportClient.cpp
@@ -94,8 +94,6 @@ FAnimationViewportClient::FAnimationViewportClient(const TSharedRef<ISkeletonTre
 	WidgetMode = FWidget::WM_Rotate;
 	ModeTools->SetWidgetMode(WidgetMode);
 
-	SetSelectedBackgroundColor(ConfigOption->ViewportBackgroundColor, false);
-
 	EngineShowFlags.Game = 0;
 	EngineShowFlags.ScreenSpaceReflections = 1;
 	EngineShowFlags.AmbientOcclusion = 1;
@@ -175,56 +173,6 @@ FAnimationViewportClient::~FAnimationViewportClient()
 	UAssetViewerSettings::Get()->OnAssetViewerSettingsChanged().RemoveAll(this);
 }
 
-FLinearColor FAnimationViewportClient::GetBackgroundColor() const
-{
-	return SelectedHSVColor.HSVToLinearRGB();
-}
-
-void FAnimationViewportClient::SetSelectedBackgroundColor(const FLinearColor& RGBColor, bool bSave/*=true*/)
-{
-	SelectedHSVColor = RGBColor.LinearRGBToHSV(); 
-
-	// save to config
-	if (bSave)
-	{
-		ConfigOption->SetViewportBackgroundColor(RGBColor);
-	}
-
-	// only consider V from HSV
-	// just inversing works against for mid range brightness
-	// V being from 0-0.3 (almost white), 0.31-1 (almost black)
-	if (SelectedHSVColor.B < 0.3f)
-	{
-		DrawHelper.GridColorAxis = FColor(230,230,230);
-		DrawHelper.GridColorMajor = FColor(180,180,180);
-		DrawHelper.GridColorMinor = FColor(128,128,128);
-	}
-	else
-	{
-		DrawHelper.GridColorAxis = FColor(20,20,20);
-		DrawHelper.GridColorMajor = FColor(60,60,60);
-		DrawHelper.GridColorMinor = FColor(128,128,128);
-	}
-
-	Invalidate();
-}
-
-void FAnimationViewportClient::SetBackgroundColor( FLinearColor InColor )
-{
-	SetSelectedBackgroundColor(InColor);
-}
-
-float FAnimationViewportClient::GetBrightnessValue() const
-{
-	return SelectedHSVColor.B;
-}
-
-void FAnimationViewportClient::SetBrightnessValue( float Value )
-{
-	SelectedHSVColor.B = Value;
-	SetSelectedBackgroundColor(SelectedHSVColor.HSVToLinearRGB());
-}
-
 void FAnimationViewportClient::OnToggleShowGrid()
 {
 	FEditorViewportClient::SetShowGrid();
@@ -394,8 +342,9 @@ void FAnimationViewportClient::HandleSkeletalMeshChanged(USkeletalMesh* OldSkele
 		PreviewMeshComponent->TermArticulated();
 		PreviewMeshComponent->InitArticulated(GetWorld()->GetPhysicsScene());
 
-		// Set to block all to enable tracing.
-		PreviewMeshComponent->SetCollisionProfileName(UCollisionProfile::BlockAll_ProfileName);
+		// Set to PhysicsActor to enable tracing regardless of project overrides
+		static FName CollisionProfileName(TEXT("PhysicsActor"));
+		PreviewMeshComponent->SetCollisionProfileName(CollisionProfileName);
 	}
 
 	Invalidate();
@@ -491,7 +440,10 @@ void FAnimationViewportClient::DrawCanvas( FViewport& InViewport, FSceneView& Vi
 		{
 			// Allow edit modes (inc. skeletal control modes) to draw with the canvas, and collect on screen strings to draw later
 			TArray<FText> EditModeDebugText;
-			GetPersonaModeManager().GetOnScreenDebugInfo(EditModeDebugText);
+			if (GetPersonaModeManager())
+			{
+				GetPersonaModeManager()->GetOnScreenDebugInfo(EditModeDebugText);
+			}
 
 			// Draw Node info instead of mesh info if we have entries
 			DrawNodeDebugLines(EditModeDebugText, &Canvas, &View);
@@ -1006,16 +958,13 @@ void FAnimationViewportClient::DrawNodeDebugLines(TArray<FText>& Lines, FCanvas*
 
 void FAnimationViewportClient::TrackingStarted( const struct FInputEventState& InInputState, bool bIsDraggingWidget, bool bNudge )
 {
-	if (ModeTools->StartTracking(this, Viewport) && bIsDraggingWidget)
-	{
-		Widget->SetSnapEnabled(true);
-	}
+	ModeTools->StartTracking(this, Viewport);
 }
 
 void FAnimationViewportClient::TrackingStopped() 
 {
 	ModeTools->EndTracking(this, Viewport);
-	Widget->SetSnapEnabled(false);
+
 	Invalidate();
 }
 
@@ -1492,10 +1441,13 @@ FSphere FAnimationViewportClient::GetCameraTarget()
 	PreviewMeshComponent->CalcBounds(PreviewMeshComponent->GetComponentTransform());
 
 	// give the editor mode a chance to give us a camera target
-	FSphere Target;
-	if (GetPersonaModeManager().GetCameraTarget(Target))
+	if (GetPersonaModeManager())
 	{
-		return Target;
+		FSphere Target;
+		if (GetPersonaModeManager()->GetCameraTarget(Target))
+		{
+			return Target;
+		}
 	}
 
 	FBoxSphereBounds Bounds = PreviewMeshComponent->CalcBounds(FTransform::Identity);
@@ -1569,7 +1521,17 @@ void FAnimationViewportClient::GetAllVertexIndicesUsedInSection(const FRawStatic
 
 bool FAnimationViewportClient::PreviewComponentSelectionOverride(const UPrimitiveComponent* InComponent) const
 {
-	return InComponent == GetPreviewScene()->GetPreviewMeshComponent();
+	if (InComponent == GetPreviewScene()->GetPreviewMeshComponent())
+	{
+		const USkeletalMeshComponent* Component = CastChecked<USkeletalMeshComponent>(InComponent);
+		USkeletalMesh* Mesh = Component->SkeletalMesh;
+		if (Mesh)
+		{
+			return (Mesh->SelectedEditorSection != INDEX_NONE);
+		}
+	}
+
+	return false;
 }
 
 FBox FAnimationViewportClient::ComputeBoundingBoxForSelectedEditorSection() const
@@ -1866,9 +1828,9 @@ TSharedRef<FAnimationEditorPreviewScene> FAnimationViewportClient::GetAnimPrevie
 	return StaticCastSharedRef<FAnimationEditorPreviewScene>(GetPreviewScene());
 }
 
-IPersonaEditorModeManager& FAnimationViewportClient::GetPersonaModeManager() const
+IPersonaEditorModeManager* FAnimationViewportClient::GetPersonaModeManager() const
 {
-	return *static_cast<IPersonaEditorModeManager*>(ModeTools);
+	return static_cast<IPersonaEditorModeManager*>(ModeTools);
 }
 
 void FAnimationViewportClient::HandleInvalidateViews()
diff --git a/Engine/Source/Editor/Persona/Private/AnimationEditorViewportClient.h b/Engine/Source/Editor/Persona/Private/AnimationEditorViewportClient.h
--- a/Engine/Source/Editor/Persona/Private/AnimationEditorViewportClient.h
+++ b/Engine/Source/Editor/Persona/Private/AnimationEditorViewportClient.h
@@ -128,7 +128,6 @@ public:
 	virtual ~FAnimationViewportClient();
 
 	// FEditorViewportClient interface
-	virtual FLinearColor GetBackgroundColor() const override;
 	virtual void Tick(float DeltaSeconds) override;
 	virtual void Draw(const FSceneView* View,FPrimitiveDrawInterface* PDI) override;
 	virtual void DrawCanvas( FViewport& InViewport, FSceneView& View, FCanvas& Canvas ) override;
@@ -381,8 +380,6 @@ private:
 
 private:
 
-	void SetSelectedBackgroundColor(const FLinearColor& RGBColor, bool bSave = true);
-
 	void SetCameraTargetLocation(const FSphere &BoundSphere, float DeltaSeconds);
 
 	/** Draws Mesh Bones in foreground **/
@@ -413,7 +410,7 @@ private:
 	TSharedRef<class FAnimationEditorPreviewScene> GetAnimPreviewScene() const;
 
 	/** Get the persona mode manager */
-	class IPersonaEditorModeManager& GetPersonaModeManager() const;
+	class IPersonaEditorModeManager* GetPersonaModeManager() const;
 
 	/** Invalidate this view in response to a preview scene change */
 	void HandleInvalidateViews();
diff --git a/Engine/Source/Editor/Persona/Private/Customization/BlendSampleDetails.cpp b/Engine/Source/Editor/Persona/Private/Customization/BlendSampleDetails.cpp
--- a/Engine/Source/Editor/Persona/Private/Customization/BlendSampleDetails.cpp
+++ b/Engine/Source/Editor/Persona/Private/Customization/BlendSampleDetails.cpp
@@ -95,9 +95,9 @@ void FBlendSampleDetails::GenerateBlendSampleWidget(TFunction<FDetailWidgetRow&
 	const int32 NumParameters = BlendSpace->IsA<UBlendSpace1D>() ? 1 : 2;
 	for (int32 ParameterIndex = 0; ParameterIndex < NumParameters; ++ParameterIndex)
 	{
-		const FBlendParameter& BlendParameter = BlendSpace->GetBlendParameter(ParameterIndex);
-		auto ValueChangedLambda = [BlendSpace, SampleIndex, ParameterIndex, BlendParameter, OnSampleMoved](const float NewValue, bool bIsInteractive)
+		auto ValueChangedLambda = [BlendSpace, SampleIndex, ParameterIndex, OnSampleMoved](const float NewValue, bool bIsInteractive)
 		{
+			const FBlendParameter& BlendParameter = BlendSpace->GetBlendParameter(ParameterIndex);
 			const FBlendSample& Sample = BlendSpace->GetBlendSample(SampleIndex);
 			FVector SampleValue = Sample.SampleValue;
 
@@ -122,7 +122,7 @@ void FBlendSampleDetails::GenerateBlendSampleWidget(TFunction<FDetailWidgetRow&
 		[
 			SNew(STextBlock)
 			.Font(FEditorStyle::GetFontStyle(TEXT("PropertyWindow.NormalFont")))
-			.Text(FText::FromString(BlendParameter.DisplayName))
+			.Text_Lambda([BlendSpace, ParameterIndex]() { return FText::FromString(BlendSpace->GetBlendParameter(ParameterIndex).DisplayName); })
 		];
 
 		ParameterRow.ValueContent()
@@ -158,16 +158,16 @@ void FBlendSampleDetails::GenerateBlendSampleWidget(TFunction<FDetailWidgetRow&
 			})
 			.LabelVAlign(VAlign_Center)
 			.AllowSpin(true)
-			.MinValue_Lambda([BlendParameter]() -> float { return BlendParameter.Min; })
-			.MaxValue_Lambda([BlendParameter]() -> float { return BlendParameter.Max; })
-			.MinSliderValue_Lambda([BlendParameter]() -> float { return BlendParameter.Min; })
-			.MaxSliderValue_Lambda([BlendParameter]() -> float { return BlendParameter.Max; })
+			.MinValue_Lambda([BlendSpace, ParameterIndex]() -> float { return BlendSpace->GetBlendParameter(ParameterIndex).Min; })
+			.MaxValue_Lambda([BlendSpace, ParameterIndex]() -> float { return BlendSpace->GetBlendParameter(ParameterIndex).Max; })
+			.MinSliderValue_Lambda([BlendSpace, ParameterIndex]() -> float { return BlendSpace->GetBlendParameter(ParameterIndex).Min; })
+			.MaxSliderValue_Lambda([BlendSpace, ParameterIndex]() -> float { return BlendSpace->GetBlendParameter(ParameterIndex).Max; })
 			.MinDesiredValueWidth(60.0f)
 			.Label()
 			[
 				SNew(STextBlock)
 				.Visibility(bShowLabel ? EVisibility::Visible : EVisibility::Collapsed)
-				.Text_Lambda([=]() { return FText::FromString(BlendParameter.DisplayName); })
+				.Text_Lambda([BlendSpace, ParameterIndex]() { return FText::FromString(BlendSpace->GetBlendParameter(ParameterIndex).DisplayName); })
 			]
 		];
 	}
@@ -215,7 +215,7 @@ bool FBlendSampleDetails::ShouldFilterAssetStatic(const FAssetData& AssetData, c
 	if (AssetData.GetTagValue(SkeletonTagName, SkeletonName))
 	{
 		// Check whether or not the skeletons match
-		if (SkeletonName == BlendSpaceBase->GetSkeleton()->GetPathName())
+		if (SkeletonName.Contains(BlendSpaceBase->GetSkeleton()->GetPathName()))
 		{
 			// If so check if the additive animation tpye is compatible with the blend space
 			const FName AdditiveTypeTagName = GET_MEMBER_NAME_CHECKED(UAnimSequence, AdditiveAnimType);
diff --git a/Engine/Source/Editor/Persona/Private/Customization/BlendSpaceDetails.cpp b/Engine/Source/Editor/Persona/Private/Customization/BlendSpaceDetails.cpp
--- a/Engine/Source/Editor/Persona/Private/Customization/BlendSpaceDetails.cpp
+++ b/Engine/Source/Editor/Persona/Private/Customization/BlendSpaceDetails.cpp
@@ -23,21 +23,10 @@ FBlendSpaceDetails::FBlendSpaceDetails()
 {
 	Builder = nullptr;
 	BlendSpaceBase = nullptr;
-	Handle = FCoreUObjectDelegates::OnObjectPropertyChanged.AddLambda([this](UObject* Object, struct FPropertyChangedEvent& Event)
-	{
-		if(Event.ChangeType != EPropertyChangeType::Interactive)
-		{
-			if (Builder && Object == BlendSpaceBase && (Event.Property == nullptr || (Event.MemberProperty && Event.MemberProperty->GetFName() == GET_MEMBER_NAME_CHECKED(UBlendSpaceBase, BlendParameters) && Event.Property && Event.Property->GetFName() == GET_MEMBER_NAME_CHECKED(FBlendParameter, DisplayName))))
-			{
-				Builder->ForceRefreshDetails();
-			}
-		}
-	});
 }
 
 FBlendSpaceDetails::~FBlendSpaceDetails()
 {
-	FCoreUObjectDelegates::OnObjectPropertyChanged.Remove(Handle);
 }
 
 void FBlendSpaceDetails::CustomizeDetails(class IDetailLayoutBuilder& DetailBuilder)
@@ -92,14 +81,18 @@ void FBlendSpaceDetails::CustomizeDetails(class IDetailLayoutBuilder& DetailBuil
 			DefaultProperty->MarkHiddenByCustomization();
 		}
 
+		FSimpleDelegate RefreshDelegate = FSimpleDelegate::CreateLambda([this]() { Builder->ForceRefreshDetails(); });
+
 		// Retrieve blend samples array
 		TSharedPtr<IPropertyHandleArray> BlendSamplesArrayProperty = DetailBuilder.GetProperty(GET_MEMBER_NAME_CHECKED(UBlendSpaceBase, SampleData), UBlendSpaceBase::StaticClass())->AsArray();
-
+		BlendSamplesArrayProperty->SetOnNumElementsChanged(RefreshDelegate);
+		
 		uint32 NumBlendSampleEntries = 0;
 		BlendSamplesArrayProperty->GetNumElements(NumBlendSampleEntries);
 		for (uint32 SampleIndex = 0; SampleIndex < NumBlendSampleEntries; ++SampleIndex)
 		{
 			TSharedPtr<IPropertyHandle> BlendSampleProperty = BlendSamplesArrayProperty->GetElement(SampleIndex);
+			BlendSampleProperty->SetOnChildPropertyValueChanged(RefreshDelegate);
 			TSharedPtr<IPropertyHandle> AnimationProperty = BlendSampleProperty->GetChildHandle(GET_MEMBER_NAME_CHECKED(FBlendSample, Animation));
 			TSharedPtr<IPropertyHandle> SampleValueProperty = BlendSampleProperty->GetChildHandle(GET_MEMBER_NAME_CHECKED(FBlendSample, SampleValue));
 			TSharedPtr<IPropertyHandle> RateScaleProperty = BlendSampleProperty->GetChildHandle(GET_MEMBER_NAME_CHECKED(FBlendSample, RateScale));
@@ -113,11 +106,11 @@ void FBlendSpaceDetails::CustomizeDetails(class IDetailLayoutBuilder& DetailBuil
 				[
 					SNew(STextBlock)
 					.Font(DetailBuilder.GetDetailFont())
-					.Text_Lambda([AnimationProperty]() -> FText
+					.Text_Lambda([AnimationProperty, SampleIndex]() -> FText
 					{
 						FAssetData AssetData;
 						AnimationProperty->GetValue(AssetData);
-						return AssetData.IsValid() ? FText::FromString(AssetData.GetAsset()->GetName()) : FText::FromString("No Animation");
+						return AssetData.IsValid() ? FText::Format(LOCTEXT("BlendSpaceAnimationNameLabel", "{0} ({1})"), FText::FromString(AssetData.GetAsset()->GetName()), FText::FromString(FString::FromInt(SampleIndex))) : FText::FromString("No Animation");
 					})
 				]
 			];
diff --git a/Engine/Source/Editor/Persona/Private/Customization/BlendSpaceDetails.h b/Engine/Source/Editor/Persona/Private/Customization/BlendSpaceDetails.h
--- a/Engine/Source/Editor/Persona/Private/Customization/BlendSpaceDetails.h
+++ b/Engine/Source/Editor/Persona/Private/Customization/BlendSpaceDetails.h
@@ -18,9 +18,7 @@ public:
 
 	// IDetailCustomization interface
 	virtual void CustomizeDetails(class IDetailLayoutBuilder& DetailBuilder) override;
-
 private:
-	FDelegateHandle Handle;
 	class IDetailLayoutBuilder* Builder;
 	class UBlendSpaceBase* BlendSpaceBase;
 };
diff --git a/Engine/Source/Editor/Persona/Private/EditModes/SkeletonSelectionEditMode.cpp b/Engine/Source/Editor/Persona/Private/EditModes/SkeletonSelectionEditMode.cpp
--- a/Engine/Source/Editor/Persona/Private/EditModes/SkeletonSelectionEditMode.cpp
+++ b/Engine/Source/Editor/Persona/Private/EditModes/SkeletonSelectionEditMode.cpp
@@ -409,7 +409,7 @@ bool FSkeletonSelectionEditMode::UsesTransformWidget() const
 
 bool FSkeletonSelectionEditMode::UsesTransformWidget(FWidget::EWidgetMode CheckMode) const
 {
-	return CheckMode == FWidget::WM_Scale || CheckMode == FWidget::WM_Translate || CheckMode == FWidget::WM_Rotate;
+	return ShouldDrawWidget() && (CheckMode == FWidget::WM_Scale || CheckMode == FWidget::WM_Translate || CheckMode == FWidget::WM_Rotate);
 }
 
 bool FSkeletonSelectionEditMode::GetCustomDrawingCoordinateSystem(FMatrix& InMatrix, void* InData)
diff --git a/Engine/Source/Editor/Persona/Private/PersonaAssetFamily.cpp b/Engine/Source/Editor/Persona/Private/PersonaAssetFamily.cpp
--- a/Engine/Source/Editor/Persona/Private/PersonaAssetFamily.cpp
+++ b/Engine/Source/Editor/Persona/Private/PersonaAssetFamily.cpp
@@ -5,6 +5,7 @@
 #include "ARFilter.h"
 #include "AssetRegistryModule.h"
 #include "Animation/AnimBlueprint.h"
+#include "PhysicsEngine/PhysicsAsset.h"
 
 #define LOCTEXT_NAMESPACE "PersonaAssetFamily"
 
@@ -19,30 +20,25 @@ FPersonaAssetFamily::FPersonaAssetFamily(const UObject* InFromObject)
 		if (InFromObject->IsA<USkeleton>())
 		{
 			Skeleton = CastChecked<USkeleton>(InFromObject);
-			Mesh = Skeleton->GetPreviewMesh();
 		}
 		else if (InFromObject->IsA<UAnimationAsset>())
 		{
 			AnimationAsset = CastChecked<UAnimationAsset>(InFromObject);
-			Skeleton = AnimationAsset->GetSkeleton();
-			Mesh = AnimationAsset->GetPreviewMesh();
-			if (Mesh == nullptr)
-			{
-				Mesh = Skeleton->GetPreviewMesh();
-			}
 		}
 		else if (InFromObject->IsA<USkeletalMesh>())
 		{
 			Mesh = CastChecked<USkeletalMesh>(InFromObject);
-			Skeleton = Mesh->Skeleton;
 		}
 		else if (InFromObject->IsA<UAnimBlueprint>())
 		{
 			AnimBlueprint = CastChecked<UAnimBlueprint>(InFromObject);
-			Skeleton = AnimBlueprint->TargetSkeleton;
-			check(AnimBlueprint->TargetSkeleton);
-			Mesh = AnimBlueprint->TargetSkeleton->GetPreviewMesh();
 		}
+		else if (InFromObject->IsA<UPhysicsAsset>())
+		{
+			PhysicsAsset = CastChecked<UPhysicsAsset>(InFromObject);
+		}
+
+		FindCounterpartAssets(InFromObject, Skeleton, Mesh);
 	}
 }
 
@@ -53,6 +49,7 @@ void FPersonaAssetFamily::GetAssetTypes(TArray<UClass*>& OutAssetTypes) const
 	OutAssetTypes.Add(USkeletalMesh::StaticClass());
 	OutAssetTypes.Add(UAnimationAsset::StaticClass());
 	OutAssetTypes.Add(UAnimBlueprint::StaticClass());
+	OutAssetTypes.Add(UPhysicsAsset::StaticClass());
 }
 
 template<typename AssetType>
@@ -124,6 +121,33 @@ FAssetData FPersonaAssetFamily::FindAssetOfType(UClass* InAssetClass) const
 				}
 			}
 		}
+		else if (InAssetClass->IsChildOf<UPhysicsAsset>())
+		{
+			if (PhysicsAsset.IsValid())
+			{
+				return FAssetData(PhysicsAsset.Get());
+			}
+			else
+			{
+				TArray<FAssetData> Assets;
+
+				FAssetRegistryModule& AssetRegistryModule = FModuleManager::LoadModuleChecked<FAssetRegistryModule>("AssetRegistry");
+				FARFilter Filter;
+				Filter.bRecursiveClasses = true;
+				Filter.ClassNames.Add(UPhysicsAsset::StaticClass()->GetFName());
+				if(Mesh.IsValid())
+				{
+					Filter.TagsAndValues.Add(GET_MEMBER_NAME_CHECKED(UPhysicsAsset, PreviewSkeletalMesh), FAssetData(Mesh.Get()).ObjectPath.ToString());
+				}
+
+				AssetRegistryModule.Get().GetAssets(Filter, Assets);
+
+				if (Assets.Num() > 0)
+				{
+					return Assets[0];
+				}
+			}
+		}
 	}
 
 	return FAssetData();
@@ -150,6 +174,19 @@ void FPersonaAssetFamily::FindAssetsOfType(UClass* InAssetClass, TArray<FAssetDa
 		{
 			FindAssets<UAnimBlueprint>(Skeleton.Get(), OutAssets, "TargetSkeleton");
 		}
+		else if (InAssetClass->IsChildOf<UPhysicsAsset>())
+		{
+			FAssetRegistryModule& AssetRegistryModule = FModuleManager::LoadModuleChecked<FAssetRegistryModule>("AssetRegistry");
+			FARFilter Filter;
+			Filter.bRecursiveClasses = true;
+			Filter.ClassNames.Add(UPhysicsAsset::StaticClass()->GetFName());
+			if(Mesh != nullptr)
+			{
+				Filter.TagsAndValues.Add(GET_MEMBER_NAME_CHECKED(UPhysicsAsset, PreviewSkeletalMesh), FAssetData(Mesh.Get()).ObjectPath.ToString());
+			}
+
+			AssetRegistryModule.Get().GetAssets(Filter, OutAssets);
+		}
 	}
 }
 
@@ -173,6 +210,10 @@ FText FPersonaAssetFamily::GetAssetTypeDisplayName(UClass* InAssetClass) const
 		{
 			return LOCTEXT("AnimBlueprintAssetDisplayName", "Blueprint");
 		}
+		else if (InAssetClass->IsChildOf<UPhysicsAsset>())
+		{
+			return LOCTEXT("PhysicsAssetDisplayName", "Physics");
+		}
 	}
 
 	return FText();
@@ -203,6 +244,14 @@ bool FPersonaAssetFamily::IsAssetCompatible(const FAssetData& InAssetData) const
 				return *TargetSkeleton == FAssetData(Skeleton.Get()).GetExportTextName();
 			}
 		}
+		else if (Class->IsChildOf<UPhysicsAsset>())
+		{
+			const FString* PreviewMesh = InAssetData.TagsAndValues.Find(GET_MEMBER_NAME_CHECKED(UPhysicsAsset, PreviewSkeletalMesh));
+			if (PreviewMesh && Mesh.IsValid())
+			{
+				return *PreviewMesh == FAssetData(Mesh.Get()).ObjectPath.ToString();
+			}
+		}
 	}
 
 	return false;
@@ -228,6 +277,10 @@ UClass* FPersonaAssetFamily::GetAssetFamilyClass(UClass* InClass) const
 		{
 			return UAnimBlueprint::StaticClass();
 		}
+		else if (InClass->IsChildOf<UPhysicsAsset>())
+		{
+			return UPhysicsAsset::StaticClass();
+		}
 	}
 
 	return nullptr;
@@ -256,11 +309,76 @@ void FPersonaAssetFamily::RecordAssetOpened(const FAssetData& InAssetData)
 			{
 				AnimBlueprint = Cast<UAnimBlueprint>(InAssetData.GetAsset());
 			}
+			else if (Class->IsChildOf<UPhysicsAsset>())
+			{
+				PhysicsAsset = Cast<UPhysicsAsset>(InAssetData.GetAsset());
+			}
 		}
 
 		OnAssetOpened.Broadcast(InAssetData.GetAsset());
 	}
 }
 
+void FPersonaAssetFamily::FindCounterpartAssets(const UObject* InAsset, TWeakObjectPtr<const USkeleton>& OutSkeleton, TWeakObjectPtr<const USkeletalMesh>& OutMesh)
+{
+	const USkeleton* CounterpartSkeleton = OutSkeleton.Get();
+	const USkeletalMesh* CounterpartMesh = OutMesh.Get();
+	FindCounterpartAssets(InAsset, CounterpartSkeleton, CounterpartMesh);
+	OutSkeleton = CounterpartSkeleton;
+	OutMesh = CounterpartMesh;
+}
+
+void FPersonaAssetFamily::FindCounterpartAssets(const UObject* InAsset, const USkeleton*& OutSkeleton, const USkeletalMesh*& OutMesh)
+{
+	if (InAsset->IsA<USkeleton>())
+	{
+		OutSkeleton = CastChecked<USkeleton>(InAsset);
+		OutMesh = OutSkeleton->GetPreviewMesh();
+		if(OutMesh == nullptr)
+		{
+			OutMesh = OutSkeleton->FindCompatibleMesh();
+		}
+	}
+	else if (InAsset->IsA<UAnimationAsset>())
+	{
+		const UAnimationAsset* AnimationAsset = CastChecked<const UAnimationAsset>(InAsset);
+		OutSkeleton = AnimationAsset->GetSkeleton();
+		OutMesh = AnimationAsset->GetPreviewMesh();
+		if (OutMesh == nullptr)
+		{
+			OutMesh = OutSkeleton->GetPreviewMesh();
+		}
+		if(OutMesh == nullptr)
+		{
+			OutMesh = OutSkeleton->FindCompatibleMesh();
+		}
+	}
+	else if (InAsset->IsA<USkeletalMesh>())
+	{
+		OutMesh = CastChecked<USkeletalMesh>(InAsset);
+		OutSkeleton = OutMesh->Skeleton;
+	}
+	else if (InAsset->IsA<UAnimBlueprint>())
+	{
+		const UAnimBlueprint* AnimBlueprint = CastChecked<const UAnimBlueprint>(InAsset);
+		OutSkeleton = AnimBlueprint->TargetSkeleton;
+		check(AnimBlueprint->TargetSkeleton);
+		OutMesh = AnimBlueprint->TargetSkeleton->GetPreviewMesh();
+		if(OutMesh == nullptr)
+		{
+			OutMesh = OutSkeleton->FindCompatibleMesh();
+		}
+	}
+	else if (InAsset->IsA<UPhysicsAsset>())
+	{
+		const UPhysicsAsset* PhysicsAsset = CastChecked<const UPhysicsAsset>(InAsset);
+		OutMesh = PhysicsAsset->PreviewSkeletalMesh.LoadSynchronous();
+		if(OutMesh != nullptr)
+		{
+			OutSkeleton = OutMesh->Skeleton;
+		}
+	}
+}
+
 #undef LOCTEXT_NAMESPACE
  
diff --git a/Engine/Source/Editor/Persona/Private/PersonaAssetFamily.h b/Engine/Source/Editor/Persona/Private/PersonaAssetFamily.h
--- a/Engine/Source/Editor/Persona/Private/PersonaAssetFamily.h
+++ b/Engine/Source/Editor/Persona/Private/PersonaAssetFamily.h
@@ -10,6 +10,7 @@
 #include "IAssetFamily.h"
 
 class UAnimBlueprint;
+class UPhysicsAsset;
 
 class FPersonaAssetFamily : public IAssetFamily
 {
@@ -28,6 +29,10 @@ public:
 	DECLARE_DERIVED_EVENT(FPersonaAssetFamily, IAssetFamily::FOnAssetOpened, FOnAssetOpened)
 	virtual FOnAssetOpened& GetOnAssetOpened() override { return OnAssetOpened;  }
 
+	/** Helper functions for constructor and other systems that need to discover meshes/skeletons from related assets */
+	static void FindCounterpartAssets(const UObject* InAsset, TWeakObjectPtr<const USkeleton>& OutSkeleton, TWeakObjectPtr<const USkeletalMesh>& OutMesh);
+	static void FindCounterpartAssets(const UObject* InAsset, const USkeleton*& OutSkeleton, const USkeletalMesh*& OutMesh);
+
 private:
 	/** The skeleton that links all assets */
 	TWeakObjectPtr<const USkeleton> Skeleton;
@@ -41,6 +46,9 @@ private:
 	/** The last animation asset that was encountered */
 	TWeakObjectPtr<const UAnimationAsset> AnimationAsset;
 
+	/** The last physics asset that was encountered */
+	TWeakObjectPtr<const UPhysicsAsset> PhysicsAsset;
+
 	/** Event fired when an asset is opened */
 	FOnAssetOpened OnAssetOpened;
 };
diff --git a/Engine/Source/Editor/Persona/Private/PersonaModule.cpp b/Engine/Source/Editor/Persona/Private/PersonaModule.cpp
--- a/Engine/Source/Editor/Persona/Private/PersonaModule.cpp
+++ b/Engine/Source/Editor/Persona/Private/PersonaModule.cpp
@@ -71,6 +71,7 @@
 #include "BlueprintEditorUtils.h"
 #include "SNotificationList.h"
 #include "NotificationManager.h"
+#include "PhysicsEngine/PhysicsAsset.h"
 
 IMPLEMENT_MODULE( FPersonaModule, Persona );
 
@@ -88,7 +89,7 @@ void FPersonaModule::StartupModule()
 	//Call this to make sure AnimGraph module is setup
 	FModuleManager::Get().LoadModuleChecked(TEXT("AnimGraph"));
 
-	// Make sure the advanced preview scene module is loaded
+	// Make sure the advanced preview scene module is loaded 
 	FModuleManager::Get().LoadModuleChecked("AdvancedPreviewScene");
 
 	// Load all blueprint animnotifies from asset registry so they are available from drop downs in anim segment detail views
@@ -169,7 +170,7 @@ static void SetupPersonaToolkit(const TSharedRef<FPersonaToolkit>& Toolkit, cons
 {
 	if (PersonaToolkitArgs.bCreatePreviewScene)
 	{
-		Toolkit->CreatePreviewScene();
+		Toolkit->CreatePreviewScene(PersonaToolkitArgs);
 	}
 }
 
@@ -217,6 +218,17 @@ TSharedRef<IPersonaToolkit> FPersonaModule::CreatePersonaToolkit(UAnimBlueprint*
 	return NewPersonaToolkit;
 }
 
+TSharedRef<IPersonaToolkit> FPersonaModule::CreatePersonaToolkit(UPhysicsAsset* InPhysicsAsset, const FPersonaToolkitArgs& PersonaToolkitArgs) const
+{
+	TSharedRef<FPersonaToolkit> NewPersonaToolkit(new FPersonaToolkit());
+
+	NewPersonaToolkit->Initialize(InPhysicsAsset);
+
+	SetupPersonaToolkit(NewPersonaToolkit, PersonaToolkitArgs);
+
+	return NewPersonaToolkit;
+}
+
 TSharedRef<class IAssetFamily> FPersonaModule::CreatePersonaAssetFamily(const UObject* InAsset) const
 {
 	return FPersonaAssetFamilyManager::Get().CreatePersonaAssetFamily(InAsset);
@@ -232,9 +244,9 @@ TSharedRef<class FWorkflowTabFactory> FPersonaModule::CreateDetailsTabFactory(co
 	return MakeShareable(new FPersonaDetailsTabSummoner(InHostingApp, InOnDetailsCreated));
 }
 
-TSharedRef<class FWorkflowTabFactory> FPersonaModule::CreatePersonaViewportTabFactory(const TSharedRef<class FWorkflowCentricApplication>& InHostingApp, const TSharedRef<ISkeletonTree>& InSkeletonTree, const TSharedRef<IPersonaPreviewScene>& InPreviewScene, FSimpleMulticastDelegate& OnPostUndo, const TSharedPtr<FBlueprintEditor>& InBlueprintEditor, FOnViewportCreated InOnViewportCreated, bool bInShowTimeline, bool bInShowStats) const
+TSharedRef<class FWorkflowTabFactory> FPersonaModule::CreatePersonaViewportTabFactory(const TSharedRef<class FWorkflowCentricApplication>& InHostingApp, const FPersonaViewportArgs& InArgs) const
 {
-	return MakeShareable(new FPreviewViewportSummoner(InHostingApp, InSkeletonTree, InPreviewScene, OnPostUndo, InBlueprintEditor, InOnViewportCreated, bInShowTimeline, bInShowStats));
+	return MakeShareable(new FPreviewViewportSummoner(InHostingApp, InArgs));
 }
 
 TSharedRef<class FWorkflowTabFactory> FPersonaModule::CreateAnimNotifiesTabFactory(const TSharedRef<class FWorkflowCentricApplication>& InHostingApp, const TSharedRef<class IEditableSkeleton>& InEditableSkeleton, FSimpleMulticastDelegate& InOnChangeAnimNotifies, FSimpleMulticastDelegate& InOnPostUndo, FOnObjectsSelected InOnObjectsSelected) const
@@ -708,10 +720,13 @@ void FPersonaModule::AddCommonToolbarExtensions(FToolBarBuilder& InToolbarBuilde
 		static void HandleApplyPreviewMesh(TSharedPtr<FNotificationHandler> InNotificationHandler, TWeakPtr<IPersonaToolkit> InWeakPersonaToolkit)
 		{
 			TSharedPtr<IPersonaToolkit> PinnedPersonaToolkit = InWeakPersonaToolkit.Pin();
-			PinnedPersonaToolkit->SetPreviewMesh(PinnedPersonaToolkit->GetPreviewScene()->GetPreviewMesh(), true);
-			if(InNotificationHandler->Notification.IsValid())
+			if(PinnedPersonaToolkit.IsValid())	// Toolkit can become invalid while the toast is open
 			{
-				InNotificationHandler->Notification->Fadeout();
+				PinnedPersonaToolkit->SetPreviewMesh(PinnedPersonaToolkit->GetPreviewScene()->GetPreviewMesh(), true);
+				if(InNotificationHandler->Notification.IsValid())
+				{
+					InNotificationHandler->Notification->Fadeout();
+				}
 			}
 		}
 
@@ -732,25 +747,28 @@ void FPersonaModule::AddCommonToolbarExtensions(FToolBarBuilder& InToolbarBuilde
 					WeakPersonaToolkit.Pin()->SetPreviewMesh(Cast<USkeletalMesh>(AssetData.GetAsset()), false);
 				}
 
-				TSharedPtr<FNotificationHandler> NotificationHandler = MakeShared<FNotificationHandler>();
-
-				FNotificationInfo Info(LOCTEXT("PreviewMeshSetTemporarily", "Preview mesh set temporarily"));
-				Info.ExpireDuration = 10.0f;
-				Info.bUseLargeFont = true;
-				Info.ButtonDetails.Add(
-					FNotificationButtonInfo(
-						LOCTEXT("ApplyToAsset", "Apply To Asset"), 
-						LOCTEXT("ApplyToAssetToolTip", "The preview mesh has changed, but it will not be able to be saved until it is applied to the asset. Click here to make the change to the preview mesh persistent."),
-						FSimpleDelegate::CreateStatic(&FNotificationHandler::HandleApplyPreviewMesh, NotificationHandler, WeakPersonaToolkit),
-						SNotificationItem::CS_Success));
-
-				NotificationHandler->Notification = FSlateNotificationManager::Get().AddNotification(Info);
-				if (NotificationHandler->Notification.IsValid())
+				if(WeakPersonaToolkit.IsValid())	// SetPreviewMesh can invalidate the persona toolkit, so check it here before displaying toast
 				{
-					NotificationHandler->Notification->SetCompletionState(SNotificationItem::CS_Success);
-				}
+					TSharedPtr<FNotificationHandler> NotificationHandler = MakeShared<FNotificationHandler>();
+
+					FNotificationInfo Info(LOCTEXT("PreviewMeshSetTemporarily", "Preview mesh set temporarily"));
+					Info.ExpireDuration = 10.0f;
+					Info.bUseLargeFont = true;
+					Info.ButtonDetails.Add(
+						FNotificationButtonInfo(
+							LOCTEXT("ApplyToAsset", "Apply To Asset"), 
+							LOCTEXT("ApplyToAssetToolTip", "The preview mesh has changed, but it will not be able to be saved until it is applied to the asset. Click here to make the change to the preview mesh persistent."),
+							FSimpleDelegate::CreateStatic(&FNotificationHandler::HandleApplyPreviewMesh, NotificationHandler, WeakPersonaToolkit),
+							SNotificationItem::CS_Success));
+
+					NotificationHandler->Notification = FSlateNotificationManager::Get().AddNotification(Info);
+					if (NotificationHandler->Notification.IsValid())
+					{
+						NotificationHandler->Notification->SetCompletionState(SNotificationItem::CS_Success);
+					}
 
-				FSlateApplication::Get().DismissAllMenus();
+					FSlateApplication::Get().DismissAllMenus();
+				}
 			});
 			AssetPickerConfig.bAllowNullSelection = false;
 			AssetPickerConfig.InitialAssetViewType = EAssetViewType::List;
@@ -760,6 +778,11 @@ void FPersonaModule::AddCommonToolbarExtensions(FToolBarBuilder& InToolbarBuilde
 			{
 				if (WeakPersonaToolkit.IsValid())
 				{
+					if(WeakPersonaToolkit.Pin()->GetContext() == UPhysicsAsset::StaticClass()->GetFName())
+					{
+						return false;
+					}
+
 					FString TagValue;
 					if (AssetData.GetTagValue("Skeleton", TagValue))
 					{
@@ -800,16 +823,22 @@ void FPersonaModule::AddCommonToolbarExtensions(FToolBarBuilder& InToolbarBuilde
 
 void FPersonaModule::HandleNewAnimNotifyBlueprintCreated(UBlueprint* InBlueprint)
 {
-	UEdGraph* const NewGraph = FBlueprintEditorUtils::CreateNewGraph(InBlueprint, "Received_Notify", UEdGraph::StaticClass(), UEdGraphSchema_K2::StaticClass());
-	FBlueprintEditorUtils::AddFunctionGraph(InBlueprint, NewGraph, /*bIsUserCreated=*/ false, UAnimNotify::StaticClass());
-	InBlueprint->LastEditedDocuments.Add(NewGraph);
+	if (InBlueprint->BlueprintType == BPTYPE_Normal)
+	{
+		UEdGraph* const NewGraph = FBlueprintEditorUtils::CreateNewGraph(InBlueprint, "Received_Notify", UEdGraph::StaticClass(), UEdGraphSchema_K2::StaticClass());
+		FBlueprintEditorUtils::AddFunctionGraph(InBlueprint, NewGraph, /*bIsUserCreated=*/ false, UAnimNotify::StaticClass());
+		InBlueprint->LastEditedDocuments.Add(NewGraph);
+	}
 }
 
 void FPersonaModule::HandleNewAnimNotifyStateBlueprintCreated(UBlueprint* InBlueprint)
 {
-	UEdGraph* const NewGraph = FBlueprintEditorUtils::CreateNewGraph(InBlueprint, "Received_NotifyTick", UEdGraph::StaticClass(), UEdGraphSchema_K2::StaticClass());
-	FBlueprintEditorUtils::AddFunctionGraph(InBlueprint, NewGraph, /*bIsUserCreated=*/ false, UAnimNotifyState::StaticClass());
-	InBlueprint->LastEditedDocuments.Add(NewGraph);
+	if (InBlueprint->BlueprintType == BPTYPE_Normal)
+	{
+		UEdGraph* const NewGraph = FBlueprintEditorUtils::CreateNewGraph(InBlueprint, "Received_NotifyTick", UEdGraph::StaticClass(), UEdGraphSchema_K2::StaticClass());
+		FBlueprintEditorUtils::AddFunctionGraph(InBlueprint, NewGraph, /*bIsUserCreated=*/ false, UAnimNotifyState::StaticClass());
+		InBlueprint->LastEditedDocuments.Add(NewGraph);
+	}
 }
 
 #undef LOCTEXT_NAMESPACE
diff --git a/Engine/Source/Editor/Persona/Private/PersonaToolkit.cpp b/Engine/Source/Editor/Persona/Private/PersonaToolkit.cpp
--- a/Engine/Source/Editor/Persona/Private/PersonaToolkit.cpp
+++ b/Engine/Source/Editor/Persona/Private/PersonaToolkit.cpp
@@ -4,77 +4,79 @@
 #include "Modules/ModuleManager.h"
 #include "Engine/SkeletalMesh.h"
 #include "Animation/AnimationAsset.h"
+#include "PhysicsEngine/PhysicsAsset.h"
 #include "AnimationEditorPreviewScene.h"
 #include "ISkeletonEditorModule.h"
 #include "Animation/AnimBlueprint.h"
 #include "GameFramework/WorldSettings.h"
 #include "ScopedTransaction.h"
 #include "PersonaModule.h"
+#include "PersonaAssetFamily.h"
 
 FPersonaToolkit::FPersonaToolkit()
 	: Skeleton(nullptr)
 	, Mesh(nullptr)
 	, AnimBlueprint(nullptr)
 	, AnimationAsset(nullptr)
+	, PhysicsAsset(nullptr)
 {
 }
 
+static void FindCounterpartAssets(const UObject* InAsset, TWeakObjectPtr<USkeleton>& OutSkeleton, USkeletalMesh*& OutMesh)
+{
+	const USkeleton* CounterpartSkeleton = OutSkeleton.Get();
+	const USkeletalMesh* CounterpartMesh = OutMesh;
+	FPersonaAssetFamily::FindCounterpartAssets(InAsset, CounterpartSkeleton, CounterpartMesh);
+	OutSkeleton = CounterpartSkeleton;
+	OutMesh = const_cast<USkeletalMesh*>(CounterpartMesh);
+}
+
 void FPersonaToolkit::Initialize(USkeleton* InSkeleton)
 {
 	check(InSkeleton);
 	Skeleton = InSkeleton;
-	Mesh = InSkeleton->GetPreviewMesh();
-	if (Mesh == nullptr)
-	{
-		Mesh = Skeleton->FindCompatibleMesh();
-	}
 	InitialAssetClass = USkeleton::StaticClass();
+
+	FindCounterpartAssets(InSkeleton, Skeleton, Mesh);
 }
 
 void FPersonaToolkit::Initialize(UAnimationAsset* InAnimationAsset)
 {
 	check(InAnimationAsset);
 	AnimationAsset = InAnimationAsset;
-	Skeleton = InAnimationAsset->GetSkeleton();
-	Mesh = InAnimationAsset->GetPreviewMesh();
-	if (Mesh == nullptr)
-	{
-		Mesh = Skeleton->GetPreviewMesh();
-	}
-	if (Mesh == nullptr)
-	{
-		Mesh = Skeleton->FindCompatibleMesh();
-	}
 	InitialAssetClass = UAnimationAsset::StaticClass();
+
+	FindCounterpartAssets(InAnimationAsset, Skeleton, Mesh);
 }
 
 void FPersonaToolkit::Initialize(USkeletalMesh* InSkeletalMesh)
 {
 	check(InSkeletalMesh);
-	Skeleton = InSkeletalMesh->Skeleton;
 	Mesh = InSkeletalMesh;
 	InitialAssetClass = USkeletalMesh::StaticClass();
+
+	FindCounterpartAssets(InSkeletalMesh, Skeleton, Mesh);
 }
 
 void FPersonaToolkit::Initialize(UAnimBlueprint* InAnimBlueprint)
 {
 	check(InAnimBlueprint);
 	AnimBlueprint = InAnimBlueprint;
-	Skeleton = InAnimBlueprint->TargetSkeleton;
-	check(InAnimBlueprint->TargetSkeleton);
-	Mesh = InAnimBlueprint->GetPreviewMesh();
-	if (Mesh == nullptr)
-	{
-		Mesh = Skeleton->GetPreviewMesh();
-	}
-	if (Mesh == nullptr)
-	{
-		Mesh = Skeleton->FindCompatibleMesh();
-	}
 	InitialAssetClass = UAnimBlueprint::StaticClass();
+
+	FindCounterpartAssets(InAnimBlueprint, Skeleton, Mesh);
+}
+
+void FPersonaToolkit::Initialize(UPhysicsAsset* InPhysicsAsset)
+{
+	check(InPhysicsAsset);
+	PhysicsAsset = InPhysicsAsset;
+	InitialAssetClass = UPhysicsAsset::StaticClass();
+
+	FindCounterpartAssets(InPhysicsAsset, Skeleton, Mesh);
 }
 
-void FPersonaToolkit::CreatePreviewScene()
+void FPersonaToolkit::CreatePreviewScene(const FPersonaToolkitArgs& PersonaToolkitArgs)
 {
 	if (!PreviewScene.IsValid())
 	{
@@ -89,10 +91,34 @@ void FPersonaToolkit::CreatePreviewScene()
 		//Temporary fix for missing attached assets - MDW
 		PreviewScene->GetWorld()->GetWorldSettings()->SetIsTemporarilyHiddenInEditor(false);
 
+		if (PersonaToolkitArgs.OnPreviewSceneCreated.IsBound())
+		{
+			// Custom per-instance scene setup
+			PersonaToolkitArgs.OnPreviewSceneCreated.Execute(PreviewScene.ToSharedRef());
+		}
+		else
+		{
+			// setup default scene
+			AActor* Actor = PreviewScene->GetWorld()->SpawnActor<AActor>(AActor::StaticClass(), FTransform::Identity);
+			PreviewScene->SetActor(Actor);
+
+			// Create the preview component
+			UDebugSkelMeshComponent* SkeletalMeshComponent = NewObject<UDebugSkelMeshComponent>(Actor);
+			PreviewScene->AddComponent(SkeletalMeshComponent, FTransform::Identity);
+			PreviewScene->SetPreviewMeshComponent(SkeletalMeshComponent);
+
+			// set root component, so we can attach to it. 
+			Actor->SetRootComponent(SkeletalMeshComponent);
+		}
+
 		// allow external systems to add components or otherwise manipulate the scene
 		FPersonaModule& PersonaModule = FModuleManager::GetModuleChecked<FPersonaModule>(TEXT("Persona"));
 		PersonaModule.OnPreviewSceneCreated().Broadcast(PreviewScene.ToSharedRef());
 
+		// Force validation of preview attached assets (catch case of never doing it if we dont have a valid preview mesh)
+		PreviewScene->ValidatePreviewAttachedAssets(nullptr);
+		PreviewScene->RefreshAdditionalMeshes();
+
 		bool bSetMesh = false;
 
 		// Set the mesh
@@ -200,6 +226,11 @@ USkeletalMesh* FPersonaToolkit::GetPreviewMesh() const
 		check(AnimBlueprint);
 		return AnimBlueprint->GetPreviewMesh();
 	}
+	else if (InitialAssetClass == UPhysicsAsset::StaticClass())
+	{
+		check(PhysicsAsset);
+		return PhysicsAsset->GetPreviewMesh();
+	}
 	else if(InitialAssetClass == USkeletalMesh::StaticClass())
 	{
 		check(Mesh);
@@ -217,6 +248,14 @@ void FPersonaToolkit::SetPreviewMesh(class USkeletalMesh* InSkeletalMesh, bool b
 	// Cant set preview mesh on a skeletal mesh (makes for a confusing experience!)
 	if (InitialAssetClass != USkeletalMesh::StaticClass())
 	{
+		// If the skeleton itself is changing, then we need to re-open the asset editor
+		bool bReOpenEditor = false;
+		if(InSkeletalMesh != nullptr && InSkeletalMesh->Skeleton != &EditableSkeleton->GetSkeleton())
+		{
+			bReOpenEditor = true;
+			bSetPreviewMeshInAsset = true;
+		}
+
 		if(bSetPreviewMeshInAsset)
 		{
 			if (InitialAssetClass == UAnimationAsset::StaticClass())
@@ -233,6 +272,13 @@ void FPersonaToolkit::SetPreviewMesh(class USkeletalMesh* InSkeletalMesh, bool b
 				check(AnimBlueprint);
 				AnimBlueprint->SetPreviewMesh(InSkeletalMesh);
 			}
+			else if (InitialAssetClass == UPhysicsAsset::StaticClass())
+			{
+				FScopedTransaction Transaction(NSLOCTEXT("PersonaToolkit", "SetPhysicsAssetPreviewMesh", "Set Physics Asset Preview Mesh"));
+
+				check(PhysicsAsset);
+				PhysicsAsset->SetPreviewMesh(InSkeletalMesh);
+			}
 			else
 			{
 				check(EditableSkeleton.IsValid());
@@ -240,6 +286,28 @@ void FPersonaToolkit::SetPreviewMesh(class USkeletalMesh* InSkeletalMesh, bool b
 			}
 		}
 
+		if(bReOpenEditor)
+		{
+			UObject* Asset = nullptr;
+			if (InitialAssetClass == UAnimationAsset::StaticClass())
+			{
+				Asset = AnimationAsset;
+			}
+			else if (InitialAssetClass == UAnimBlueprint::StaticClass())
+			{
+				Asset = AnimBlueprint;
+			}
+			else if (InitialAssetClass == UPhysicsAsset::StaticClass())
+			{
+				Asset = PhysicsAsset;
+			}
+			check(Asset);
+
+			FAssetEditorManager::Get().CloseAllEditorsForAsset(Asset);
+			FAssetEditorManager::Get().OpenEditorForAsset(Asset);
+			return;
+		}
+		
 		GetPreviewScene()->SetPreviewMesh(InSkeletalMesh);
 	}
 }
diff --git a/Engine/Source/Editor/Persona/Private/PersonaToolkit.h b/Engine/Source/Editor/Persona/Private/PersonaToolkit.h
--- a/Engine/Source/Editor/Persona/Private/PersonaToolkit.h
+++ b/Engine/Source/Editor/Persona/Private/PersonaToolkit.h
@@ -11,6 +11,8 @@ class FAnimationEditorPreviewScene;
 class UAnimationAsset;
 class UAnimBlueprint;
 class USkeletalMesh;
+class UPhysicsAsset;
+struct FPersonaToolkitArgs;
 
 class FPersonaToolkit : public IPersonaToolkit, public TSharedFromThis<FPersonaToolkit>
 {
@@ -23,9 +25,10 @@ public:
 	void Initialize(UAnimationAsset* InAnimationAsset);
 	void Initialize(USkeletalMesh* InSkeletalMesh);
 	void Initialize(UAnimBlueprint* InAnimBlueprint);
+	void Initialize(UPhysicsAsset* InPhysicsAsset);
 
 	/** Optionally create a preview scene - note: creates an editable skeleton */
-	void CreatePreviewScene();
+	void CreatePreviewScene(const FPersonaToolkitArgs& PersonaToolkitArgs);
 
 	/** IPersonaToolkit interface */
 	virtual class USkeleton* GetSkeleton() const override;
@@ -57,6 +60,9 @@ private:
 	/** the animation asset we are editing */
 	UAnimationAsset* AnimationAsset;
 
+	/** the physics asset we are editing */
+	UPhysicsAsset* PhysicsAsset;
+
 	/** Preview scene for the editor */
 	TSharedPtr<FAnimationEditorPreviewScene> PreviewScene;
 
diff --git a/Engine/Source/Editor/Persona/Private/PreviewSceneCustomizations.cpp b/Engine/Source/Editor/Persona/Private/PreviewSceneCustomizations.cpp
--- a/Engine/Source/Editor/Persona/Private/PreviewSceneCustomizations.cpp
+++ b/Engine/Source/Editor/Persona/Private/PreviewSceneCustomizations.cpp
@@ -22,6 +22,7 @@
 #include "UObject/UObjectIterator.h"
 #include "STextBlock.h"
 #include "SCheckBox.h"
+#include "PhysicsEngine/PhysicsAsset.h"
 
 #define LOCTEXT_NAMESPACE "PreviewSceneCustomizations"
 
@@ -67,6 +68,7 @@ void FPreviewSceneDescriptionCustomization::CustomizeDetails(IDetailLayoutBuilde
 	TSharedRef<IPropertyHandle> AnimationModeProperty = DetailBuilder.GetProperty(GET_MEMBER_NAME_CHECKED(UPersonaPreviewSceneDescription, AnimationMode));
 	AnimationModeProperty->SetOnPropertyValueChanged(FSimpleDelegate::CreateSP(this, &FPreviewSceneDescriptionCustomization::HandleAnimationModeChanged));
 	TSharedRef<IPropertyHandle> AnimationProperty = DetailBuilder.GetProperty(GET_MEMBER_NAME_CHECKED(UPersonaPreviewSceneDescription, Animation));
+	AnimationProperty->SetOnPropertyValueChanged(FSimpleDelegate::CreateSP(this, &FPreviewSceneDescriptionCustomization::HandleAnimationChanged));
 	TSharedRef<IPropertyHandle> SkeletalMeshProperty = DetailBuilder.GetProperty(GET_MEMBER_NAME_CHECKED(UPersonaPreviewSceneDescription, PreviewMesh));
 	TSharedRef<IPropertyHandle> AdditionalMeshesProperty = DetailBuilder.GetProperty(GET_MEMBER_NAME_CHECKED(UPersonaPreviewSceneDescription, AdditionalMeshes));
 
@@ -88,7 +90,6 @@ void FPreviewSceneDescriptionCustomization::CustomizeDetails(IDetailLayoutBuilde
 		.AllowedClass(UAnimationAsset::StaticClass())
 		.PropertyHandle(AnimationProperty)
 		.OnShouldFilterAsset(this, &FPreviewSceneDescriptionCustomization::HandleShouldFilterAsset, false)
-		.OnObjectChanged(this, &FPreviewSceneDescriptionCustomization::HandleAnimationChanged)
 		.ThumbnailPool(DetailBuilder.GetThumbnailPool())
 	];
 
@@ -103,6 +104,10 @@ void FPreviewSceneDescriptionCustomization::CustomizeDetails(IDetailLayoutBuilde
 		{
 			PreviewMeshName = FText::Format(LOCTEXT("PreviewMeshAnimBlueprint", "{0}\n(Animation Blueprint)"), SkeletalMeshProperty->GetPropertyDisplayName());
 		}
+		else if(PersonaToolkit.Pin()->GetContext() == UPhysicsAsset::StaticClass()->GetFName())
+		{
+			PreviewMeshName = FText::Format(LOCTEXT("PreviewMeshPhysicsAsset", "{0}\n(Physics Asset)"), SkeletalMeshProperty->GetPropertyDisplayName());
+		}
 		else
 		{
 			PreviewMeshName = FText::Format(LOCTEXT("PreviewMeshSkeleton", "{0}\n(Skeleton)"), SkeletalMeshProperty->GetPropertyDisplayName());
@@ -151,7 +156,7 @@ void FPreviewSceneDescriptionCustomization::CustomizeDetails(IDetailLayoutBuilde
 			SNew(SObjectPropertyEntryBox)
 			.AllowedClass(USkeletalMesh::StaticClass())
 			.PropertyHandle(SkeletalMeshProperty)
-			.OnShouldFilterAsset(this, &FPreviewSceneDescriptionCustomization::HandleShouldFilterAsset, false)
+			.OnShouldFilterAsset(this, &FPreviewSceneDescriptionCustomization::HandleShouldFilterAsset, PersonaToolkit.Pin()->GetContext() == UPhysicsAsset::StaticClass()->GetFName())
 			.OnObjectChanged(this, &FPreviewSceneDescriptionCustomization::HandleMeshChanged)
 			.ThumbnailPool(DetailBuilder.GetThumbnailPool())
 		];
@@ -333,9 +338,10 @@ void FPreviewSceneDescriptionCustomization::HandleAnimationModeChanged()
 	}
 }
 
-void FPreviewSceneDescriptionCustomization::HandleAnimationChanged(const FAssetData& InAssetData)
+void FPreviewSceneDescriptionCustomization::HandleAnimationChanged()
 {
-	UAnimationAsset* AnimationAsset = Cast<UAnimationAsset>(InAssetData.GetAsset());
+	UPersonaPreviewSceneDescription* PersonaPreviewSceneDescription = PreviewScene.Pin()->GetPreviewSceneDescription();
+	UAnimationAsset* AnimationAsset = PersonaPreviewSceneDescription->Animation.LoadSynchronous();
 	PreviewScene.Pin()->SetPreviewAnimationAsset(AnimationAsset);
 
 	if (AnimationAsset != nullptr)
diff --git a/Engine/Source/Editor/Persona/Private/PreviewSceneCustomizations.h b/Engine/Source/Editor/Persona/Private/PreviewSceneCustomizations.h
--- a/Engine/Source/Editor/Persona/Private/PreviewSceneCustomizations.h
+++ b/Engine/Source/Editor/Persona/Private/PreviewSceneCustomizations.h
@@ -38,7 +38,7 @@ private:
 
 	void HandleAnimationModeChanged();
 
-	void HandleAnimationChanged(const FAssetData& InAssetData);
+	void HandleAnimationChanged();
 
 	void HandleMeshChanged(const FAssetData& InAssetData);
 
