Fix gradient editor missing textures on stop handles
Added hover for stop areas.

diff --git a/Engine/Source/Editor/EditorStyle/Private/SlateEditorStyle.cpp b/Engine/Source/Editor/EditorStyle/Private/SlateEditorStyle.cpp
--- a/Engine/Source/Editor/EditorStyle/Private/SlateEditorStyle.cpp
+++ b/Engine/Source/Editor/EditorStyle/Private/SlateEditorStyle.cpp
@@ -6928,47 +6928,50 @@ void FSlateEditorStyle::FStyle::SetupToolkitStyles()
 
 	// Curve Editor
 	{
-		Set( "CurveAssetEditor.Tabs.Properties", new IMAGE_BRUSH( "Icons/AssetIcons/CurveBase_16x", Icon16x16 ) );
-
-		Set( "CurveEditor.FitHorizontally", new IMAGE_BRUSH( "Icons/icon_CurveEditor_Horizontal_40x", Icon40x40 ) );
-		Set( "CurveEditor.FitVertically", new IMAGE_BRUSH( "Icons/icon_CurveEditor_Vertical_40x", Icon40x40 ) );
-		Set( "CurveEditor.Fit", new IMAGE_BRUSH( "Icons/icon_CurveEditor_ZoomToFit_40x", Icon40x40 ) );
-		Set( "CurveEditor.PanMode", new IMAGE_BRUSH( "Icons/icon_CurveEditor_Pan_40x", Icon40x40 ) );
-		Set( "CurveEditor.ZoomMode", new IMAGE_BRUSH( "Icons/icon_CurveEditor_Zoom_40x", Icon40x40 ) );
-		Set( "CurveEditor.CurveAuto", new IMAGE_BRUSH( "Icons/icon_CurveEditor_Auto_40x", Icon40x40 ) );
-		Set( "CurveEditor.CurveAutoClamped", new IMAGE_BRUSH( "Icons/icon_CurveEditor_AutoClamped_40x", Icon40x40 ) );
-		Set( "CurveEditor.CurveUser", new IMAGE_BRUSH( "Icons/icon_CurveEditor_User_40x", Icon40x40 ) );
-		Set( "CurveEditor.CurveBreak", new IMAGE_BRUSH( "Icons/icon_CurveEditor_Break_40x", Icon40x40 ) );
-		Set( "CurveEditor.CurveWeight", new IMAGE_BRUSH("Icons/icon_CurveEditor_Break_40x", Icon40x40));
-		Set( "CurveEditor.CurveNonWeight", new IMAGE_BRUSH("Icons/icon_CurveEditor_Break_40x", Icon40x40));
-
-		Set( "CurveEditor.Linear", new IMAGE_BRUSH( "Icons/icon_CurveEditor_Linear_40x", Icon40x40 ) );
-		Set( "CurveEditor.Constant", new IMAGE_BRUSH( "Icons/icon_CurveEditor_Constant_40x", Icon40x40 ) );
-		Set( "CurveEditor.FlattenTangents", new IMAGE_BRUSH( "Icons/icon_CurveEditor_Flatten_40x", Icon40x40 ) );
-		Set( "CurveEditor.StraightenTangents", new IMAGE_BRUSH( "Icons/icon_CurveEditor_Straighten_40x", Icon40x40 ) );
-		Set( "CurveEditor.ShowAllTangents", new IMAGE_BRUSH( "Icons/icon_CurveEditor_ShowAll_40x", Icon40x40 ) );
-		Set( "CurveEditor.CreateTab", new IMAGE_BRUSH( "Icons/icon_CurveEditor_Create_40x", Icon40x40 ) );
-		Set( "CurveEditor.DeleteTab", new IMAGE_BRUSH( "Icons/icon_CurveEditor_DeleteTab_40x", Icon40x40 ) );
-		
-		Set( "CurveEditor.FitHorizontally.Small", new IMAGE_BRUSH( "Icons/icon_CurveEditor_Horizontal_40x", Icon20x20 ) );
-		Set( "CurveEditor.FitVertically.Small", new IMAGE_BRUSH( "Icons/icon_CurveEditor_Vertical_40x", Icon20x20 ) );
-		Set( "CurveEditor.Fit.Small", new IMAGE_BRUSH( "Icons/icon_CurveEditor_ZoomToFit_40x", Icon20x20 ) );
-		Set( "CurveEditor.PanMode.Small", new IMAGE_BRUSH( "Icons/icon_CurveEditor_Pan_40x", Icon20x20 ) );
-		Set( "CurveEditor.ZoomMode.Small", new IMAGE_BRUSH( "Icons/icon_CurveEditor_Zoom_40x", Icon20x20 ) );
-		Set( "CurveEditor.CurveAuto.Small", new IMAGE_BRUSH( "Icons/icon_CurveEditor_Auto_40x", Icon20x20 ) );
-		Set( "CurveEditor.CurveAutoClamped.Small", new IMAGE_BRUSH( "Icons/icon_CurveEditor_AutoClamped_40x", Icon20x20 ) );
-		Set( "CurveEditor.CurveUser.Small", new IMAGE_BRUSH( "Icons/icon_CurveEditor_User_40x", Icon20x20 ) );
-		Set( "CurveEditor.CurveBreak.Small", new IMAGE_BRUSH( "Icons/icon_CurveEditor_Break_40x", Icon20x20 ) );
+		Set("CurveAssetEditor.Tabs.Properties", new IMAGE_BRUSH("Icons/AssetIcons/CurveBase_16x", Icon16x16));
+
+		Set("CurveEditor.FitHorizontally", new IMAGE_BRUSH("Icons/icon_CurveEditor_Horizontal_40x", Icon40x40));
+		Set("CurveEditor.FitVertically", new IMAGE_BRUSH("Icons/icon_CurveEditor_Vertical_40x", Icon40x40));
+		Set("CurveEditor.Fit", new IMAGE_BRUSH("Icons/icon_CurveEditor_ZoomToFit_40x", Icon40x40));
+		Set("CurveEditor.PanMode", new IMAGE_BRUSH("Icons/icon_CurveEditor_Pan_40x", Icon40x40));
+		Set("CurveEditor.ZoomMode", new IMAGE_BRUSH("Icons/icon_CurveEditor_Zoom_40x", Icon40x40));
+		Set("CurveEditor.CurveAuto", new IMAGE_BRUSH("Icons/icon_CurveEditor_Auto_40x", Icon40x40));
+		Set("CurveEditor.CurveAutoClamped", new IMAGE_BRUSH("Icons/icon_CurveEditor_AutoClamped_40x", Icon40x40));
+		Set("CurveEditor.CurveUser", new IMAGE_BRUSH("Icons/icon_CurveEditor_User_40x", Icon40x40));
+		Set("CurveEditor.CurveBreak", new IMAGE_BRUSH("Icons/icon_CurveEditor_Break_40x", Icon40x40));
+		Set("CurveEditor.CurveWeight", new IMAGE_BRUSH("Icons/icon_CurveEditor_Break_40x", Icon40x40));
+		Set("CurveEditor.CurveNonWeight", new IMAGE_BRUSH("Icons/icon_CurveEditor_Break_40x", Icon40x40));
+
+		Set("CurveEditor.Linear", new IMAGE_BRUSH("Icons/icon_CurveEditor_Linear_40x", Icon40x40));
+		Set("CurveEditor.Constant", new IMAGE_BRUSH("Icons/icon_CurveEditor_Constant_40x", Icon40x40));
+		Set("CurveEditor.FlattenTangents", new IMAGE_BRUSH("Icons/icon_CurveEditor_Flatten_40x", Icon40x40));
+		Set("CurveEditor.StraightenTangents", new IMAGE_BRUSH("Icons/icon_CurveEditor_Straighten_40x", Icon40x40));
+		Set("CurveEditor.ShowAllTangents", new IMAGE_BRUSH("Icons/icon_CurveEditor_ShowAll_40x", Icon40x40));
+		Set("CurveEditor.CreateTab", new IMAGE_BRUSH("Icons/icon_CurveEditor_Create_40x", Icon40x40));
+		Set("CurveEditor.DeleteTab", new IMAGE_BRUSH("Icons/icon_CurveEditor_DeleteTab_40x", Icon40x40));
+
+		Set("CurveEditor.FitHorizontally.Small", new IMAGE_BRUSH("Icons/icon_CurveEditor_Horizontal_40x", Icon20x20));
+		Set("CurveEditor.FitVertically.Small", new IMAGE_BRUSH("Icons/icon_CurveEditor_Vertical_40x", Icon20x20));
+		Set("CurveEditor.Fit.Small", new IMAGE_BRUSH("Icons/icon_CurveEditor_ZoomToFit_40x", Icon20x20));
+		Set("CurveEditor.PanMode.Small", new IMAGE_BRUSH("Icons/icon_CurveEditor_Pan_40x", Icon20x20));
+		Set("CurveEditor.ZoomMode.Small", new IMAGE_BRUSH("Icons/icon_CurveEditor_Zoom_40x", Icon20x20));
+		Set("CurveEditor.CurveAuto.Small", new IMAGE_BRUSH("Icons/icon_CurveEditor_Auto_40x", Icon20x20));
+		Set("CurveEditor.CurveAutoClamped.Small", new IMAGE_BRUSH("Icons/icon_CurveEditor_AutoClamped_40x", Icon20x20));
+		Set("CurveEditor.CurveUser.Small", new IMAGE_BRUSH("Icons/icon_CurveEditor_User_40x", Icon20x20));
+		Set("CurveEditor.CurveBreak.Small", new IMAGE_BRUSH("Icons/icon_CurveEditor_Break_40x", Icon20x20));
 		Set("CurveEditor.CurveWeight.Small", new IMAGE_BRUSH("Icons/icon_CurveEditor_Break_40x", Icon20x20));
 		Set("CurveEditor.CurveNonWeight.Small", new IMAGE_BRUSH("Icons/icon_CurveEditor_Break_40x", Icon20x20));
 
-		Set( "CurveEditor.Linear.Small", new IMAGE_BRUSH( "Icons/icon_CurveEditor_Linear_40x", Icon20x20 ) );
-		Set( "CurveEditor.Constant.Small", new IMAGE_BRUSH( "Icons/icon_CurveEditor_Constant_40x", Icon20x20 ) );
-		Set( "CurveEditor.FlattenTangents.Small", new IMAGE_BRUSH( "Icons/icon_CurveEditor_Flatten_40x", Icon20x20 ) );
-		Set( "CurveEditor.StraightenTangents.Small", new IMAGE_BRUSH( "Icons/icon_CurveEditor_Straighten_40x", Icon20x20 ) );
-		Set( "CurveEditor.ShowAllTangents.Small", new IMAGE_BRUSH( "Icons/icon_CurveEditor_ShowAll_40x", Icon20x20 ) );
-		Set( "CurveEditor.CreateTab.Small", new IMAGE_BRUSH( "Icons/icon_CurveEditor_Create_40x", Icon20x20 ) );
-		Set( "CurveEditor.DeleteTab.Small", new IMAGE_BRUSH( "Icons/icon_CurveEditor_DeleteTab_40x", Icon20x20 ) );
+		Set("CurveEditor.Linear.Small", new IMAGE_BRUSH("Icons/icon_CurveEditor_Linear_40x", Icon20x20));
+		Set("CurveEditor.Constant.Small", new IMAGE_BRUSH("Icons/icon_CurveEditor_Constant_40x", Icon20x20));
+		Set("CurveEditor.FlattenTangents.Small", new IMAGE_BRUSH("Icons/icon_CurveEditor_Flatten_40x", Icon20x20));
+		Set("CurveEditor.StraightenTangents.Small", new IMAGE_BRUSH("Icons/icon_CurveEditor_Straighten_40x", Icon20x20));
+		Set("CurveEditor.ShowAllTangents.Small", new IMAGE_BRUSH("Icons/icon_CurveEditor_ShowAll_40x", Icon20x20));
+		Set("CurveEditor.CreateTab.Small", new IMAGE_BRUSH("Icons/icon_CurveEditor_Create_40x", Icon20x20));
+		Set("CurveEditor.DeleteTab.Small", new IMAGE_BRUSH("Icons/icon_CurveEditor_DeleteTab_40x", Icon20x20));
+
+		Set("CurveEditor.Gradient.HandleDown", new BOX_BRUSH("Sequencer/ScrubHandleDown", FMargin(6.f / 13.f, 5 / 12.f, 6 / 13.f, 8 / 12.f)));
+		Set("CurveEditor.Gradient.HandleUp", new BOX_BRUSH("Sequencer/ScrubHandleUp", FMargin(6.f / 13.f, 8 / 12.f, 6 / 13.f, 5 / 12.f)));
 	}
 
 	// Curve Editor
diff --git a/Engine/Source/Editor/UnrealEd/Private/SColorGradientEditor.cpp b/Engine/Source/Editor/UnrealEd/Private/SColorGradientEditor.cpp
--- a/Engine/Source/Editor/UnrealEd/Private/SColorGradientEditor.cpp
+++ b/Engine/Source/Editor/UnrealEd/Private/SColorGradientEditor.cpp
@@ -167,9 +167,37 @@ int32 SColorGradientEditor::OnPaint( const FPaintArgs& Args, const FGeometry& Al
 		// If no alpha keys are available, treat the curve as being completely opaque for drawing purposes
 		bool bHasAnyAlphaKeys = CurveOwner->HasAnyAlphaKeys(); 
 
-		// If any transpareny (A < 1) is found, we'll draw a checkerboard to visualize the color with alpha
+		// If any transparency (A < 1) is found, we'll draw a checkerboard to visualize the color with alpha
 		bool bHasTransparency = false;
+		static const FSlateBrush* WhiteBrush = FEditorStyle::GetBrush("WhiteBrush");
 
+		if (bColorAreaHovered)
+		{
+			// Draw a checkerboard behind there is any transparency visible
+			FSlateDrawElement::MakeBox
+			(
+				OutDrawElements,
+				LayerId,
+				ColorMarkAreaGeometry.ToPaintGeometry(),
+				WhiteBrush,
+				DrawEffects,
+				FLinearColor(.5f, .5f, .5f, .15f)
+				);
+		}
+		else if (bAlphaAreaHovered)
+		{
+			// Draw a checkerboard behind there is any transparency visible
+			FSlateDrawElement::MakeBox
+			(
+				OutDrawElements,
+				LayerId,
+				AlphaMarkAreaGeometry.ToPaintGeometry(),
+				WhiteBrush,
+				DrawEffects,
+				FLinearColor(.5f, .5f, .5f, .15f)
+			);
+		}
+			
 		// Sample the curve every 2 units.  THe curve could be non-linear so sampling at each stop would display an incorrect gradient
 		for( int32 CurrentStep = Start; CurrentStep < Finish; CurrentStep+=2 )
 		{
@@ -352,6 +380,25 @@ FReply SColorGradientEditor::OnMouseButtonDoubleClick( const FGeometry& InMyGeom
 
 FReply SColorGradientEditor::OnMouseMove( const FGeometry& MyGeometry, const FPointerEvent& MouseEvent )
 {
+	FGeometry ColorMarkAreaGeometry = GetColorMarkAreaGeometry(MyGeometry);
+	FGeometry AlphaMarkAreaGeometry = GetAlphaMarkAreaGeometry(MyGeometry);
+
+	if (ColorMarkAreaGeometry.IsUnderLocation(MouseEvent.GetScreenSpacePosition()))
+	{
+		bColorAreaHovered = true;
+		bAlphaAreaHovered = false;
+	}
+	else if (AlphaMarkAreaGeometry.IsUnderLocation(MouseEvent.GetScreenSpacePosition()))
+	{
+		bColorAreaHovered = false;
+		bAlphaAreaHovered = true;
+	}
+	else
+	{
+		bColorAreaHovered = false;
+		bAlphaAreaHovered = false;
+	}
+
 	if( HasMouseCapture() && IsEditingEnabled.Get() == true )
 	{
 		DistanceDragged += FMath::Abs( MouseEvent.GetCursorDelta().X );
@@ -453,6 +500,12 @@ FReply SColorGradientEditor::OnKeyDown( const FGeometry& MyGeometry, const FKeyE
 	return FReply::Unhandled();
 }
 
+void SColorGradientEditor::OnMouseLeave(const FPointerEvent& MouseEvent)
+{
+	bColorAreaHovered = false;
+	bAlphaAreaHovered = false;
+}
+
 FVector2D SColorGradientEditor::ComputeDesiredSize( float ) const
 {
 	return FVector2D( 1000, 75 );
@@ -701,8 +754,8 @@ void SColorGradientEditor::OnSetGradientStopTimeFromPopup( const FText& NewText,
 void SColorGradientEditor::DrawGradientStopMark( const FGradientStopMark& Mark, const FGeometry& Geometry, float XPos, const FLinearColor& Color, FSlateWindowElementList& OutDrawElements, int32 LayerId, const FSlateRect& InClippingRect, ESlateDrawEffect DrawEffects, bool bColor, const FWidgetStyle& InWidgetStyle ) const
 {
 	static const FSlateBrush* WhiteBrush = FEditorStyle::GetBrush("WhiteBrush");
-	static const FSlateBrush* ColorStopBrush = FEditorStyle::GetBrush("Sequencer.Timeline.ScrubHandleDown");
-	static const FSlateBrush* AlphaStopBrush = FEditorStyle::GetBrush("Sequencer.Timeline.ScrubHandleUp");
+	static const FSlateBrush* ColorStopBrush = FEditorStyle::GetBrush("CurveEditor.Gradient.HandleDown");
+	static const FSlateBrush* AlphaStopBrush = FEditorStyle::GetBrush("CurveEditor.Gradient.HandleUp");
 	static const FLinearColor SelectionColor = FEditorStyle::GetSlateColor("SelectionColor").GetColor( InWidgetStyle );
 
 	const float HandleSize = 13.0f;
diff --git a/Engine/Source/Editor/UnrealEd/Public/SColorGradientEditor.h b/Engine/Source/Editor/UnrealEd/Public/SColorGradientEditor.h
--- a/Engine/Source/Editor/UnrealEd/Public/SColorGradientEditor.h
+++ b/Engine/Source/Editor/UnrealEd/Public/SColorGradientEditor.h
@@ -66,6 +66,7 @@ public:
 	virtual FReply OnMouseMove( const FGeometry& MyGeometry, const FPointerEvent& MouseEvent ) override;
 	virtual FReply OnMouseButtonUp( const FGeometry& MyGeometry, const FPointerEvent& MouseEvent ) override;
 	virtual FReply OnKeyDown( const FGeometry& MyGeometry, const FKeyEvent& InKeyEvent ) override;
+	virtual void OnMouseLeave(const FPointerEvent& MouseEvent) override;
 	virtual FVector2D ComputeDesiredSize(float) const override;
 
 	/**
@@ -220,6 +221,10 @@ private:
 	TAttribute<bool> IsEditingEnabled;
 	/** Cached position where context menus should appear */
 	FVector2D ContextMenuPosition;
+	/** Whether or not the color gradient stop area is hovered */
+	bool bColorAreaHovered;
+	/** Whether or not the alpha gradient stop area is hovered */
+	bool bAlphaAreaHovered;
 	/** Current distance dragged since we captured the mouse */
 	float DistanceDragged;
 	/** True if an alpha value is being dragged */
