Fix for anim slot window of small icons

diff --git a/Engine/Source/Editor/EditorStyle/Private/SlateEditorStyle.cpp b/Engine/Source/Editor/EditorStyle/Private/SlateEditorStyle.cpp
--- a/Engine/Source/Editor/EditorStyle/Private/SlateEditorStyle.cpp
+++ b/Engine/Source/Editor/EditorStyle/Private/SlateEditorStyle.cpp
@@ -5318,8 +5318,11 @@ void FSlateEditorStyle::FStyle::SetupPersonaStyle()
 
 		// Anim Slot Manager
 		Set("AnimSlotManager.SaveSkeleton", new IMAGE_BRUSH("Persona/AnimSlotManager/icon_SaveSkeleton_40x", Icon40x40));
+		Set("AnimSlotManager.SaveSkeleton.Small", new IMAGE_BRUSH("Persona/AnimSlotManager/icon_SaveSkeleton_40x", Icon20x20));
 		Set("AnimSlotManager.AddGroup", new IMAGE_BRUSH("Persona/AnimSlotManager/icon_AddGroup_40x", Icon40x40));
+		Set("AnimSlotManager.AddGroup.Small", new IMAGE_BRUSH("Persona/AnimSlotManager/icon_AddGroup_40x", Icon20x20));
 		Set("AnimSlotManager.AddSlot", new IMAGE_BRUSH("Persona/AnimSlotManager/icon_AddSlot_40x", Icon40x40));
+		Set("AnimSlotManager.AddSlot.Small", new IMAGE_BRUSH("Persona/AnimSlotManager/icon_AddSlot_40x", Icon20x20));
 		Set("AnimSlotManager.Warning", new IMAGE_BRUSH("Persona/AnimSlotManager/icon_Warning_14x", Icon16x16));
 
 		// Anim Notify Editor
