Fix UE4 icon to be the correct one

diff --git a/Engine/Source/Editor/EditorStyle/Private/SlateEditorStyle.cpp b/Engine/Source/Editor/EditorStyle/Private/SlateEditorStyle.cpp
--- a/Engine/Source/Editor/EditorStyle/Private/SlateEditorStyle.cpp
+++ b/Engine/Source/Editor/EditorStyle/Private/SlateEditorStyle.cpp
@@ -1660,8 +1660,8 @@ void FSlateEditorStyle::FStyle::SetupGeneralStyles()
 		Set( "AboutScreen.Background", new IMAGE_BRUSH( "About/Background", FVector2D(600,332), FLinearColor::White, ESlateBrushTileType::Both) );
 		Set( "AboutScreen.Facebook", new IMAGE_BRUSH( "About/FacebookIcon", FVector2D(35,35) ) );
 		Set( "AboutScreen.FacebookHovered", new IMAGE_BRUSH( "About/FacebookIcon_Hovered", FVector2D(35,35) ) );
-		Set( "AboutScreen.UE4", new IMAGE_BRUSH( "About/UE4Icon", FVector2D(50,50) ) );
-		Set( "AboutScreen.UE4Hovered", new IMAGE_BRUSH( "About/UE4Icon_Hovered", FVector2D(50,50) ) );
+		Set( "AboutScreen.UE4", new IMAGE_BRUSH( "About/UE4Icon", FVector2D(50,50), FLinearColor::Gray) );
+		Set( "AboutScreen.UE4Hovered", new IMAGE_BRUSH( "About/UE4Icon", FVector2D(50,50), FLinearColor::White) );
 		Set( "AboutScreen.EpicGames", new IMAGE_BRUSH( "About/EpicGamesIcon", FVector2D(50,50) ) );
 		Set( "AboutScreen.EpicGamesHovered", new IMAGE_BRUSH( "About/EpicGamesIcon_Hovered", FVector2D(50,50) ) );
 	}
