Enable small icons for Preview modes (Contributed by [REDACTED_PERSON])
#5577

diff --git a/Engine/Source/Editor/EditorStyle/Private/SlateEditorStyle.cpp b/Engine/Source/Editor/EditorStyle/Private/SlateEditorStyle.cpp
--- a/Engine/Source/Editor/EditorStyle/Private/SlateEditorStyle.cpp
+++ b/Engine/Source/Editor/EditorStyle/Private/SlateEditorStyle.cpp
@@ -4765,21 +4765,37 @@ void FSlateEditorStyle::FStyle::SetupLevelEditorStyle()
 		Set("LevelEditor.SourceControl.Problem.Small", new IMAGE_BRUSH("Icons/icon_source_control_40x_problem", Icon20x20));
 
 		Set("LevelEditor.PreviewMode.Enabled", new IMAGE_BRUSH("Icons/icon_PreviewMode_SM5_Enabled_40x", Icon40x40));
+		Set("LevelEditor.PreviewMode.Enabled.Small", new IMAGE_BRUSH("Icons/icon_PreviewMode_SM5_Enabled_40x", Icon20x20));
 		Set("LevelEditor.PreviewMode.Disabled", new IMAGE_BRUSH("Icons/icon_PreviewMode_SM5_Disabled_40x", Icon40x40));
+		Set("LevelEditor.PreviewMode.Disabled.Small", new IMAGE_BRUSH("Icons/icon_PreviewMode_SM5_Disabled_40x", Icon20x20));
 		Set("LevelEditor.PreviewMode.SM5.Enabled", new IMAGE_BRUSH("Icons/icon_PreviewMode_SM5_Enabled_40x", Icon40x40));
+		Set("LevelEditor.PreviewMode.SM5.Enabled.Small", new IMAGE_BRUSH("Icons/icon_PreviewMode_SM5_Enabled_40x", Icon20x20));
 		Set("LevelEditor.PreviewMode.SM5.Disabled", new IMAGE_BRUSH("Icons/icon_PreviewMode_SM5_Enabled_40x", Icon40x40));
+		Set("LevelEditor.PreviewMode.SM5.Disabled.Small", new IMAGE_BRUSH("Icons/icon_PreviewMode_SM5_Enabled_40x", Icon20x20));
 		Set("LevelEditor.PreviewMode.SM4.Enabled", new IMAGE_BRUSH("Icons/icon_PreviewMode_SM4_Enabled_40x", Icon40x40));
+		Set("LevelEditor.PreviewMode.SM4.Enabled.Small", new IMAGE_BRUSH("Icons/icon_PreviewMode_SM4_Enabled_40x", Icon20x20));
 		Set("LevelEditor.PreviewMode.SM4.Disabled", new IMAGE_BRUSH("Icons/icon_PreviewMode_SM4_Enabled_40x", Icon40x40));
+		Set("LevelEditor.PreviewMode.SM4.Disabled.Small", new IMAGE_BRUSH("Icons/icon_PreviewMode_SM4_Enabled_40x", Icon20x20));
 		Set("LevelEditor.PreviewMode.AndroidES2.Enabled", new IMAGE_BRUSH("Icons/icon_PreviewMode_AndroidES2_Enabled_40x", Icon40x40));
+		Set("LevelEditor.PreviewMode.AndroidES2.Enabled.Small", new IMAGE_BRUSH("Icons/icon_PreviewMode_AndroidES2_Enabled_40x", Icon20x20));
 		Set("LevelEditor.PreviewMode.AndroidES2.Disabled", new IMAGE_BRUSH("Icons/icon_PreviewMode_AndroidES2_Disabled_40x", Icon40x40));
+		Set("LevelEditor.PreviewMode.AndroidES2.Disabled.Small", new IMAGE_BRUSH("Icons/icon_PreviewMode_AndroidES2_Disabled_40x", Icon20x20));
 		Set("LevelEditor.PreviewMode.AndroidES31.Enabled", new IMAGE_BRUSH("Icons/icon_PreviewMode_AndroidES31_Enabled_40x", Icon40x40));
+		Set("LevelEditor.PreviewMode.AndroidES31.Enabled.Small", new IMAGE_BRUSH("Icons/icon_PreviewMode_AndroidES31_Enabled_40x", Icon20x20));
 		Set("LevelEditor.PreviewMode.AndroidES31.Disabled", new IMAGE_BRUSH("Icons/icon_PreviewMode_AndroidES31_Disabled_40x", Icon40x40));
+		Set("LevelEditor.PreviewMode.AndroidES31.Disabled.Small", new IMAGE_BRUSH("Icons/icon_PreviewMode_AndroidES31_Disabled_40x", Icon20x20));
 		Set("LevelEditor.PreviewMode.AndroidVulkan.Enabled", new IMAGE_BRUSH("Icons/icon_PreviewMode_AndroidVulkan_Enabled_40x", Icon40x40));
+		Set("LevelEditor.PreviewMode.AndroidVulkan.Enabled.Small", new IMAGE_BRUSH("Icons/icon_PreviewMode_AndroidVulkan_Enabled_40x", Icon20x20));
 		Set("LevelEditor.PreviewMode.AndroidVulkan.Disabled", new IMAGE_BRUSH("Icons/icon_PreviewMode_AndroidVulkan_Disabled_40x", Icon40x40));
+		Set("LevelEditor.PreviewMode.AndroidVulkan.Disabled.Small", new IMAGE_BRUSH("Icons/icon_PreviewMode_AndroidVulkan_Disabled_40x", Icon20x20));
 		Set("LevelEditor.PreviewMode.iOS.Enabled", new IMAGE_BRUSH("Icons/icon_PreviewMode_iOS_Enabled_40x", Icon40x40));
+		Set("LevelEditor.PreviewMode.iOS.Enabled.Small", new IMAGE_BRUSH("Icons/icon_PreviewMode_iOS_Enabled_40x", Icon20x20));
 		Set("LevelEditor.PreviewMode.iOS.Disabled", new IMAGE_BRUSH("Icons/icon_PreviewMode_iOS_Disabled_40x", Icon40x40));
+		Set("LevelEditor.PreviewMode.iOS.Disabled.Small", new IMAGE_BRUSH("Icons/icon_PreviewMode_iOS_Disabled_40x", Icon20x20));
 		Set("LevelEditor.PreviewMode.HTML5.Enabled", new IMAGE_BRUSH("Icons/icon_PreviewMode_HTML5_Enabled_40x", Icon40x40));
+		Set("LevelEditor.PreviewMode.HTML5.Enabled.Small", new IMAGE_BRUSH("Icons/icon_PreviewMode_HTML5_Enabled_40x", Icon20x20));
 		Set("LevelEditor.PreviewMode.HTML5.Disabled", new IMAGE_BRUSH("Icons/icon_PreviewMode_HTML5_Disabled_40x", Icon40x40));
+		Set("LevelEditor.PreviewMode.HTML5.Disabled.Small", new IMAGE_BRUSH("Icons/icon_PreviewMode_HTML5_Disabled_40x", Icon20x20));
 
 		Set("LevelEditor.ViewOptions", new IMAGE_BRUSH("Icons/icon_view_40x", Icon40x40));
 		Set( "LevelEditor.ViewOptions.Small", new IMAGE_BRUSH( "Icons/icon_view_40x", Icon20x20 ) );
