- Added support for platform launcher icons to be in platform extensions directory

diff --git a/Engine/Source/Editor/EditorStyle/Private/SlateEditorStyle.cpp b/Engine/Source/Editor/EditorStyle/Private/SlateEditorStyle.cpp
--- a/Engine/Source/Editor/EditorStyle/Private/SlateEditorStyle.cpp
+++ b/Engine/Source/Editor/EditorStyle/Private/SlateEditorStyle.cpp
@@ -7541,9 +7541,22 @@ void FSlateEditorStyle::FStyle::SetupAutomationStyles()
 		for(const PlatformInfo::FPlatformInfo& PlatformInfo : PlatformInfo::GetPlatformInfoArray())
 		{
 			// @todo platplug: Add in .Small icons
-			Set( PlatformInfo.GetIconStyleName(PlatformInfo::EPlatformIconSize::Normal), new IMAGE_BRUSH( *PlatformInfo.GetIconPath(PlatformInfo::EPlatformIconSize::Normal), Icon24x24 ) );
-			Set( PlatformInfo.GetIconStyleName(PlatformInfo::EPlatformIconSize::Large),  new IMAGE_BRUSH( *PlatformInfo.GetIconPath(PlatformInfo::EPlatformIconSize::Large),  Icon64x64 ) );
-			Set( PlatformInfo.GetIconStyleName(PlatformInfo::EPlatformIconSize::XLarge), new IMAGE_BRUSH( *PlatformInfo.GetIconPath(PlatformInfo::EPlatformIconSize::XLarge), Icon128x128 ) );
+
+			// some platforms may specify a "rooted" path in the platform extensions directory, so look for that case here, and use a different path for the brush
+			FString NormalIconPath = PlatformInfo.GetIconPath(PlatformInfo::EPlatformIconSize::Normal);
+			if (NormalIconPath.StartsWith(TEXT("/Platforms/")))
+			{
+				#define PLATFORM_IMAGE_BRUSH( PlatformPath, ... ) FSlateImageBrush( PlatformPath.Replace(TEXT("/Platforms/"), *FPaths::PlatformExtensionsDir()) + TEXT(".png") , __VA_ARGS__ )
+				Set(PlatformInfo.GetIconStyleName(PlatformInfo::EPlatformIconSize::Normal), new PLATFORM_IMAGE_BRUSH(NormalIconPath, Icon24x24));
+				Set(PlatformInfo.GetIconStyleName(PlatformInfo::EPlatformIconSize::Large), new PLATFORM_IMAGE_BRUSH(PlatformInfo.GetIconPath(PlatformInfo::EPlatformIconSize::Large), Icon64x64));
+				Set(PlatformInfo.GetIconStyleName(PlatformInfo::EPlatformIconSize::XLarge), new PLATFORM_IMAGE_BRUSH(PlatformInfo.GetIconPath(PlatformInfo::EPlatformIconSize::XLarge), Icon128x128));
+			}
+			else
+			{
+				Set(PlatformInfo.GetIconStyleName(PlatformInfo::EPlatformIconSize::Normal), new IMAGE_BRUSH(*NormalIconPath, Icon24x24));
+				Set(PlatformInfo.GetIconStyleName(PlatformInfo::EPlatformIconSize::Large), new IMAGE_BRUSH(*PlatformInfo.GetIconPath(PlatformInfo::EPlatformIconSize::Large), Icon64x64));
+				Set(PlatformInfo.GetIconStyleName(PlatformInfo::EPlatformIconSize::XLarge), new IMAGE_BRUSH(*PlatformInfo.GetIconPath(PlatformInfo::EPlatformIconSize::XLarge), Icon128x128));
+			}
 		}
 #endif
 
