Sequencer: Added back in ScrubHandleUp/Down for the VisualLoggerTimeSliderController to fix missing brush warning

diff --git a/Engine/Source/Editor/EditorStyle/Private/SlateEditorStyle.cpp b/Engine/Source/Editor/EditorStyle/Private/SlateEditorStyle.cpp
--- a/Engine/Source/Editor/EditorStyle/Private/SlateEditorStyle.cpp
+++ b/Engine/Source/Editor/EditorStyle/Private/SlateEditorStyle.cpp
@@ -1792,6 +1792,8 @@ void FSlateEditorStyle::FStyle::SetupGeneralStyles()
 		Set( "Sequencer.TangentHandle", new IMAGE_BRUSH( "Sequencer/TangentHandle", FVector2D(7, 7) ) );
 		Set( "Sequencer.GenericDivider", new IMAGE_BRUSH( "Sequencer/GenericDivider", FVector2D(2.f, 2.f), FLinearColor::White, ESlateBrushTileType::Vertical ) );
 
+		Set("Sequencer.Timeline.ScrubHandleDown", new BOX_BRUSH("Sequencer/ScrubHandleDown", FMargin(6.f / 13.f, 5 / 12.f, 6 / 13.f, 8 / 12.f)));
+		Set("Sequencer.Timeline.ScrubHandleUp", new BOX_BRUSH("Sequencer/ScrubHandleUp", FMargin(6.f / 13.f, 8 / 12.f, 6 / 13.f, 5 / 12.f)));
 		Set( "Sequencer.Timeline.ScrubFill", new BOX_BRUSH( "Sequencer/ScrubFill", FMargin( 2.f/4.f, 0.f ) ) );
 		Set( "Sequencer.Timeline.FrameBlockScrubHandleDown", new BOX_BRUSH( "Sequencer/ScrubHandleDown", FMargin( 6.f/13.f, 5/12.f, 6/13.f, 8/12.f ) ) );
 		Set( "Sequencer.Timeline.FrameBlockScrubHandleUp", new BOX_BRUSH( "Sequencer/ScrubHandleUp", FMargin( 6.f/13.f, 8/12.f, 6/13.f, 5/12.f ) ) );
