PR #5059: UE-52882: Add small icons for paint mode button options (Contributed by [REDACTED_PERSON])

diff --git a/Engine/Source/Editor/EditorStyle/Private/SlateEditorStyle.cpp b/Engine/Source/Editor/EditorStyle/Private/SlateEditorStyle.cpp
--- a/Engine/Source/Editor/EditorStyle/Private/SlateEditorStyle.cpp
+++ b/Engine/Source/Editor/EditorStyle/Private/SlateEditorStyle.cpp
@@ -5096,14 +5096,22 @@ void FSlateEditorStyle::FStyle::SetupLevelEditorStyle()
 	// Mesh Paint
 	{
 		Set( "MeshPaint.Fill", new IMAGE_BRUSH( "/Icons/icon_MeshPaint_Fill_40x", Icon20x20) );
+		Set( "MeshPaint.Fill.Small", new IMAGE_BRUSH( "/Icons/icon_MeshPaint_Fill_40x", Icon16x16) );
 		Set( "MeshPaint.Propagate", new IMAGE_BRUSH( "/Icons/icon_MatEd_Apply_40x", Icon20x20) );
+		Set( "MeshPaint.Propagate.Small", new IMAGE_BRUSH( "/Icons/icon_MatEd_Apply_40x", Icon16x16) );
 		Set( "MeshPaint.Import", new IMAGE_BRUSH( "/Icons/icon_Import_40x", Icon20x20) );
+		Set( "MeshPaint.Import.Small", new IMAGE_BRUSH( "/Icons/icon_Import_40x", Icon16x16) );
 		Set( "MeshPaint.FindInCB", new IMAGE_BRUSH( "/Icons/icon_toolbar_genericfinder_40px", Icon20x20) );
 		Set( "MeshPaint.Save", new IMAGE_BRUSH( "/Icons/icon_file_save_40x", Icon20x20) );
+		Set( "MeshPaint.Save.Small", new IMAGE_BRUSH( "/Icons/icon_file_save_40x", Icon16x16) );
 		Set( "MeshPaint.Fix", new IMAGE_BRUSH( "/Icons/icon_tab_Toolbars_40x", Icon20x20) );
+		Set( "MeshPaint.Fix.Small", new IMAGE_BRUSH( "/Icons/icon_tab_Toolbars_40x", Icon16x16) );
 		Set( "MeshPaint.Remove", new IMAGE_BRUSH("/Icons/Edit/icon_Edit_Delete_40x", Icon20x20));
+		Set( "MeshPaint.Remove.Small", new IMAGE_BRUSH("/Icons/Edit/icon_Edit_Delete_40x", Icon16x16));
 		Set( "MeshPaint.Copy", new IMAGE_BRUSH("/Icons/Edit/icon_Edit_Copy_40x", Icon20x20));
+		Set( "MeshPaint.Copy.Small", new IMAGE_BRUSH("/Icons/Edit/icon_Edit_Copy_40x", Icon16x16));
 		Set( "MeshPaint.Paste", new IMAGE_BRUSH("/Icons/Edit/icon_Edit_Paste_40x", Icon20x20));
+		Set( "MeshPaint.Paste.Small", new IMAGE_BRUSH("/Icons/Edit/icon_Edit_Paste_40x", Icon16x16));
 		Set( "MeshPaint.Swap", new IMAGE_BRUSH("/Icons/icon_MeshPaint_Swap_16x", Icon12x12));
 	}
 
