Back out incorrectly included changes from changelist 4688641 to fix CIS

diff --git a/Engine/Source/Editor/EditorStyle/Private/SlateEditorStyle.cpp b/Engine/Source/Editor/EditorStyle/Private/SlateEditorStyle.cpp
--- a/Engine/Source/Editor/EditorStyle/Private/SlateEditorStyle.cpp
+++ b/Engine/Source/Editor/EditorStyle/Private/SlateEditorStyle.cpp
@@ -229,10 +229,6 @@ void FSlateEditorStyle::FStyle::SetupGeneralStyles()
 		.SetActiveHoveredBrush( IMAGE_BRUSH( "Common/Selection", Icon8x8, SelectionColor ) )
 		.SetInactiveBrush( IMAGE_BRUSH( "Common/Selection", Icon8x8, SelectionColor_Inactive ) )
 		.SetInactiveHoveredBrush( IMAGE_BRUSH( "Common/Selection", Icon8x8, SelectionColor_Inactive ) )
-
-		.SetActiveHighlightedBrush(IMAGE_BRUSH("Common/Selection", Icon8x8, SelectionColor_Inactive))
-		.SetInactiveHighlightedBrush(IMAGE_BRUSH("Common/Selection", Icon8x8, SelectionColor_Inactive))
-
 		.SetTextColor( DefaultForeground )
 		.SetSelectedTextColor(InvertedForeground)
 		.SetDropIndicator_Above(BOX_BRUSH("Common/DropZoneIndicator_Above", FMargin(10.0f / 16.0f, 10.0f / 16.0f, 0, 0), SelectionColor))
@@ -916,20 +912,12 @@ void FSlateEditorStyle::FStyle::SetupGeneralStyles()
 			.SetActiveHoveredBrush(IMAGE_BRUSH("Common/Selection", Icon8x8, SelectionColor))
 			.SetInactiveBrush(IMAGE_BRUSH("Common/Selection", Icon8x8, SelectionColor_Inactive))
 			.SetInactiveHoveredBrush(IMAGE_BRUSH("Common/Selection", Icon8x8, SelectionColor_Inactive))
-
-			.SetActiveHighlightedBrush(IMAGE_BRUSH("Common/Selection", Icon8x8, SelectionColor_Inactive))
-			.SetInactiveHighlightedBrush(IMAGE_BRUSH("Common/Selection", Icon8x8, SelectionColor_Inactive))
 		);
 		Set("TableView.NoHoverTableRow", FTableRowStyle(NormalTableRowStyle)
 			.SetEvenRowBackgroundHoveredBrush(FSlateNoResource())
 			.SetOddRowBackgroundHoveredBrush(FSlateNoResource())
 			.SetActiveHoveredBrush(FSlateNoResource())
 			.SetInactiveHoveredBrush(FSlateNoResource())
-			.SetActiveHoveredBrush(FSlateNoResource())
-			.SetInactiveHoveredBrush(FSlateNoResource())
-
-			.SetActiveHighlightedBrush(FSlateNoResource())
-			.SetInactiveHighlightedBrush(FSlateNoResource())
 			);
 
 
@@ -2215,10 +2203,6 @@ void FSlateEditorStyle::FStyle::SetupGeneralStyles()
 			.SetActiveHoveredBrush( IMAGE_BRUSH( "Common/Selection", Icon8x8, SelectionColor ) )
 			.SetInactiveBrush( IMAGE_BRUSH( "Common/Selection", Icon8x8, SelectionColor ) )
 			.SetInactiveHoveredBrush( IMAGE_BRUSH( "Common/Selection", Icon8x8, SelectionColor )  )
-
-			.SetActiveHighlightedBrush(IMAGE_BRUSH("Common/Selection", Icon8x8, SelectionColor))
-			.SetInactiveHighlightedBrush(IMAGE_BRUSH("Common/Selection", Icon8x8, SelectionColor))
-
 			.SetTextColor( DefaultForeground )
 			.SetSelectedTextColor( InvertedForeground )
 			);
