Programs need to build with support for developer tools in order to get all SlateEditorStyle code..

TODO: Be more restrictive about SlateReflector's use of MessageLog so that there's no dependency on EditorStyle unless needed.

diff --git a/Engine/Source/Editor/EditorStyle/Private/SlateEditorStyle.cpp b/Engine/Source/Editor/EditorStyle/Private/SlateEditorStyle.cpp
--- a/Engine/Source/Editor/EditorStyle/Private/SlateEditorStyle.cpp
+++ b/Engine/Source/Editor/EditorStyle/Private/SlateEditorStyle.cpp
@@ -784,7 +784,7 @@ void FSlateEditorStyle::FStyle::SetupGeneralStyles()
 
 	Set("GenericLink", new IMAGE_BRUSH("Common/link", Icon16x16));
 
-#if WITH_EDITOR || IS_PROGRAM
+#if WITH_EDITOR || (IS_PROGRAM && WITH_UNREAL_DEVELOPER_TOOLS)
 	{
 		// Dark Hyperlink - for use on light backgrounds
 		FButtonStyle DarkHyperlinkButton = FButtonStyle()
@@ -810,7 +810,7 @@ void FSlateEditorStyle::FStyle::SetupGeneralStyles()
 			.SetPadding(FMargin(0.0f));
 		Set("HoverOnlyHyperlink", HoverOnlyHyperlink);
 	}
-#endif // WITH_EDITOR || IS_PROGRAM
+#endif // WITH_EDITOR || (IS_PROGRAM && WITH_UNREAL_DEVELOPER_TOOLS)
 
 	// Expandable button
 	{
@@ -827,7 +827,7 @@ void FSlateEditorStyle::FStyle::SetupGeneralStyles()
 	}
 
 	// Content reference
-#if WITH_EDITOR || IS_PROGRAM
+#if WITH_EDITOR || (IS_PROGRAM && WITH_UNREAL_DEVELOPER_TOOLS)
 	{
 		Set( "ContentReference.Background.Normal", new BOX_BRUSH( "Common/TextBox", FMargin(4.0f/16.0f) ) );
 		Set( "ContentReference.Background.Hovered", new BOX_BRUSH( "Common/TextBox_Hovered", FMargin(4.0f/16.0f) ) );
@@ -838,9 +838,9 @@ void FSlateEditorStyle::FStyle::SetupGeneralStyles()
 		Set( "ContentReference.Clear", new IMAGE_BRUSH( "Icons/Cross_12x", Icon12x12 ) );
 		Set( "ContentReference.Tools", new IMAGE_BRUSH( "Icons/wrench_16x", Icon12x12 ) );
 	}
-#endif // WITH_EDITOR
+#endif // WITH_EDITOR || (IS_PROGRAM && WITH_UNREAL_DEVELOPER_TOOLS)
 
-#if WITH_EDITOR || IS_PROGRAM
+#if WITH_EDITOR || (IS_PROGRAM && WITH_UNREAL_DEVELOPER_TOOLS)
 
 	{
 		Set( "SystemWideCommands.FindInContentBrowser", new IMAGE_BRUSH( "Icons/icon_toolbar_genericfinder_40px", Icon40x40 ) );
@@ -876,7 +876,7 @@ void FSlateEditorStyle::FStyle::SetupGeneralStyles()
 		Set( "DialogueWaveParameter.DropDownBorder", new BOX_BRUSH( "Old/Border", 4.0f/16.0f, FLinearColor::Black) );
 	}
 
-#endif // WITH_EDITOR || IS_PROGRAM
+#endif // WITH_EDITOR || (IS_PROGRAM && WITH_UNREAL_DEVELOPER_TOOLS)
 
 	Set( "DashedBorder", new BORDER_BRUSH( "Old/DashedBorder", FMargin(6.0f/32.0f) ) );
 	Set( "Checker", new IMAGE_BRUSH( "Old/Checker", Icon16x16, FLinearColor::White, ESlateBrushTileType::Both ) );
@@ -884,14 +884,14 @@ void FSlateEditorStyle::FStyle::SetupGeneralStyles()
 	Set( "UniformShadow_Tint", new BORDER_BRUSH( "Common/UniformShadow_Tint", FMargin( 16.0f / 64.0f ) ) );
 
 	// Splitter
-#if WITH_EDITOR || IS_PROGRAM
+#if WITH_EDITOR || (IS_PROGRAM && WITH_UNREAL_DEVELOPER_TOOLS)
 	{
 		Set( "Splitter", FSplitterStyle()
 			.SetHandleNormalBrush( FSlateNoResource() )
 			.SetHandleHighlightBrush( IMAGE_BRUSH( "Common/SplitterHandleHighlight", Icon8x8, FLinearColor::White ) )
 			);
 	}
-#endif // WITH_EDITOR || IS_PROGRAM
+#endif // WITH_EDITOR || (IS_PROGRAM && WITH_UNREAL_DEVELOPER_TOOLS)
 
 	// Scroll Box
 	{
@@ -1121,7 +1121,7 @@ void FSlateEditorStyle::FStyle::SetupGeneralStyles()
 		Set( "MessageLog.ListBorder", new BOX_BRUSH( "/Docking/AppTabContentArea", FMargin(4/16.0f) ) );
 	}
 
-#if WITH_EDITOR || IS_PROGRAM
+#if WITH_EDITOR || (IS_PROGRAM && WITH_UNREAL_DEVELOPER_TOOLS)
 
 	// Animation tools
 	if (IncludeEditorSpecificStyles())
@@ -1148,7 +1148,7 @@ void FSlateEditorStyle::FStyle::SetupGeneralStyles()
 	
 	Set("TrashCan", new IMAGE_BRUSH( "Old/TrashCan", FVector2D(64, 64)));
 	Set("TrashCan_Small", new IMAGE_BRUSH( "Old/TrashCan_Small", FVector2D(18, 18)));
-#endif // WITH_EDITOR || IS_PROGRAM
+#endif // WITH_EDITOR || (IS_PROGRAM && WITH_UNREAL_DEVELOPER_TOOLS)
 
 	// Embossed Widget Text
 	Set( "EmbossedText", FTextBlockStyle(NormalText)
@@ -1160,7 +1160,7 @@ void FSlateEditorStyle::FStyle::SetupGeneralStyles()
 
 
 	// Output Log Window
-#if WITH_EDITOR || IS_PROGRAM
+#if WITH_EDITOR || (IS_PROGRAM && WITH_UNREAL_DEVELOPER_TOOLS)
 	{
 		const int32 LogFontSize = Settings.IsValid() ? Settings->LogFontSize : 9;
 
@@ -1309,9 +1309,9 @@ void FSlateEditorStyle::FStyle::SetupGeneralStyles()
 	{
 		Set( "PackageRestore.FolderOpen", new IMAGE_BRUSH( "Icons/FolderOpen", FVector2D(18, 16) ) );
 	}
-#endif // WITH_EDITOR || IS_PROGRAM
+#endif // WITH_EDITOR || (IS_PROGRAM && WITH_UNREAL_DEVELOPER_TOOLS)
 
-#if WITH_EDITOR || IS_PROGRAM
+#if WITH_EDITOR || (IS_PROGRAM && WITH_UNREAL_DEVELOPER_TOOLS)
 	// Expandable area
 	{
 		Set( "ExpandableArea", FExpandableAreaStyle()
@@ -1333,9 +1333,9 @@ void FSlateEditorStyle::FStyle::SetupGeneralStyles()
 		Set( "SavePackages.SCC_DlgReadOnly", new IMAGE_BRUSH( "Old/SavePackages/SCC_DlgReadOnly", FVector2D( 18, 16 ) ) );
 		Set( "SavePackages.SCC_DlgNoIcon", new FSlateNoResource() );
 	}
-#endif // WITH_EDITOR || IS_PROGRAM
+#endif // WITH_EDITOR || (IS_PROGRAM && WITH_UNREAL_DEVELOPER_TOOLS)
 
-#if WITH_EDITOR || IS_PROGRAM
+#if WITH_EDITOR || (IS_PROGRAM && WITH_UNREAL_DEVELOPER_TOOLS)
 	// Layers General
 	{
 		Set( "Layer.Icon16x", new IMAGE_BRUSH( "Icons/layer_16x", Icon16x16 ) );
@@ -1506,9 +1506,9 @@ void FSlateEditorStyle::FStyle::SetupGeneralStyles()
 			.SetFont(DEFAULT_FONT("Fonts/Roboto-Italic", 9))
 			.SetColorAndOpacity(FSlateColor::UseSubduedForeground()));
 	}
-#endif // WITH_EDITOR
+#endif // WITH_EDITOR || (IS_PROGRAM && WITH_UNREAL_DEVELOPER_TOOLS)
 
-#if WITH_EDITOR || IS_PROGRAM
+#if WITH_EDITOR || (IS_PROGRAM && WITH_UNREAL_DEVELOPER_TOOLS)
 	// Breadcrumb Trail
 	{
 		Set( "BreadcrumbTrail.Delimiter", new IMAGE_BRUSH( "Common/Delimiter", Icon16x16 ) );
@@ -1532,9 +1532,9 @@ void FSlateEditorStyle::FStyle::SetupGeneralStyles()
 		Set( "NotificationList.DefaultMessage", new IMAGE_BRUSH( "Old/EventMessage_Default", Icon40x40 ) );
 		Set( "NotificationList.Glow", new FSlateColorBrush( FColor(255, 255, 255, 255) ) );
 	}
-#endif // WITH_EDITOR || IS_PROGRAM
+#endif // WITH_EDITOR || (IS_PROGRAM && WITH_UNREAL_DEVELOPER_TOOLS)
 
-#if WITH_EDITOR || IS_PROGRAM
+#if WITH_EDITOR || (IS_PROGRAM && WITH_UNREAL_DEVELOPER_TOOLS)
 	// Asset editors (common)
 	{
 		Set( "AssetEditor.SaveAsset.Greyscale", new IMAGE_BRUSH( "Icons/icon_file_save_16px", Icon16x16 ) );
@@ -1631,7 +1631,7 @@ void FSlateEditorStyle::FStyle::SetupGeneralStyles()
 	}
 
 	Set( "Editor.SearchBoxFont", DEFAULT_FONT( "Regular", 12) );
-#endif // WITH_EDITOR || IS_PROGRAM
+#endif // WITH_EDITOR || (IS_PROGRAM && WITH_UNREAL_DEVELOPER_TOOLS)
 
 	// Slider and Volume Control
 	{
@@ -1660,7 +1660,7 @@ void FSlateEditorStyle::FStyle::SetupGeneralStyles()
 		Set( "DebugConsole.Background", new BOX_BRUSH( "Old/Menu_Background", FMargin(8.0f/64.0f) ) );
 	}
 
-#if WITH_EDITOR || IS_PROGRAM
+#if WITH_EDITOR || (IS_PROGRAM && WITH_UNREAL_DEVELOPER_TOOLS)
 	// About screen
 	if (IncludeEditorSpecificStyles())
 	{
@@ -1672,7 +1672,7 @@ void FSlateEditorStyle::FStyle::SetupGeneralStyles()
 		Set( "AboutScreen.EpicGames", new IMAGE_BRUSH( "About/EpicGamesIcon", FVector2D(50,50) ) );
 		Set( "AboutScreen.EpicGamesHovered", new IMAGE_BRUSH( "About/EpicGamesIcon_Hovered", FVector2D(50,50) ) );
 	}
-#endif // WITH_EDITOR || IS_PROGRAM
+#endif // WITH_EDITOR || (IS_PROGRAM && WITH_UNREAL_DEVELOPER_TOOLS)
 
 #if WITH_EDITOR
 	// Credits screen
@@ -1759,7 +1759,7 @@ void FSlateEditorStyle::FStyle::SetupGeneralStyles()
 #endif
 
 	// New Level Dialog
-#if WITH_EDITOR || IS_PROGRAM
+#if WITH_EDITOR || (IS_PROGRAM && WITH_UNREAL_DEVELOPER_TOOLS)
 	{
 		Set( "NewLevelDialog.BlackBorder", new FSlateColorBrush( FColor(0, 0, 0, 100) ) );
 		Set( "NewLevelDialog.Blank", new IMAGE_BRUSH( "NewLevels/NewLevelBlank", FVector2D(256,256) ) );
@@ -2148,7 +2148,7 @@ void FSlateEditorStyle::FStyle::SetupGeneralStyles()
 			.SetShadowOffset(FVector2D(1, 1))
 			.SetShadowColorAndOpacity(FLinearColor(0, 0, 0, 0.9f)));
 	}
-#endif // WITH_EDITOR || IS_PROGRAM
+#endif // WITH_EDITOR || (IS_PROGRAM && WITH_UNREAL_DEVELOPER_TOOLS)
 
 
 #if WITH_EDITOR
@@ -2322,7 +2322,7 @@ void FSlateEditorStyle::FStyle::SetupGeneralStyles()
 
 #endif // WITH_EDITOR
 
-#if WITH_EDITOR || IS_PROGRAM
+#if WITH_EDITOR || (IS_PROGRAM && WITH_UNREAL_DEVELOPER_TOOLS)
 	// ToolBar
 	{
 		Set( "ToolBar.Background", new BOX_BRUSH( "Common/GroupBorder", FMargin(4.0f/16.0f) ) );
@@ -2841,9 +2841,9 @@ void FSlateEditorStyle::FStyle::SetupGeneralStyles()
 
 		Set("Common.GotoNativeCodeHyperlink", EditNativeHyperlinkStyle);
 	}
-#endif // WITH_EDITOR || IS_PROGRAM
+#endif // WITH_EDITOR || (IS_PROGRAM && WITH_UNREAL_DEVELOPER_TOOLS)
 
-#if WITH_EDITOR || IS_PROGRAM
+#if WITH_EDITOR || (IS_PROGRAM && WITH_UNREAL_DEVELOPER_TOOLS)
 
 	// Gameplay Tags
 	{
@@ -2889,7 +2889,7 @@ void FSlateEditorStyle::FStyle::SetupGeneralStyles()
 	{
 		Set("TimecodeProvider.TabIcon", new IMAGE_BRUSH("Icons/icon_tab_TimecodeProvider_16x", Icon16x16));
 	}
-#endif
+#endif // WITH_EDITOR || (IS_PROGRAM && WITH_UNREAL_DEVELOPER_TOOLS)
 }
 
 void FSlateEditorStyle::FStyle::SetupGeneralIcons()
@@ -3004,7 +3004,7 @@ void FSlateEditorStyle::FStyle::SetupProjectBadgeStyle()
 
 void FSlateEditorStyle::FStyle::SetupDockingStyles()
 {
-#if WITH_EDITOR || IS_PROGRAM
+#if WITH_EDITOR || (IS_PROGRAM && WITH_UNREAL_DEVELOPER_TOOLS)
 	// Tabs, Docking, Flexible Layout
 	{
 		// Tab Text
@@ -3063,7 +3063,7 @@ void FSlateEditorStyle::FStyle::SetupDockingStyles()
 		Set( "Docking.TabConextButton.Pressed", new IMAGE_BRUSH( "/Docking/TabContextButton", FVector2D(24,24) ) );
 		Set( "Docking.TabConextButton.Hovered", new IMAGE_BRUSH( "/Docking/TabContextButton", FVector2D(24,24) ) );
 	}
-#endif // WITH_EDITOR || IS_PROGRAM
+#endif // WITH_EDITOR || (IS_PROGRAM && WITH_UNREAL_DEVELOPER_TOOLS)
 	}
 
 void FSlateEditorStyle::FStyle::SetupTutorialStyles()
@@ -3672,7 +3672,7 @@ void FSlateEditorStyle::FStyle::SetupPropertyEditorStyles()
 
 void FSlateEditorStyle::FStyle::SetupProfilerStyle()
 {
-#if WITH_EDITOR || IS_PROGRAM
+#if WITH_EDITOR || (IS_PROGRAM && WITH_UNREAL_DEVELOPER_TOOLS)
 	// Profiler
 	{
 		// Profiler group brushes
@@ -3842,7 +3842,7 @@ void FSlateEditorStyle::FStyle::SetupProfilerStyle()
 		//Set( "Profiler.Type.Calls", new IMAGE_BRUSH( "Icons/Profiler/profiler_Calls_32x", Icon40x40) );
 		//Set( "Profiler.Type.Calls.Small", new IMAGE_BRUSH( "Icons/Profiler/profiler_Calls_32x", Icon20x20) );
 	}
-#endif // WITH_EDITOR || IS_PROGRAM
+#endif // WITH_EDITOR || (IS_PROGRAM && WITH_UNREAL_DEVELOPER_TOOLS)
 }
 	
 void FSlateEditorStyle::FStyle::SetupGraphEditorStyles()
@@ -3850,7 +3850,7 @@ void FSlateEditorStyle::FStyle::SetupGraphEditorStyles()
 	const FScrollBarStyle ScrollBar = GetWidgetStyle<FScrollBarStyle>( "ScrollBar" );
 
 	// Graph Editor
-#if WITH_EDITOR || IS_PROGRAM
+#if WITH_EDITOR || (IS_PROGRAM && WITH_UNREAL_DEVELOPER_TOOLS)
 	{
 		Set( "Graph.ForegroundColor", FLinearColor(218.0f/255.0f, 218.0f/255.0f, 218.0f/255.0f, 1.0f) );
 
@@ -4635,7 +4635,7 @@ void FSlateEditorStyle::FStyle::SetupGraphEditorStyles()
 			.SetShadowOffset(FVector2D::ZeroVector)
 		);
 	}
-#endif // WITH_EDITOR || IS_PROGRAM
+#endif // WITH_EDITOR || (IS_PROGRAM && WITH_UNREAL_DEVELOPER_TOOLS)
 	}
 
 void FSlateEditorStyle::FStyle::SetupLevelEditorStyle()
@@ -4715,7 +4715,7 @@ void FSlateEditorStyle::FStyle::SetupLevelEditorStyle()
 		Set( "EditorViewport.Back", new IMAGE_BRUSH("Icons/icon_ViewMode_ViewBack_16px", Icon16x16 ) );
 #endif
 
-#if WITH_EDITOR || IS_PROGRAM
+#if WITH_EDITOR || (IS_PROGRAM && WITH_UNREAL_DEVELOPER_TOOLS)
 		{
 			Set( "LevelEditor.Tabs.Details", new IMAGE_BRUSH( "/Icons/icon_tab_SelectionDetails_16x", Icon16x16 ) );
 			Set( "LevelEditor.Tabs.EditorModes", new IMAGE_BRUSH( "/Icons/icon_Editor_Modes_16x", Icon16x16 ) );
@@ -4734,7 +4734,7 @@ void FSlateEditorStyle::FStyle::SetupLevelEditorStyle()
 			Set( "LevelEditor.Tabs.Viewports", new IMAGE_BRUSH( "/Icons/icon_tab_Viewports_16x", Icon16x16 ) );
 			Set( "LevelEditor.Tabs.HLOD", new IMAGE_BRUSH("/Icons/icon_tab_layers_16px", Icon16x16));
 		}
-#endif
+#endif // WITH_EDITOR || (IS_PROGRAM && WITH_UNREAL_DEVELOPER_TOOLS)
 
 #if WITH_EDITOR
 		Set( "LevelEditor.NewLevel", new IMAGE_BRUSH( "Icons/icon_file_new_16px", Icon16x16 ) );
@@ -4992,10 +4992,10 @@ void FSlateEditorStyle::FStyle::SetupLevelEditorStyle()
 	{
 		Set("MeshProxy.SimplygonLogo", new IMAGE_BRUSH( "Icons/SimplygonBanner_Sml", FVector2D(174, 36) ) );
 	}
-#endif // WITH_EDITOR || IS_PROGRAM
+#endif // WITH_EDITOR
 
 	// Level viewport 
-#if WITH_EDITOR || IS_PROGRAM
+#if WITH_EDITOR || (IS_PROGRAM && WITH_UNREAL_DEVELOPER_TOOLS)
 	{
 		Set( "LevelViewport.ActiveViewportBorder", new BORDER_BRUSH( "Old/White", FMargin(1), SelectionColor ) ); 
 		Set( "LevelViewport.NoViewportBorder", new FSlateNoResource() );
@@ -5071,7 +5071,7 @@ void FSlateEditorStyle::FStyle::SetupLevelEditorStyle()
 		Set( "ShowFlagsMenu.Translucency", new IMAGE_BRUSH( "Icons/icon_ShowTranslucency_16x", Icon16x16 ) );
 		Set( "ShowFlagsMenu.WidgetComponents", new IMAGE_BRUSH( "UMG/Designer_16x", Icon16x16 ) );
 	}
-#endif // WITH_EDITOR || IS_PROGRAM
+#endif // WITH_EDITOR || (IS_PROGRAM && WITH_UNREAL_DEVELOPER_TOOLS)
 
 	// Mobility Icons
 	{
@@ -7322,13 +7322,13 @@ void FSlateEditorStyle::FStyle::SetupMatineeStyle()
 		Set( "Matinee.ToggleCurveEditor.Small", new IMAGE_BRUSH("Icons/icon_MatineeCurveView_40px", Icon20x20) );
 		Set( "Matinee.ToggleDirectorTimeline.Small", new IMAGE_BRUSH("Icons/icon_MatineeDirectorView_40px", Icon20x20) );
 	}
-#endif // WITH_EDITOR || IS_PROGRAM
+#endif // WITH_EDITOR
 	}
 
 void FSlateEditorStyle::FStyle::SetupSourceControlStyles()
 {
 	//Source Control
-#if WITH_EDITOR || IS_PROGRAM
+#if WITH_EDITOR || (IS_PROGRAM && WITH_UNREAL_DEVELOPER_TOOLS)
 	{
 		Set( "SourceControl.Add", new IMAGE_BRUSH( "Old/SourceControl/SCC_Action_Add",Icon10x10));
 		Set( "SourceControl.Edit", new IMAGE_BRUSH( "Old/SourceControl/SCC_Action_Edit",Icon10x10));
@@ -7356,10 +7356,10 @@ void FSlateEditorStyle::FStyle::SetupSourceControlStyles()
 		Set( "SourceControl.Actions.Refresh", new IMAGE_BRUSH( "Icons/icon_Refresh_16x", Icon16x16 ) );
 		Set( "SourceControl.Actions.ChangeSettings", new IMAGE_BRUSH( "Icons/icon_SCC_Change_Source_Control_Settings_16x", Icon16x16 ) );
 	}
-#endif // WITH_EDITOR || IS_PROGRAM
+#endif // WITH_EDITOR || (IS_PROGRAM && WITH_UNREAL_DEVELOPER_TOOLS)
 
 	// Perforce
-#if WITH_EDITOR || IS_PROGRAM
+#if WITH_EDITOR || (IS_PROGRAM && WITH_UNREAL_DEVELOPER_TOOLS)
 	{
 		Set( "Perforce.CheckedOut", new IMAGE_BRUSH( "ContentBrowser/SCC_CheckedOut", Icon32x32) );
 		Set( "Perforce.CheckedOut_Small", new IMAGE_BRUSH( "ContentBrowser/SCC_CheckedOut", Icon16x16) );
@@ -7380,10 +7380,10 @@ void FSlateEditorStyle::FStyle::SetupSourceControlStyles()
 		Set( "Perforce.Branched", new IMAGE_BRUSH( "ContentBrowser/SCC_Branched", Icon32x32) );
 		Set( "Perforce.Branched_Small", new IMAGE_BRUSH( "ContentBrowser/SCC_Branched", Icon16x16) );
 	}
-#endif // WITH_EDITOR || IS_PROGRAM
+#endif // WITH_EDITOR || (IS_PROGRAM && WITH_UNREAL_DEVELOPER_TOOLS)
 
 	// Subversion
-#if WITH_EDITOR || IS_PROGRAM
+#if WITH_EDITOR || (IS_PROGRAM && WITH_UNREAL_DEVELOPER_TOOLS)
 	{
 		Set( "Subversion.CheckedOut", new IMAGE_BRUSH( "ContentBrowser/SCC_CheckedOut", Icon32x32) );
 		Set( "Subversion.CheckedOut_Small", new IMAGE_BRUSH( "ContentBrowser/SCC_CheckedOut", Icon16x16) );
@@ -7404,13 +7404,13 @@ void FSlateEditorStyle::FStyle::SetupSourceControlStyles()
 		Set( "Subversion.Branched", new IMAGE_BRUSH( "ContentBrowser/SCC_Branched", Icon32x32) );
 		Set( "Subversion.Branched_Small", new IMAGE_BRUSH( "ContentBrowser/SCC_Branched", Icon16x16) );
 	}
-#endif // WITH_EDITOR || IS_PROGRAM
+#endif // WITH_EDITOR || (IS_PROGRAM && WITH_UNREAL_DEVELOPER_TOOLS)
 	}
 
 void FSlateEditorStyle::FStyle::SetupAutomationStyles()
 {
 	//Automation
-#if WITH_EDITOR || IS_PROGRAM
+#if WITH_EDITOR || (IS_PROGRAM && WITH_UNREAL_DEVELOPER_TOOLS)
 	{
 		Set( "Automation.Header" , FTextBlockStyle(NormalText)
 			.SetFont( DEFAULT_FONT( "Mono", 12 ) )
@@ -7597,7 +7597,7 @@ void FSlateEditorStyle::FStyle::SetupAutomationStyles()
 		Set("NavigationHyperlink", NavigationHyperlink);
 	}
 
-#endif // WITH_EDITOR || IS_PROGRAM
+#endif // WITH_EDITOR || (IS_PROGRAM && WITH_UNREAL_DEVELOPER_TOOLS)
 
 	// External image picker
 	{
diff --git a/Engine/Source/Programs/SlateViewer/SlateViewer.Target.cs b/Engine/Source/Programs/SlateViewer/SlateViewer.Target.cs
--- a/Engine/Source/Programs/SlateViewer/SlateViewer.Target.cs
+++ b/Engine/Source/Programs/SlateViewer/SlateViewer.Target.cs
@@ -26,5 +26,8 @@ public class SlateViewerTarget : TargetRules
 		// SlateViewer.exe has no exports, so no need to verify that a .lib and .exp file was emitted by
 		// the linker.
 		bHasExports = false;
+
+		// Make sure to get all code in SlateEditorStyle compiled in
+        bBuildDeveloperTools = true;
 	}
 }
