Enable optimizations for the EditorStyle module, and just disable them inside the SlateEditorStyle.cpp file. This prevents a special non-optimized PCH having to be generated to compile it.

diff --git a/Engine/Source/Editor/EditorStyle/EditorStyle.Build.cs b/Engine/Source/Editor/EditorStyle/EditorStyle.Build.cs
--- a/Engine/Source/Editor/EditorStyle/EditorStyle.Build.cs
+++ b/Engine/Source/Editor/EditorStyle/EditorStyle.Build.cs
@@ -39,8 +39,5 @@ public class EditorStyle : ModuleRules
 				}
 			);
 		}
-
-		// I'm told this is to improve compilation performance of this module
-		OptimizeCode = CodeOptimization.Never;
 	}
 }
diff --git a/Engine/Source/Editor/EditorStyle/Private/SlateEditorStyle.cpp b/Engine/Source/Editor/EditorStyle/Private/SlateEditorStyle.cpp
--- a/Engine/Source/Editor/EditorStyle/Private/SlateEditorStyle.cpp
+++ b/Engine/Source/Editor/EditorStyle/Private/SlateEditorStyle.cpp
@@ -3,11 +3,14 @@
 #include "SlateEditorStyle.h"
 #include "Misc/CommandLine.h"
 #include "Styling/CoreStyle.h"
+#include "SlateOptMacros.h"
 
 #if (WITH_EDITOR || (IS_PROGRAM && PLATFORM_DESKTOP))
 	#include "PlatformInfo.h"
 #endif
 
+BEGIN_SLATE_FUNCTION_BUILD_OPTIMIZATION
+
 #define IMAGE_BRUSH( RelativePath, ... ) FSlateImageBrush( RootToContentDir( RelativePath, TEXT(".png") ), __VA_ARGS__ )
 #define BOX_BRUSH( RelativePath, ... ) FSlateBoxBrush( RootToContentDir( RelativePath, TEXT(".png") ), __VA_ARGS__ )
 #define BORDER_BRUSH( RelativePath, ... ) FSlateBorderBrush( RootToContentDir( RelativePath, TEXT(".png") ), __VA_ARGS__ )
@@ -7870,3 +7873,5 @@ void FSlateEditorStyle::FStyle::SetupMyBlueprintStyles()
 #undef BORDER_BRUSH
 #undef DEFAULT_FONT
 #undef ICON_FONT
+
+END_SLATE_FUNCTION_BUILD_OPTIMIZATION
\ No newline at end of file
