Adding missing Toolbar.Check style

diff --git a/Engine/Source/Editor/EditorStyle/Private/SlateEditorStyle.cpp b/Engine/Source/Editor/EditorStyle/Private/SlateEditorStyle.cpp
--- a/Engine/Source/Editor/EditorStyle/Private/SlateEditorStyle.cpp
+++ b/Engine/Source/Editor/EditorStyle/Private/SlateEditorStyle.cpp
@@ -2366,6 +2366,22 @@ void FSlateEditorStyle::FStyle::SetupGeneralStyles()
 		/* ... and add new style */
 		Set( "ToolBar.CheckBox", ToolBarCheckBoxStyle );
 
+		/* Read-only checkbox that appears next to a toolbar item */
+		/* Set images for various SCheckBox states associated with read-only toolbar check box items... */
+		const FCheckBoxStyle ToolBarCheckStyle = FCheckBoxStyle()
+			.SetUncheckedImage(IMAGE_BRUSH("Icons/Empty_14x", Icon14x14))
+			.SetUncheckedHoveredImage(IMAGE_BRUSH("Icons/Empty_14x", Icon14x14))
+			.SetUncheckedPressedImage(IMAGE_BRUSH("Common/SmallCheckBox_Hovered", Icon14x14))
+			.SetCheckedImage(IMAGE_BRUSH("Common/SmallCheck", Icon14x14))
+			.SetCheckedHoveredImage(IMAGE_BRUSH("Common/SmallCheck", Icon14x14))
+			.SetCheckedPressedImage(IMAGE_BRUSH("Common/SmallCheck", Icon14x14))
+			.SetUndeterminedImage(IMAGE_BRUSH("Icons/Empty_14x", Icon14x14))
+			.SetUndeterminedHoveredImage(FSlateNoResource())
+			.SetUndeterminedPressedImage(FSlateNoResource());
+
+		/* ...and add the new style */
+		Set("Toolbar.Check", ToolBarCheckStyle);
+
 		// This radio button is actually just a check box with different images
 		/* Create style for "ToolBar.RadioButton" widget ... */
 		const FCheckBoxStyle ToolbarRadioButtonCheckBoxStyle = FCheckBoxStyle()
