Sequencer: Remove disabled autokey icon

diff --git a/Engine/Source/Editor/EditorStyle/Private/SlateEditorStyle.cpp b/Engine/Source/Editor/EditorStyle/Private/SlateEditorStyle.cpp
--- a/Engine/Source/Editor/EditorStyle/Private/SlateEditorStyle.cpp
+++ b/Engine/Source/Editor/EditorStyle/Private/SlateEditorStyle.cpp
@@ -1890,6 +1890,8 @@ void FSlateEditorStyle::FStyle::SetupGeneralStyles()
 		Set( "Sequencer.AnimationOutliner.RegularFont", DEFAULT_FONT( "Regular", 9 ) );
 		Set( "Sequencer.ShotFilter", new IMAGE_BRUSH( "Sequencer/FilteredArea", FVector2D(74,74), FLinearColor::White, ESlateBrushTileType::Both ) );
 		Set( "Sequencer.KeyMark", new IMAGE_BRUSH("Sequencer/KeyMark", FVector2D(3,21), FLinearColor::White, ESlateBrushTileType::NoTile ) );
+		Set( "Sequencer.ToggleAutoKeyEnabled", new IMAGE_BRUSH("Sequencer/Main_Icons/Icon_Sequencer_Auto_Key_24x", Icon48x48));
+		Set( "Sequencer.ToggleAutoKeyEnabled.Small", new IMAGE_BRUSH("Sequencer/Main_Icons/Icon_Sequencer_Auto_Key_24x", Icon24x24));
 		Set( "Sequencer.SetAutoKey", new IMAGE_BRUSH( "Sequencer/Main_Icons/Icon_Sequencer_Auto_Key_24x", Icon48x48 ) );
 		Set( "Sequencer.SetAutoKey.Small", new IMAGE_BRUSH( "Sequencer/Main_Icons/Icon_Sequencer_Auto_Key_24x", Icon24x24 ) );
 		Set( "Sequencer.SetAutoTrack", new IMAGE_BRUSH("Sequencer/Main_Icons/Icon_Sequencer_Auto_Track_24x", Icon48x48 ) );
diff --git a/Engine/Source/Editor/Sequencer/Private/SSequencer.cpp b/Engine/Source/Editor/Sequencer/Private/SSequencer.cpp
--- a/Engine/Source/Editor/Sequencer/Private/SSequencer.cpp
+++ b/Engine/Source/Editor/Sequencer/Private/SSequencer.cpp
@@ -911,15 +911,7 @@ TSharedRef<SWidget> SSequencer::MakeToolBar()
 		}
 		else
 		{
-			TAttribute<FSlateIcon> AutoKeyIcon;
-			AutoKeyIcon.Bind(TAttribute<FSlateIcon>::FGetter::CreateLambda([&]{
-				static FSlateIcon AutoKeyEnabledIcon(FEditorStyle::GetStyleSetName(), "Sequencer.SetAutoKey");
-				static FSlateIcon AutoKeyDisabledIcon(FEditorStyle::GetStyleSetName(), "Sequencer.SetAutoChangeNone");
-
-				return SequencerPtr.Pin()->GetAutoChangeMode() == EAutoChangeMode::None ? AutoKeyDisabledIcon : AutoKeyEnabledIcon;
-			}));
-
-			ToolBarBuilder.AddToolBarButton( FSequencerCommands::Get().ToggleAutoKeyEnabled, NAME_None, TAttribute<FText>(), TAttribute<FText>(), AutoKeyIcon );
+			ToolBarBuilder.AddToolBarButton( FSequencerCommands::Get().ToggleAutoKeyEnabled );
 		}
 
 		if( SequencerPtr.Pin()->IsLevelEditorSequencer() )
