Sequencer: Fixed bad merge, icon definition got lost in dev-editor refactor

diff --git a/Engine/Source/Editor/EditorStyle/Private/SlateEditorStyle.cpp b/Engine/Source/Editor/EditorStyle/Private/SlateEditorStyle.cpp
--- a/Engine/Source/Editor/EditorStyle/Private/SlateEditorStyle.cpp
+++ b/Engine/Source/Editor/EditorStyle/Private/SlateEditorStyle.cpp
@@ -2433,6 +2433,8 @@ void FSlateEditorStyle::FStyle::SetupSequencerStyles()
 		Set("Sequencer.AnimationOutliner.RegularFont", DEFAULT_FONT("Regular", 9));
 		Set("Sequencer.ShotFilter", new IMAGE_BRUSH("Sequencer/FilteredArea", FVector2D(74, 74), FLinearColor::White, ESlateBrushTileType::Both));
 		Set("Sequencer.KeyMark", new IMAGE_BRUSH("Sequencer/KeyMark", FVector2D(3, 21), FLinearColor::White, ESlateBrushTileType::NoTile));
+		Set("Sequencer.ToggleAutoKeyEnabled", new IMAGE_BRUSH("Sequencer/Main_Icons/Icon_Sequencer_Auto_Key_24x", Icon48x48));
+		Set("Sequencer.ToggleAutoKeyEnabled.Small", new IMAGE_BRUSH("Sequencer/Main_Icons/Icon_Sequencer_Auto_Key_24x", Icon24x24));
 		Set("Sequencer.SetAutoKey", new IMAGE_BRUSH("Sequencer/Main_Icons/Icon_Sequencer_Auto_Key_24x", Icon48x48));
 		Set("Sequencer.SetAutoKey.Small", new IMAGE_BRUSH("Sequencer/Main_Icons/Icon_Sequencer_Auto_Key_24x", Icon24x24));
 		Set("Sequencer.SetAutoTrack", new IMAGE_BRUSH("Sequencer/Main_Icons/Icon_Sequencer_Auto_Track_24x", Icon48x48));
@@ -2455,8 +2457,8 @@ void FSlateEditorStyle::FStyle::SetupSequencerStyles()
 		Set("Sequencer.SetKeyChanged.Small", new IMAGE_BRUSH("Sequencer/Main_Icons/Icon_Sequencer_Key_Part_24x", Icon24x24));
 		Set("Sequencer.ToggleIsSnapEnabled", new IMAGE_BRUSH("Sequencer/Main_Icons/Icon_Sequencer_Snap_24x", Icon48x48));
 		Set("Sequencer.ToggleIsSnapEnabled.Small", new IMAGE_BRUSH("Sequencer/Main_Icons/Icon_Sequencer_Snap_24x", Icon24x24));
-		Set("Sequencer.ToggleShowCurveEditor", new IMAGE_BRUSH("Icons/SequencerIcons/icon_Sequencer_CurveEditor_24x", Icon48x48));
-		Set("Sequencer.ToggleShowCurveEditor.Small", new IMAGE_BRUSH("Icons/SequencerIcons/icon_Sequencer_CurveEditor_24x", Icon24x24));
+		Set("Sequencer.ToggleShowCurveEditor", new IMAGE_BRUSH("GenericCurveEditor/Icons/GenericCurveEditor_48x", Icon48x48));
+		Set("Sequencer.ToggleShowCurveEditor.Small", new IMAGE_BRUSH("GenericCurveEditor/Icons/GenericCurveEditor_48x", Icon24x24));
 		Set("Sequencer.ToggleAutoScroll", new IMAGE_BRUSH("Icons/icon_Sequencer_ToggleAutoScroll_40x", Icon48x48));
 		Set("Sequencer.ToggleAutoScroll.Small", new IMAGE_BRUSH("Icons/icon_Sequencer_ToggleAutoScroll_16x", Icon16x16));
 		Set("Sequencer.MoveTool.Small", new IMAGE_BRUSH("Icons/SequencerIcons/icon_Sequencer_Move_24x", Icon16x16));
