Sequencer: Fixed a stray '.' that came in from main merge in the icon refactor styles

diff --git a/Engine/Source/Editor/EditorStyle/Private/SlateEditorStyle.cpp b/Engine/Source/Editor/EditorStyle/Private/SlateEditorStyle.cpp
--- a/Engine/Source/Editor/EditorStyle/Private/SlateEditorStyle.cpp
+++ b/Engine/Source/Editor/EditorStyle/Private/SlateEditorStyle.cpp
@@ -2413,7 +2413,7 @@ void FSlateEditorStyle::FStyle::SetupSequencerStyles()
 
 		Set("Sequencer.Section.PinCusion", new IMAGE_BRUSH(TEXT("Sequencer/PinCusion"), Icon16x16, FLinearColor::White, ESlateBrushTileType::Both));
 		Set("Sequencer.Section.OverlapBorder", new BORDER_BRUSH(TEXT("Sequencer/OverlapBorder"), FMargin(1.f / 4.f, 0.f)));
-		Set("Sequencer.Section.StripeO.verlay", new BOX_BRUSH("Sequencer/SectionStripeOverlay", FMargin(0.f, .5f)));
+		Set("Sequencer.Section.StripeOverlay", new BOX_BRUSH("Sequencer/SectionStripeOverlay", FMargin(0.f, .5f)));
 		Set("Sequencer.Section.BackgroundText", DEFAULT_FONT("Bold", 24));
 		Set("Sequencer.Section.EmptySpace", new BOX_BRUSH(TEXT("Sequencer/EmptySpace"), FMargin(0.f, 7.f / 14.f)));
 
