Update the slate icon for the level blueprint editor tab to be correct

#rnx

diff --git a/Engine/Source/Editor/EditorStyle/Private/SlateEditorStyle.cpp b/Engine/Source/Editor/EditorStyle/Private/SlateEditorStyle.cpp
--- a/Engine/Source/Editor/EditorStyle/Private/SlateEditorStyle.cpp
+++ b/Engine/Source/Editor/EditorStyle/Private/SlateEditorStyle.cpp
@@ -6162,6 +6162,7 @@ void FSlateEditorStyle::FStyle::SetupClassIconsAndThumbnails()
 		Set("ClassIcon.Deleted", new IMAGE_BRUSH("Icons/ActorIcons/DeletedActor_16px", Icon16x16));
 		Set("ClassIcon.StaticMeshActor", new IMAGE_BRUSH("Icons/AssetIcons/StaticMesh_16x", Icon16x16));
 		Set("ClassIcon.SkeletalMeshActor", new IMAGE_BRUSH("Icons/AssetIcons/SkeletalMesh_16x", Icon16x16));
+		Set("ClassIcon.LevelScriptActor", new IMAGE_BRUSH("Icons/AssetIcons/Blueprint_16x", Icon16x16));
 
 		// Component classes
 		Set("ClassIcon.AudioComponent", new IMAGE_BRUSH("Icons/ActorIcons/SoundActor_16x", Icon16x16));
