Added new icons for the SkyAtmosphere component.

diff --git a/Engine/Source/Editor/EditorStyle/Private/SlateEditorStyle.cpp b/Engine/Source/Editor/EditorStyle/Private/SlateEditorStyle.cpp
--- a/Engine/Source/Editor/EditorStyle/Private/SlateEditorStyle.cpp
+++ b/Engine/Source/Editor/EditorStyle/Private/SlateEditorStyle.cpp
@@ -6387,6 +6387,7 @@ void FSlateEditorStyle::FStyle::SetupClassIconsAndThumbnails()
 			TEXT("SceneComponent"),
 			TEXT("SkyLight"),
 			TEXT("SkyLightComponent"),
+			TEXT("SkyAtmosphere"),
 			TEXT("SkeletalMesh"),
 			TEXT("Skeleton"),
 			TEXT("SlateBrushAsset"),
@@ -6436,14 +6437,6 @@ void FSlateEditorStyle::FStyle::SetupClassIconsAndThumbnails()
 			Set( *FString::Printf(TEXT("ClassIcon.%s"), Type),		new IMAGE_BRUSH(FString::Printf(TEXT("Icons/AssetIcons/%s_%dx"), Type, 16), Icon16x16 ) );
 			Set( *FString::Printf(TEXT("ClassThumbnail.%s"), Type),	new IMAGE_BRUSH(FString::Printf(TEXT("Icons/AssetIcons/%s_%dx"), Type, 64), Icon64x64 ) );
 		}
-
-		{
-			// For now, the SkyAtmosphere component copy the style of the AtmosphericFog component.
-			const TCHAR* Type = TEXT("SkyAtmosphere");
-			const TCHAR* Icon = TEXT("AtmosphericFog");
-			Set(*FString::Printf(TEXT("ClassIcon.%s"), Type), new IMAGE_BRUSH(FString::Printf(TEXT("Icons/AssetIcons/%s_%dx"), Icon, 16), Icon16x16));
-			Set(*FString::Printf(TEXT("ClassThumbnail.%s"), Type), new IMAGE_BRUSH(FString::Printf(TEXT("Icons/AssetIcons/%s_%dx"), Icon, 64), Icon64x64));
-		}
 	}
 #endif
 }
diff --git a/Engine/Source/Runtime/Engine/Private/Components/SkyAtmosphereComponent.cpp b/Engine/Source/Runtime/Engine/Private/Components/SkyAtmosphereComponent.cpp
--- a/Engine/Source/Runtime/Engine/Private/Components/SkyAtmosphereComponent.cpp
+++ b/Engine/Source/Runtime/Engine/Private/Components/SkyAtmosphereComponent.cpp
@@ -257,7 +257,7 @@ ASkyAtmosphere::ASkyAtmosphere(const FObjectInitializer& ObjectInitializer)
 			FName ID_SkyAtmosphere;
 			FText NAME_SkyAtmosphere;
 			FConstructorStatics()
-				: SkyAtmosphereTextureObject(TEXT("/Engine/EditorResources/S_ExpoHeightFog"))
+				: SkyAtmosphereTextureObject(TEXT("/Engine/EditorResources/S_SkyAtmosphere"))
 				, ID_SkyAtmosphere(TEXT("Fog"))
 				, NAME_SkyAtmosphere(NSLOCTEXT("SpriteCategory", "Fog", "Fog"))
 			{
