Merge CL#9424546 fix viewport transform control spacing regression

diff --git a/Engine/Source/Editor/EditorStyle/Private/SlateEditorStyle.cpp b/Engine/Source/Editor/EditorStyle/Private/SlateEditorStyle.cpp
--- a/Engine/Source/Editor/EditorStyle/Private/SlateEditorStyle.cpp
+++ b/Engine/Source/Editor/EditorStyle/Private/SlateEditorStyle.cpp
@@ -2043,7 +2043,7 @@ void FSlateEditorStyle::FStyle::SetupGeneralStyles()
 		Set( "ToolBar.Block.Padding", FMargin( 2.0f, 2.0f, 4.0f, 4.0f ) );
 
 		Set( "ToolBar.Separator", new FSlateColorBrush( FLinearColor(FColor(47, 47, 47)) ) );
-		Set( "ToolBar.Separator.Padding", FMargin( 0.5f ) );
+		Set( "ToolBar.Separator.Padding", FMargin( 8.f, 0.f, 8.f, 0.f) );
 
 		Set( "ToolBar.Label", FTextBlockStyle(NormalText) .SetFont( DEFAULT_FONT( "Regular", 9 ) ) );
 		Set( "ToolBar.EditableText", FEditableTextBoxStyle(NormalEditableTextBoxStyle) .SetFont( DEFAULT_FONT( "Regular", 9 ) ) );
@@ -2664,7 +2664,7 @@ void FSlateEditorStyle::FStyle::SetupViewportStyles()
 		Set("ViewportMenu.SToolBarComboButtonBlock.ComboButton.Color", FLinearColor(0.f, 0.f, 0.f, 0.75f));
 
 		Set("ViewportMenu.Separator", new BOX_BRUSH("Old/Button", 8.0f / 32.0f, FLinearColor::Transparent));
-		Set("ViewportMenu.Separator.Padding", FMargin(100.0f));
+		Set("ViewportMenu.Separator.Padding", FMargin( 1.f, 0.f, 0.f, 0.f) );
 
 		Set("ViewportMenu.Label", FTextBlockStyle(NormalText)
 			.SetFont(DEFAULT_FONT("Bold", 9))
@@ -5334,7 +5334,7 @@ void FSlateEditorStyle::FStyle::SetupLevelEditorStyle()
 		Set( "EditorModesToolbar.Block.Padding", FMargin( 0 ) );
 
 		Set( "EditorModesToolbar.Separator", new BOX_BRUSH( "Old/Button", 4.0f / 32.0f ) );
-		Set( "EditorModesToolbar.Separator.Padding", FMargin( 0.5f ) );
+		Set( "EditorModesToolbar.Separator.Padding", FMargin( 8.f, 0.f, 8.f, 0.f) );
 
 		Set( "EditorModesToolbar.Label", FTextBlockStyle(NormalText) .SetFont( DEFAULT_FONT( "Regular", 7 ) ) );
 		Set( "EditorModesToolbar.EditableText", FEditableTextBoxStyle(NormalEditableTextBoxStyle) .SetFont( DEFAULT_FONT( "Regular", 9 ) ) );
diff --git a/Engine/Source/Runtime/Slate/Private/Framework/MultiBox/SToolBarSeparatorBlock.cpp b/Engine/Source/Runtime/Slate/Private/Framework/MultiBox/SToolBarSeparatorBlock.cpp
--- a/Engine/Source/Runtime/Slate/Private/Framework/MultiBox/SToolBarSeparatorBlock.cpp
+++ b/Engine/Source/Runtime/Slate/Private/Framework/MultiBox/SToolBarSeparatorBlock.cpp
@@ -55,7 +55,7 @@ void SToolBarSeparatorBlock::BuildMultiBlockWidget(const ISlateStyle* StyleSet,
 		SNew( SHorizontalBox )
 		+SHorizontalBox::Slot()
 		.AutoWidth()
-		.Padding( 8.0f, 0.0f, 8.0f, 0.0f )
+		.Padding( StyleSet->GetMargin( ISlateStyle::Join( StyleName, ".Separator.Padding" ) ) )
 		[
 			SNew(SSeparator)
 				.Orientation(Orient_Vertical)
diff --git a/Engine/Source/Runtime/SlateCore/Private/Styling/CoreStyle.cpp b/Engine/Source/Runtime/SlateCore/Private/Styling/CoreStyle.cpp
--- a/Engine/Source/Runtime/SlateCore/Private/Styling/CoreStyle.cpp
+++ b/Engine/Source/Runtime/SlateCore/Private/Styling/CoreStyle.cpp
@@ -848,7 +848,7 @@ TSharedRef<ISlateStyle> FCoreStyle::Create( const FName& InStyleSetName )
 		Style->Set( "ToolBar.Block.Padding", FMargin( 2.0f, 2.0f, 4.0f, 4.0f ) );
 
 		Style->Set( "ToolBar.Separator", new FSlateColorBrush( FLinearColor(FColor(48, 48, 48)) ) );
-		Style->Set( "ToolBar.Separator.Padding", FMargin( 0.5f ) );
+		Style->Set( "ToolBar.Separator.Padding", FMargin( 8.f, 0.f, 8.f, 0.f) );
 
 		Style->Set( "ToolBar.Label", FTextBlockStyle(NormalText) .SetFont( DEFAULT_FONT( "Regular", 9 ) ) );
 		Style->Set( "ToolBar.EditableText", FEditableTextBoxStyle(NormalEditableTextBoxStyle) .SetFont( DEFAULT_FONT( "Regular", 9 ) ) );
