Merging CL#9610140 from Dev-Editor

Fixes Sequencer ToolBar Spacing and Height

diff --git a/Engine/Plugins/MovieScene/LevelSequenceEditor/Source/LevelSequenceEditor/Private/Misc/LevelSequencePlaybackContext.cpp b/Engine/Plugins/MovieScene/LevelSequenceEditor/Source/LevelSequenceEditor/Private/Misc/LevelSequencePlaybackContext.cpp
--- a/Engine/Plugins/MovieScene/LevelSequenceEditor/Source/LevelSequenceEditor/Private/Misc/LevelSequencePlaybackContext.cpp
+++ b/Engine/Plugins/MovieScene/LevelSequenceEditor/Source/LevelSequenceEditor/Private/Misc/LevelSequencePlaybackContext.cpp
@@ -230,9 +230,11 @@ void SLevelSequenceContextPicker::Construct(const FArguments& InArgs)
 	check(OnSetValueEvent.IsBound());
 
 	ChildSlot
+	.Padding(0.0f)
 	[
 		SNew(SBorder)
 		.BorderImage(this, &SLevelSequenceContextPicker::GetBorderBrush)
+		.Padding(0.0f)
 		[
 			SNew(SComboButton)
 			.ContentPadding(0)
diff --git a/Engine/Source/Editor/EditorStyle/Private/SlateEditorStyle.cpp b/Engine/Source/Editor/EditorStyle/Private/SlateEditorStyle.cpp
--- a/Engine/Source/Editor/EditorStyle/Private/SlateEditorStyle.cpp
+++ b/Engine/Source/Editor/EditorStyle/Private/SlateEditorStyle.cpp
@@ -2623,7 +2623,43 @@ void FSlateEditorStyle::FStyle::SetupSequencerStyles()
 	Set("EMovieSceneBlendType::Absolute", new IMAGE_BRUSH("Sequencer/EMovieSceneBlendType_Absolute", FVector2D(32, 16)));
 	Set("EMovieSceneBlendType::Relative", new IMAGE_BRUSH("Sequencer/EMovieSceneBlendType_Relative", FVector2D(32, 16)));
 	Set("EMovieSceneBlendType::Additive", new IMAGE_BRUSH("Sequencer/EMovieSceneBlendType_Additive", FVector2D(32, 16)));
-		}
+
+
+	// Sequencer & Curve Editor Toolbar Style ( Grabs core default, copies brushes, changes separator and block spacing )
+	Set( "Sequencer.ToolBar.Background", 		new CORE_BOX_BRUSH( "Common/GroupBorder", FMargin(4.0f/16.0f) ) );
+	Set( "Sequencer.ToolBar.Icon", 				new CORE_IMAGE_BRUSH( "Icons/icon_tab_toolbar_16px", Icon16x16 ) );
+	Set( "Sequencer.ToolBar.Expand", 			new CORE_IMAGE_BRUSH( "Icons/toolbar_expand_16x", Icon16x16) );
+	Set( "Sequencer.ToolBar.SubMenuIndicator", 	new CORE_IMAGE_BRUSH( "Common/SubmenuArrow", Icon8x8 ) );
+	Set( "Sequencer.ToolBar.SToolBarComboButtonBlock.Padding", 		 FCoreStyle::Get().GetMargin("ToolBar.SToolBarComboButtonBlock.Padding"));
+	Set( "Sequencer.ToolBar.SToolBarButtonBlock.Padding", 			 FCoreStyle::Get().GetMargin("ToolBar.SToolBarButtonBlock.Padding"));
+	Set( "Sequencer.ToolBar.SToolBarCheckComboButtonBlock.Padding",  FCoreStyle::Get().GetMargin("ToolBar.SToolBarCheckComboButtonBlock.Padding"));
+	Set( "Sequencer.ToolBar.SToolBarButtonBlock.CheckBox.Padding", 	 FCoreStyle::Get().GetMargin("ToolBar.SToolBarButtonBlock.CheckBox.Padding"));
+	Set( "Sequencer.ToolBar.SToolBarComboButtonBlock.ComboButton.Color", GetColor("ToolBar.SToolBarComboButtonBlock.ComboButton.Color"));
+
+	// Used only for WidgetBlocks
+	Set( "Sequencer.ToolBar.Block.IndentedPadding", 	FCoreStyle::Get().GetMargin("Sequencer.ToolBar.Block.IndentedPadding") );
+	Set( "Sequencer.ToolBar.Block.Padding", 			FMargin(0.0f, 4.0f));
+
+	Set( "Sequencer.ToolBar.Separator", 				new FSlateColorBrush( FLinearColor(FColor(48, 48, 48)) ) );
+	Set( "Sequencer.ToolBar.Separator.Padding", 		FMargin( 1.f, 1.f ) );
+
+	Set( "Sequencer.ToolBar.Label", 					FCoreStyle::Get().GetWidgetStyle<FTextBlockStyle>("ToolBar.Label"));
+	Set( "Sequencer.ToolBar.EditableText",  			FCoreStyle::Get().GetWidgetStyle<FEditableTextBoxStyle>("ToolBar.EditableText"));
+	Set( "Sequencer.ToolBar.Keybinding", 				FCoreStyle::Get().GetWidgetStyle<FTextBlockStyle>("ToolBar.Keybinding"));
+	Set( "Sequencer.ToolBar.Heading", 					FCoreStyle::Get().GetWidgetStyle<FTextBlockStyle>("ToolBar.Heading"));
+	Set( "Sequencer.ToolBar.CheckBox", 					FCoreStyle::Get().GetWidgetStyle<FCheckBoxStyle>("ToolBar.CheckBox"));
+	Set( "Sequencer.Toolbar.Check", 					FCoreStyle::Get().GetWidgetStyle<FCheckBoxStyle>("ToolBar.Check"));
+	Set( "Sequencer.ToolBar.RadioButton", 				FCoreStyle::Get().GetWidgetStyle<FCheckBoxStyle>("ToolBar.RadioButton"));
+	Set( "Sequencer.ToolBar.ToggleButton", 				FCoreStyle::Get().GetWidgetStyle<FCheckBoxStyle>("ToolBar.ToggleButton"));
+	Set( "Sequencer.ToolBar.Button", 					FCoreStyle::Get().GetWidgetStyle<FButtonStyle>("ToolBar.Button"));
+	Set( "ToolBar.Button.Normal", 						new FSlateNoResource() );
+	Set( "ToolBar.Button.Pressed", 						new CORE_BOX_BRUSH( "Common/RoundedSelection_16x", 4.0f/16.0f, SelectionColor_Pressed ) );
+	Set( "ToolBar.Button.Hovered", 						new CORE_BOX_BRUSH( "Common/RoundedSelection_16x", 4.0f/16.0f, SelectionColor ) );
+	Set( "ToolBar.Button.Checked", 						new CORE_BOX_BRUSH( "Common/RoundedSelection_16x",  4.0f/16.0f, SelectionColor_Pressed ) );
+	Set( "ToolBar.Button.Checked_Hovered", 				new CORE_BOX_BRUSH( "Common/RoundedSelection_16x",  4.0f/16.0f, SelectionColor_Pressed ) );
+	Set( "ToolBar.Button.Checked_Pressed", 				new CORE_BOX_BRUSH( "Common/RoundedSelection_16x",  4.0f/16.0f, SelectionColor ) );
+
+	}
 
 
 	// Sequence recorder standalone UI
diff --git a/Engine/Source/Editor/Sequencer/Private/SSequencer.cpp b/Engine/Source/Editor/Sequencer/Private/SSequencer.cpp
--- a/Engine/Source/Editor/Sequencer/Private/SSequencer.cpp
+++ b/Engine/Source/Editor/Sequencer/Private/SSequencer.cpp
@@ -130,6 +130,7 @@ public:
 	TSharedRef<SWidget> MakeToolbar(TSharedRef<SCurveEditorPanel> InEditorPanel)
 	{
 		FToolBarBuilder ToolBarBuilder(InEditorPanel->GetCommands(), FMultiBoxCustomization::None, InEditorPanel->GetToolbarExtender(), EOrientation::Orient_Horizontal, true);
+		ToolBarBuilder.SetStyle(&FEditorStyle::Get(), "Sequencer.ToolBar");
 		ToolBarBuilder.BeginSection("Asset");
 		ToolBarBuilder.EndSection();
 		// We just use all of the extenders as our toolbar, we don't have a need to create a separate toolbar.
@@ -567,7 +568,7 @@ void SSequencer::Construct(const FArguments& InArgs, TSharedRef<FSequencer> InSe
 							.FillEmptySpace(true)
 							[
 								SNew(SBorder)
-								.Padding(FMargin(3))
+								.Padding(FMargin(3.0f, 0.0f))
 								.BorderImage(FEditorStyle::GetBrush("ToolPanel.GroupBorder"))
 								[
 									SNew(SHorizontalBox)
@@ -1288,6 +1289,7 @@ TSharedRef<SWidget> SSequencer::MakeToolBar()
 	}
 
 	FToolBarBuilder ToolBarBuilder( SequencerPtr.Pin()->GetCommandBindings(), FMultiBoxCustomization::None, Extender, Orient_Horizontal, true);
+	ToolBarBuilder.SetStyle(&FEditorStyle::Get(), "Sequencer.ToolBar");
 
 	ToolBarBuilder.BeginSection("Base Commands");
 	{
diff --git a/Engine/Source/Editor/Sequencer/Private/SSequencerPlayRateCombo.cpp b/Engine/Source/Editor/Sequencer/Private/SSequencerPlayRateCombo.cpp
--- a/Engine/Source/Editor/Sequencer/Private/SSequencerPlayRateCombo.cpp
+++ b/Engine/Source/Editor/Sequencer/Private/SSequencerPlayRateCombo.cpp
@@ -33,7 +33,7 @@ void SSequencerPlayRateCombo::Construct(const FArguments& InArgs, TWeakPtr<FSequ
 	.VAlign(VAlign_Fill)
 	[
 		SNew(SComboButton)
-		.ContentPadding(FMargin(2.f, 0.f))
+		.ContentPadding(FMargin(2.f, 1.0f))
 		.VAlign(VAlign_Fill)
 		.ButtonStyle( InArgs._StyleSet, BlockStyle )
 		.ForegroundColor( InArgs._StyleSet->GetSlateColor( ColorStyle ) )
@@ -91,8 +91,6 @@ void SSequencerPlayRateCombo::Construct(const FArguments& InArgs, TWeakPtr<FSequ
 			]
 		]
 	];
-
-	ChildSlot.Padding(InArgs._StyleSet->GetMargin(ISlateStyle::Join( InArgs._StyleName, ".SToolBarComboButtonBlock.Padding" )));
 }
 
 EVisibility SSequencerPlayRateCombo::GetFrameLockedVisibility() const
