'Add new' button in animation interfaces now allows the creation of animation layers

diff --git a/Engine/Source/Editor/EditorStyle/Private/SlateEditorStyle.cpp b/Engine/Source/Editor/EditorStyle/Private/SlateEditorStyle.cpp
--- a/Engine/Source/Editor/EditorStyle/Private/SlateEditorStyle.cpp
+++ b/Engine/Source/Editor/EditorStyle/Private/SlateEditorStyle.cpp
@@ -5753,6 +5753,8 @@ void FSlateEditorStyle::FStyle::SetupPersonaStyle()
 		Set( "BlueprintEditor.ManageInterfaces.Small", new IMAGE_BRUSH( "Icons/icon_BlueprintEditor_Interfaces_40x", Icon20x20) );	
 		Set( "BlueprintEditor.AddNewDelegate.Small", new IMAGE_BRUSH( "Icons/icon_Blueprint_AddDelegate_40x", Icon20x20) );
 		Set( "BlueprintEditor.AddNewDelegate", new IMAGE_BRUSH( "Icons/icon_Blueprint_AddDelegate_40x", Icon40x40) );
+		Set( "BlueprintEditor.AddNewAnimationLayer.Small", new IMAGE_BRUSH( "Icons/icon_Blueprint_AddFunction_40px", Icon20x20) );
+		Set( "BlueprintEditor.AddNewAnimationLayer", new IMAGE_BRUSH( "Icons/icon_Blueprint_AddFunction_40px", Icon40x40) );
 
 		Set( "Kismet.Status.Unknown.Small", new IMAGE_BRUSH( "Old/Kismet2/CompileStatus_Unknown_Small", Icon16x16 ) );
 		Set( "Kismet.Status.Error.Small", new IMAGE_BRUSH( "Old/Kismet2/CompileStatus_Broken_Small", Icon16x16 ) );
diff --git a/Engine/Source/Editor/Kismet/Private/SMyBlueprint.cpp b/Engine/Source/Editor/Kismet/Private/SMyBlueprint.cpp
--- a/Engine/Source/Editor/Kismet/Private/SMyBlueprint.cpp
+++ b/Engine/Source/Editor/Kismet/Private/SMyBlueprint.cpp
@@ -2150,6 +2150,10 @@ void SMyBlueprint::BuildAddNewMenu(FMenuBuilder& MenuBuilder)
 		{
 			MenuBuilder.AddMenuEntry(FBlueprintEditorCommands::Get().AddNewDelegate);
 		}
+		if (CurrentBlueprint->SupportsAnimLayers())
+		{
+			MenuBuilder.AddMenuEntry(FBlueprintEditorCommands::Get().AddNewAnimationLayer);
+		}
 	}
 	MenuBuilder.EndSection();
 }
