- Adding small version of map check image brush

diff --git a/Engine/Source/Editor/EditorStyle/Private/SlateEditorStyle.cpp b/Engine/Source/Editor/EditorStyle/Private/SlateEditorStyle.cpp
--- a/Engine/Source/Editor/EditorStyle/Private/SlateEditorStyle.cpp
+++ b/Engine/Source/Editor/EditorStyle/Private/SlateEditorStyle.cpp
@@ -4925,6 +4925,7 @@ void FSlateEditorStyle::FStyle::SetupLevelEditorStyle()
 		Set( "LevelEditor.Build", new IMAGE_BRUSH( "Icons/icon_build_40x", Icon40x40 ) );
 		Set( "LevelEditor.Build.Small", new IMAGE_BRUSH( "Icons/icon_build_40x", Icon20x20 ) );
 		Set( "LevelEditor.MapCheck", new IMAGE_BRUSH( "Icons/icon_MapCheck_40x", Icon40x40 ) );
+		Set( "LevelEditor.MapCheck.Small", new IMAGE_BRUSH( "Icons/icon_MapCheck_40x", Icon20x20 ) );
 
 		Set( "LevelEditor.Recompile", new IMAGE_BRUSH( "Icons/icon_compile_40x", Icon40x40 ) );
 		Set( "LevelEditor.Recompile.Small", new IMAGE_BRUSH( "Icons/icon_compile_40x", Icon20x20 ) );
