ROBOMERGE conflict - manual submission

#rnx

diff --git a/Engine/Source/Editor/EditorStyle/Private/SlateEditorStyle.cpp b/Engine/Source/Editor/EditorStyle/Private/SlateEditorStyle.cpp
--- a/Engine/Source/Editor/EditorStyle/Private/SlateEditorStyle.cpp
+++ b/Engine/Source/Editor/EditorStyle/Private/SlateEditorStyle.cpp
@@ -4969,10 +4969,6 @@ void FSlateEditorStyle::FStyle::SetupLevelEditorStyle()
 		Set("LevelEditor.PreviewMode.iOS.Enabled.Small", new IMAGE_BRUSH("Icons/icon_PreviewMode_iOS_Enabled_40x", Icon20x20));
 		Set("LevelEditor.PreviewMode.iOS.Disabled", new IMAGE_BRUSH("Icons/icon_PreviewMode_iOS_Disabled_40x", Icon40x40));
 		Set("LevelEditor.PreviewMode.iOS.Disabled.Small", new IMAGE_BRUSH("Icons/icon_PreviewMode_iOS_Disabled_40x", Icon20x20));
-		Set("LevelEditor.PreviewMode.HTML5.Enabled", new IMAGE_BRUSH("Icons/icon_PreviewMode_HTML5_Enabled_40x", Icon40x40));
-		Set("LevelEditor.PreviewMode.HTML5.Enabled.Small", new IMAGE_BRUSH("Icons/icon_PreviewMode_HTML5_Enabled_40x", Icon20x20));
-		Set("LevelEditor.PreviewMode.HTML5.Disabled", new IMAGE_BRUSH("Icons/icon_PreviewMode_HTML5_Disabled_40x", Icon40x40));
-		Set("LevelEditor.PreviewMode.HTML5.Disabled.Small", new IMAGE_BRUSH("Icons/icon_PreviewMode_HTML5_Disabled_40x", Icon20x20));
 
 		Set("LevelEditor.ViewOptions", new IMAGE_BRUSH("Icons/icon_view_40x", Icon40x40));
 		Set( "LevelEditor.ViewOptions.Small", new IMAGE_BRUSH( "Icons/icon_view_40x", Icon20x20 ) );
@@ -8323,4 +8319,4 @@ void FSlateEditorStyle::FStyle::SetupMyBlueprintStyles()
 #undef DEFAULT_FONT
 #undef ICON_FONT
 
-END_SLATE_FUNCTION_BUILD_OPTIMIZATION
\ No newline at end of file
+END_SLATE_FUNCTION_BUILD_OPTIMIZATION
diff --git a/Engine/Source/Editor/LevelEditor/Private/LevelEditor.cpp b/Engine/Source/Editor/LevelEditor/Private/LevelEditor.cpp
--- a/Engine/Source/Editor/LevelEditor/Private/LevelEditor.cpp
+++ b/Engine/Source/Editor/LevelEditor/Private/LevelEditor.cpp
@@ -1776,12 +1776,6 @@ void FLevelEditorModule::BindGlobalLevelEditorCommands()
 		FCanExecuteAction(),
 		FIsActionChecked::CreateStatic(&FLevelEditorActionCallbacks::IsPreviewPlatformChecked, FPreviewPlatformInfo(ERHIFeatureLevel::SM5, NAME_None)));
 
-	ActionList.MapAction(
-		Commands.PreviewPlatformOverride_HTML5,
-		FExecuteAction::CreateStatic(&FLevelEditorActionCallbacks::SetPreviewPlatform, FPreviewPlatformInfo(ERHIFeatureLevel::ES2, NAME_None, true)),
-		FCanExecuteAction(),
-		FIsActionChecked::CreateStatic(&FLevelEditorActionCallbacks::IsPreviewPlatformChecked, FPreviewPlatformInfo(ERHIFeatureLevel::ES2, NAME_None)));
-
 	ActionList.MapAction(
 		Commands.PreviewPlatformOverride_AndroidGLES2,
 		FExecuteAction::CreateStatic(&FLevelEditorActionCallbacks::SetPreviewPlatform, FPreviewPlatformInfo(ERHIFeatureLevel::ES2, LegacyShaderPlatformToShaderFormat(SP_OPENGL_ES2_ANDROID), true)),
diff --git a/Engine/Source/Editor/LevelEditor/Private/LevelEditorActions.cpp b/Engine/Source/Editor/LevelEditor/Private/LevelEditorActions.cpp
--- a/Engine/Source/Editor/LevelEditor/Private/LevelEditorActions.cpp
+++ b/Engine/Source/Editor/LevelEditor/Private/LevelEditorActions.cpp
@@ -3371,7 +3371,6 @@ void FLevelEditorCommands::RegisterCommands()
 	UI_COMMAND(ToggleFeatureLevelPreview, "Preview Mode Toggle", "Toggles the Preview Mode on or off for the currently selected Preview target", EUserInterfaceActionType::ToggleButton, FInputChord());
 
 	UI_COMMAND(PreviewPlatformOverride_SM5, "Shader Model 5", "DirectX 11, OpenGL 4.3+, PS4, XB1", EUserInterfaceActionType::Check, FInputChord());
-	UI_COMMAND(PreviewPlatformOverride_HTML5, "HTML5", "HTML5 preview.", EUserInterfaceActionType::Check, FInputChord());
 	UI_COMMAND(PreviewPlatformOverride_AndroidGLES2, "Android ES2", "Mobile preview using Android's quality settings.", EUserInterfaceActionType::Check, FInputChord());
 	UI_COMMAND(PreviewPlatformOverride_AndroidGLES31, "Android ES 3.1", "Mobile preview using Android ES3.1 quality settings.", EUserInterfaceActionType::Check, FInputChord());
 	UI_COMMAND(PreviewPlatformOverride_AndroidVulkanES31, "Android Vulkan", "Mobile preview using Android Vulkan quality settings.", EUserInterfaceActionType::Check, FInputChord());
diff --git a/Engine/Source/Editor/LevelEditor/Private/LevelEditorToolBar.cpp b/Engine/Source/Editor/LevelEditor/Private/LevelEditorToolBar.cpp
--- a/Engine/Source/Editor/LevelEditor/Private/LevelEditorToolBar.cpp
+++ b/Engine/Source/Editor/LevelEditor/Private/LevelEditorToolBar.cpp
@@ -1390,7 +1390,7 @@ void FLevelEditorToolBar::RegisterLevelEditorToolBar( const TSharedRef<FUIComman
 					case SP_METAL_MACES2:
                 	case SP_OPENGL_ES2_WEBGL:
 					{
-						return FSlateIcon(FEditorStyle::GetStyleSetName(), GEditor->IsFeatureLevelPreviewActive() ? "LevelEditor.PreviewMode.HTML5.Enabled" : "LevelEditor.PreviewMode.HTML5.Disabled");
+						return FSlateIcon(FEditorStyle::GetStyleSetName(), GEditor->IsFeatureLevelPreviewActive() ? "LevelEditor.PreviewMode.AndroidES2.Enabled" : "LevelEditor.PreviewMode.AndroidES2.Disabled");
 					}
 				}
 				switch (GEditor->PreviewPlatform.PreviewFeatureLevel)
@@ -1903,9 +1903,6 @@ static void MakeShaderModelPreviewMenu( UToolMenu* InMenu )
 		Section.AddMenuEntry(FLevelEditorCommands::Get().PreviewPlatformOverride_IOSMetalES31);
 	}
 
-	// HTML5
-	Section.AddMenuEntry(FLevelEditorCommands::Get().PreviewPlatformOverride_HTML5);
-
 #undef LOCTEXT_NAMESPACE
 }
 
diff --git a/Engine/Source/Editor/LevelEditor/Public/LevelEditorActions.h b/Engine/Source/Editor/LevelEditor/Public/LevelEditorActions.h
--- a/Engine/Source/Editor/LevelEditor/Public/LevelEditorActions.h
+++ b/Engine/Source/Editor/LevelEditor/Public/LevelEditorActions.h
@@ -591,7 +591,6 @@ public:
 	TSharedPtr< FUICommandInfo > ToggleFeatureLevelPreview;
 
 	TSharedPtr< FUICommandInfo > PreviewPlatformOverride_SM5;
-	TSharedPtr< FUICommandInfo > PreviewPlatformOverride_HTML5;
 	TSharedPtr< FUICommandInfo > PreviewPlatformOverride_AndroidGLES2;
 	TSharedPtr< FUICommandInfo > PreviewPlatformOverride_AndroidGLES31;
 	TSharedPtr< FUICommandInfo > PreviewPlatformOverride_AndroidVulkanES31;
