Fix to static analysis warning on stack size usage in FSlateEditorStyle initialization.

diff --git a/Engine/Source/Editor/EditorStyle/Private/SlateEditorStyle.cpp b/Engine/Source/Editor/EditorStyle/Private/SlateEditorStyle.cpp
--- a/Engine/Source/Editor/EditorStyle/Private/SlateEditorStyle.cpp
+++ b/Engine/Source/Editor/EditorStyle/Private/SlateEditorStyle.cpp
@@ -177,6 +177,12 @@ void FSlateEditorStyle::FStyle::Initialize()
 	SetCoreContentRoot(FPaths::EngineContentDir() / TEXT("Slate"));
 
 	SetupGeneralStyles();
+	SetupLevelGeneralStyles();
+	SetupWorldBrowserStyles();
+	SetupSequencerStyles();
+	SetupViewportStyles();
+	SetupNotificationBarStyles();
+	SetupMenuBarStyles();
 	SetupGeneralIcons();
 	SetupWindowStyles();
 	SetupPropertyEditorStyles();
@@ -1374,64 +1380,7 @@ void FSlateEditorStyle::FStyle::SetupGeneralStyles()
 		Set( "LayerBrowserButton.LabelFont", DEFAULT_FONT( "Regular", 8 ) );
 	}
 
-	// Levels General
-	{
-		Set( "Level.VisibleIcon16x", new IMAGE_BRUSH( "Icons/icon_levels_visible_16px", Icon16x16 ) );
-		Set( "Level.VisibleHighlightIcon16x", new IMAGE_BRUSH( "Icons/icon_levels_visible_hi_16px", Icon16x16 ) );
-		Set( "Level.NotVisibleIcon16x", new IMAGE_BRUSH( "Icons/icon_levels_invisible_16px", Icon16x16 ) );
-		Set( "Level.NotVisibleHighlightIcon16x", new IMAGE_BRUSH( "Icons/icon_levels_invisible_hi_16px", Icon16x16 ) );
-		Set( "Level.LightingScenarioIcon16x", new IMAGE_BRUSH( "Icons/icon_levels_LightingScenario_16px", Icon16x16 ) );
-		Set( "Level.LightingScenarioNotIcon16x", new IMAGE_BRUSH( "Icons/icon_levels_LightingScenarioNot_16px", Icon16x16 ) );
-		Set( "Level.LockedIcon16x", new IMAGE_BRUSH( "Icons/icon_locked_16px", Icon16x16 ) );
-		Set( "Level.LockedHighlightIcon16x", new IMAGE_BRUSH( "Icons/icon_locked_highlight_16px", Icon16x16 ) );
-		Set( "Level.UnlockedIcon16x", new IMAGE_BRUSH( "Icons/icon_levels_unlocked_16px", Icon16x16 ) );
-		Set( "Level.UnlockedHighlightIcon16x", new IMAGE_BRUSH( "Icons/icon_levels_unlocked_hi_16px", Icon16x16 ) );
-		Set( "Level.ReadOnlyLockedIcon16x", new IMAGE_BRUSH( "Icons/icon_levels_LockedReadOnly_16px", Icon16x16 ) );
-		Set( "Level.ReadOnlyLockedHighlightIcon16x", new IMAGE_BRUSH( "Icons/icon_levels_LockedReadOnly_hi_16px", Icon16x16 ) );
-		Set( "Level.SaveIcon16x", new IMAGE_BRUSH( "Icons/icon_levels_Save_16px", Icon16x16 ) );
-		Set( "Level.SaveHighlightIcon16x", new IMAGE_BRUSH( "Icons/icon_levels_Save_hi_16px", Icon16x16 ) );
-		Set( "Level.SaveModifiedIcon16x", new IMAGE_BRUSH( "Icons/icon_levels_SaveModified_16px", Icon16x16 ) );
-		Set( "Level.SaveModifiedHighlightIcon16x", new IMAGE_BRUSH( "Icons/icon_levels_SaveModified_hi_16px", Icon16x16 ) );
-		Set( "Level.SaveDisabledIcon16x", new IMAGE_BRUSH( "Icons/icon_levels_SaveDisabled_16px", Icon16x16 ) );
-		Set( "Level.SaveDisabledHighlightIcon16x", new IMAGE_BRUSH( "Icons/icon_levels_SaveDisabled_hi_16px", Icon16x16 ) );
-		Set( "Level.ScriptIcon16x", new IMAGE_BRUSH( "Icons/icon_levels_Blueprint_16px", Icon16x16 ) );
-		Set( "Level.ScriptHighlightIcon16x", new IMAGE_BRUSH( "Icons/icon_levels_Blueprint_hi_16px", Icon16x16 ) );
-		Set( "Level.EmptyIcon16x", new IMAGE_BRUSH( "Icons/Empty_16x", Icon16x16 ) );
-		Set( "Level.ColorIcon40x", new IMAGE_BRUSH( "Icons/icon_levels_back_16px", Icon16x16 ) );
-	}
-
-	// World Browser
-	{
-		Set( "WorldBrowser.AddLayer", new IMAGE_BRUSH( "Icons/icon_levels_addlayer_16x", Icon16x16 ) );
-		Set( "WorldBrowser.SimulationViewPositon", new IMAGE_BRUSH( "Icons/icon_levels_simulationviewpos_16x", Icon16x16 ) );
-		Set( "WorldBrowser.MouseLocation", new IMAGE_BRUSH( "Icons/icon_levels_mouselocation_16x", Icon16x16 ) );
-		Set( "WorldBrowser.MarqueeRectSize", new IMAGE_BRUSH( "Icons/icon_levels_marqueerectsize_16x", Icon16x16 ) );
-		Set( "WorldBrowser.WorldSize", new IMAGE_BRUSH( "Icons/icon_levels_worldsize_16x", Icon16x16 ) );
-		Set( "WorldBrowser.WorldOrigin", new IMAGE_BRUSH( "Icons/icon_levels_worldorigin_16x", Icon16x16 ) );
-		Set( "WorldBrowser.DirectionXPositive", new IMAGE_BRUSH( "Icons/icon_PanRight", Icon16x16 ) );
-		Set( "WorldBrowser.DirectionXNegative", new IMAGE_BRUSH( "Icons/icon_PanLeft", Icon16x16 ) );
-		Set( "WorldBrowser.DirectionYPositive", new IMAGE_BRUSH( "Icons/icon_PanUp", Icon16x16 ) );
-		Set( "WorldBrowser.DirectionYNegative", new IMAGE_BRUSH( "Icons/icon_PanDown", Icon16x16 ) );
-		Set( "WorldBrowser.LevelStreamingAlwaysLoaded", new FSlateNoResource() );
-		Set( "WorldBrowser.LevelStreamingBlueprint", new IMAGE_BRUSH( "Icons/icon_levels_blueprinttype_7x16", Icon7x16 ) );
-		Set( "WorldBrowser.LevelsMenuBrush", new IMAGE_BRUSH( "Icons/icon_levels_levelsmenu_40x", Icon25x25 ) );
-		Set( "WorldBrowser.HierarchyButtonBrush", new IMAGE_BRUSH( "Icons/icon_levels_hierarchybutton_16x", Icon16x16 ) );
-		Set( "WorldBrowser.DetailsButtonBrush", new IMAGE_BRUSH( "Icons/icon_levels_detailsbutton_40x", Icon16x16 ) );
-		Set( "WorldBrowser.CompositionButtonBrush", new IMAGE_BRUSH( "Icons/icon_levels_compositionbutton_16x", Icon16x16 ) );
-		
-		Set( "WorldBrowser.FolderClosed", new IMAGE_BRUSH( "Icons/FolderClosed", Icon16x16 ) );
-		Set( "WorldBrowser.FolderOpen", new IMAGE_BRUSH( "Icons/FolderOpen", Icon16x16 ) );
-		Set( "WorldBrowser.NewFolderIcon", new IMAGE_BRUSH( "Icons/icon_AddFolder_16x", Icon16x16 ) );
-
-		Set( "WorldBrowser.StatusBarText", FTextBlockStyle(NormalText)
-				.SetFont( DEFAULT_FONT( "BoldCondensed", 12 ) )
-				.SetColorAndOpacity( FLinearColor(0.9, 0.9f, 0.9f, 0.5f) )
-				.SetShadowOffset( FVector2D::ZeroVector )
-			);
-
-		Set( "WorldBrowser.LabelFont", DEFAULT_FONT( "Regular", 9 ) );
-		Set( "WorldBrowser.LabelFontBold", DEFAULT_FONT( "Bold", 10 ) );
-	}
+	
 
 	// Scene Outliner
 	{
@@ -1775,292 +1724,6 @@ void FSlateEditorStyle::FStyle::SetupGeneralStyles()
 		Set( "BuildAndSubmit.SmallFont", DEFAULT_FONT( "Regular", 7 ) );
 	}
 
-	// Sequencer
-	if (IncludeEditorSpecificStyles())
-	{
-		Set( "Sequencer.IconKeyAuto", new IMAGE_BRUSH( "Sequencer/IconKeyAuto", Icon12x12 ) );
-		Set( "Sequencer.IconKeyBreak", new IMAGE_BRUSH( "Sequencer/IconKeyBreak", Icon12x12 ) );
-		Set( "Sequencer.IconKeyConstant", new IMAGE_BRUSH( "Sequencer/IconKeyConstant", Icon12x12 ) );
-		Set( "Sequencer.IconKeyLinear", new IMAGE_BRUSH( "Sequencer/IconKeyLinear", Icon12x12 ) );
-		Set( "Sequencer.IconKeyUser", new IMAGE_BRUSH( "Sequencer/IconKeyUser", Icon12x12 ) );
-
-		Set( "Sequencer.KeyCircle", new IMAGE_BRUSH( "Sequencer/KeyCircle", Icon12x12 ) );
-		Set( "Sequencer.KeyDiamond", new IMAGE_BRUSH( "Sequencer/KeyDiamond", Icon12x12 ) );
-		Set( "Sequencer.KeyDiamondBorder", new IMAGE_BRUSH( "Sequencer/KeyDiamondBorder", Icon12x12 ) );
-		Set( "Sequencer.KeySquare", new IMAGE_BRUSH( "Sequencer/KeySquare", Icon12x12 ) );
-		Set( "Sequencer.KeyTriangle", new IMAGE_BRUSH( "Sequencer/KeyTriangle", Icon12x12 ) );
-		Set( "Sequencer.KeyLeft", new IMAGE_BRUSH( "Sequencer/KeyLeft", Icon12x12 ) );
-		Set( "Sequencer.KeyRight", new IMAGE_BRUSH( "Sequencer/KeyRight", Icon12x12 ) );
-		Set( "Sequencer.PartialKey", new IMAGE_BRUSH( "Sequencer/PartialKey", FVector2D(11.f, 11.f) ) );
-		Set( "Sequencer.Star", new IMAGE_BRUSH( "Sequencer/Star", Icon12x12 ) );
-		Set( "Sequencer.Empty", new IMAGE_BRUSH( "Sequencer/Empty", Icon12x12 ) );
-		Set( "Sequencer.TangentHandle", new IMAGE_BRUSH( "Sequencer/TangentHandle", FVector2D(7, 7) ) );
-		Set( "Sequencer.GenericDivider", new IMAGE_BRUSH( "Sequencer/GenericDivider", FVector2D(2.f, 2.f), FLinearColor::White, ESlateBrushTileType::Vertical ) );
-
-		Set("Sequencer.Timeline.ScrubHandleDown", new BOX_BRUSH("Sequencer/ScrubHandleDown", FMargin(6.f / 13.f, 5 / 12.f, 6 / 13.f, 8 / 12.f)));
-		Set("Sequencer.Timeline.ScrubHandleUp", new BOX_BRUSH("Sequencer/ScrubHandleUp", FMargin(6.f / 13.f, 8 / 12.f, 6 / 13.f, 5 / 12.f)));
-		Set( "Sequencer.Timeline.ScrubFill", new BOX_BRUSH( "Sequencer/ScrubFill", FMargin( 2.f/4.f, 0.f ) ) );
-		Set( "Sequencer.Timeline.FrameBlockScrubHandleDown", new BOX_BRUSH( "Sequencer/ScrubHandleDown", FMargin( 6.f/13.f, 5/12.f, 6/13.f, 8/12.f ) ) );
-		Set( "Sequencer.Timeline.FrameBlockScrubHandleUp", new BOX_BRUSH( "Sequencer/ScrubHandleUp", FMargin( 6.f/13.f, 8/12.f, 6/13.f, 5/12.f ) ) );
-		Set( "Sequencer.Timeline.VanillaScrubHandleDown", new BOX_BRUSH( "Sequencer/ScrubHandleDown_Clamped", FMargin( 6.f/13.f, 3.f/12.f, 6.f/13.f, 7.f/12.f ) ) );
-		Set( "Sequencer.Timeline.VanillaScrubHandleUp", new BOX_BRUSH( "Sequencer/ScrubHandleUp_Clamped", FMargin( 6.f/13.f, 8/12.f, 6/13.f, 5/12.f ) ) );
-		Set( "Sequencer.Timeline.ScrubHandleWhole", new BOX_BRUSH( "Sequencer/ScrubHandleWhole", FMargin( 6.f/13.f, 10/24.f, 6/13.f, 10/24.f  ) ) );
-		Set( "Sequencer.Timeline.RangeHandleLeft", new BOX_BRUSH( "Sequencer/GenericGripLeft", FMargin(5.f/16.f) ) );
-		Set( "Sequencer.Timeline.RangeHandleRight", new BOX_BRUSH( "Sequencer/GenericGripRight", FMargin(5.f/16.f) ) );
-		Set( "Sequencer.Timeline.RangeHandle", new BOX_BRUSH( "Sequencer/GenericSectionBackground", FMargin(5.f/16.f) ) );
-		Set( "Sequencer.Timeline.NotifyAlignmentMarker", new IMAGE_BRUSH( "Sequencer/NotifyAlignmentMarker", FVector2D(10,19) ) );
-		Set( "Sequencer.Timeline.PlayRange_Top_L", new BOX_BRUSH( "Sequencer/PlayRange_Top_L", FMargin(1.f, 0.5f, 0.f, 0.5f) ) );
-		Set( "Sequencer.Timeline.PlayRange_Top_R", new BOX_BRUSH( "Sequencer/PlayRange_Top_R", FMargin(0.f, 0.5f, 1.f, 0.5f) ) );
-		Set( "Sequencer.Timeline.PlayRange_L", new BOX_BRUSH( "Sequencer/PlayRange_L", FMargin(1.f, 0.5f, 0.f, 0.5f) ) );
-		Set( "Sequencer.Timeline.PlayRange_R", new BOX_BRUSH( "Sequencer/PlayRange_R", FMargin(0.f, 0.5f, 1.f, 0.5f) ) );
-		Set( "Sequencer.Timeline.PlayRange_Bottom_L", new BOX_BRUSH( "Sequencer/PlayRange_Bottom_L", FMargin(1.f, 0.5f, 0.f, 0.5f) ) );
-		Set( "Sequencer.Timeline.PlayRange_Bottom_R", new BOX_BRUSH( "Sequencer/PlayRange_Bottom_R", FMargin(0.f, 0.5f, 1.f, 0.5f) ) );
-
-		Set( "Sequencer.Timeline.SubSequenceRangeHashL", new BORDER_BRUSH( "Sequencer/SubSequenceRangeHashL", FMargin(1.f, 0.f, 0.f, 0.f) ) );
-		Set( "Sequencer.Timeline.SubSequenceRangeHashR", new BORDER_BRUSH( "Sequencer/SubSequenceRangeHashR", FMargin(1.f, 0.f, 0.f, 0.f) ) );
-		Set( "Sequencer.Timeline.EaseInOut", new IMAGE_BRUSH( "Sequencer/EaseInOut", FVector2D(128, 128) ) );
-		Set( "Sequencer.InterpLine", new BOX_BRUSH( "Sequencer/InterpLine", FMargin(5.f/7.f, 0.f, 0.f, 0.f) ) );
-		
-		Set( "Sequencer.Transport.JumpToPreviousKey", FButtonStyle()
-			.SetNormal ( IMAGE_BRUSH( "/Sequencer/Transport_Bar/Previous_Frame_OFF", Icon24x24 ) )
-			.SetPressed( IMAGE_BRUSH( "/Sequencer/Transport_Bar/Previous_Frame", Icon24x24 ) )
-			.SetHovered( IMAGE_BRUSH( "/Sequencer/Transport_Bar/Previous_Frame_OFF", Icon24x24 ) ));
-		Set( "Sequencer.Transport.JumpToNextKey", FButtonStyle()
-			.SetNormal ( IMAGE_BRUSH( "/Sequencer/Transport_Bar/Next_Frame_24x_OFF", Icon24x24 ) )
-			.SetPressed( IMAGE_BRUSH( "/Sequencer/Transport_Bar/Next_Frame_24x", Icon24x24 ) )
-			.SetHovered( IMAGE_BRUSH( "/Sequencer/Transport_Bar/Next_Frame_24x_OFF", Icon24x24 ) ) );
-		Set( "Sequencer.Transport.SetPlayStart", FButtonStyle()
-			.SetNormal ( IMAGE_BRUSH( "/Sequencer/Transport_Bar/Bracket_In_16x24_OFF", FVector2D(16,24) ) )
-			.SetPressed( IMAGE_BRUSH( "/Sequencer/Transport_Bar/Bracket_In_16x24", FVector2D(16,24) ) )
-			.SetHovered( IMAGE_BRUSH( "/Sequencer/Transport_Bar/Bracket_In_16x24_OFF", FVector2D(16,24) ) ) );
-		Set( "Sequencer.Transport.SetPlayEnd", FButtonStyle()
-			.SetNormal ( IMAGE_BRUSH( "/Sequencer/Transport_Bar/Bracket_Out_16x24_OFF", FVector2D(16,24) ) )
-			.SetPressed( IMAGE_BRUSH( "/Sequencer/Transport_Bar/Bracket_Out_16x24", FVector2D(16,24) ) )
-			.SetHovered( IMAGE_BRUSH( "/Sequencer/Transport_Bar/Bracket_Out_16x24_OFF", FVector2D(16,24) ) ) );
-
-		Set( "Sequencer.Transport.CloseButton", FButtonStyle()
-			.SetNormal ( IMAGE_BRUSH( "/Docking/CloseApp_Normal", Icon16x16 ) )
-			.SetPressed( IMAGE_BRUSH( "/Docking/CloseApp_Pressed", Icon16x16 ) )
-			.SetHovered( IMAGE_BRUSH( "/Docking/CloseApp_Hovered", Icon16x16 )));
-
-		Set( "Sequencer.NotificationImage_AddedPlayMovieSceneEvent", new IMAGE_BRUSH( "Old/Checkbox_checked", Icon16x16 ) );
-
-		Set( "Sequencer.Save", new IMAGE_BRUSH( "Sequencer/Main_Icons/Icon_Sequencer_Save_24x", Icon48x48 ) );
-		Set( "Sequencer.Save.Small", new IMAGE_BRUSH( "Sequencer/Main_Icons/Icon_Sequencer_Save_24x", Icon24x24 ) );
-		Set( "Sequencer.SaveAs", new IMAGE_BRUSH( "Sequencer/Main_Icons/Icon_Sequencer_Save_As_24x", Icon48x48 ) );
-		Set( "Sequencer.SaveAs.Small", new IMAGE_BRUSH( "Sequencer/Main_Icons/Icon_Sequencer_Save_As_24x", Icon24x24 ) );
-		Set( "Sequencer.DiscardChanges", new IMAGE_BRUSH( "Sequencer/Main_Icons/Icon_Sequencer_Revert_24x", Icon48x48 ) );
-		Set( "Sequencer.DiscardChanges.Small", new IMAGE_BRUSH( "Sequencer/Main_Icons/Icon_Sequencer_Revert_24x", Icon24x24 ) );
-		Set( "Sequencer.RestoreAnimatedState", new IMAGE_BRUSH("Sequencer/Main_Icons/Icon_Sequencer_RestoreAnimatedState_24x", Icon48x48 ) );
-		Set( "Sequencer.RestoreAnimatedState.Small", new IMAGE_BRUSH("Sequencer/Main_Icons/Icon_Sequencer_RestoreAnimatedState_24x", Icon24x24 ) );
-		Set( "Sequencer.GenericGripLeft", new BOX_BRUSH( "Sequencer/GenericGripLeft", FMargin(5.f/16.f) ) );
-		Set( "Sequencer.GenericGripRight", new BOX_BRUSH( "Sequencer/GenericGripRight", FMargin(5.f/16.f) ) );
-		Set( "Sequencer.SectionArea.Background", new FSlateColorBrush( FColor::White ) );
-
-		Set( "Sequencer.Section.Background", new BORDER_BRUSH( TEXT("Sequencer/SectionBackground"), FMargin(4.f/16.f) ) );
-		Set( "Sequencer.Section.BackgroundTint", new BOX_BRUSH( TEXT("Sequencer/SectionBackgroundTint"), FMargin(4/16.f) ) );
-		Set( "Sequencer.Section.SelectedSectionOverlay", new IMAGE_BRUSH( TEXT("Sequencer/SelectedSectionOverlay"), Icon16x16, FLinearColor::White, ESlateBrushTileType::Both ) );
-		Set( "Sequencer.Section.SelectedTrackTint", new BOX_BRUSH( TEXT("Sequencer/SelectedTrackTint"), FMargin(0.f, 0.5f) ) );
-		Set( "Sequencer.Section.SelectionBorder", new BORDER_BRUSH( TEXT("Sequencer/SectionHighlight"), FMargin(7.f/16.f) ) );
-		Set( "Sequencer.Section.LockedBorder", new BORDER_BRUSH( TEXT("Sequencer/SectionLocked"), FMargin(7.f/16.f) ) );
-		Set( "Sequencer.Section.SelectedSectionOverlay", new IMAGE_BRUSH( TEXT("Sequencer/SelectedSectionOverlay"), Icon16x16, FLinearColor::White, ESlateBrushTileType::Both ) );
-		Set( "Sequencer.Section.FilmBorder", new IMAGE_BRUSH( TEXT("Sequencer/SectionFilmBorder"), FVector2D(10, 7), FLinearColor::White, ESlateBrushTileType::Horizontal ) );
-		Set( "Sequencer.Section.GripLeft", new BOX_BRUSH( "Sequencer/SectionGripLeft", FMargin(5.f/16.f) ) );
-		Set( "Sequencer.Section.GripRight", new BOX_BRUSH( "Sequencer/SectionGripRight", FMargin(5.f/16.f) ) );
-		Set( "Sequencer.Section.EasingHandle", new IMAGE_BRUSH( "Sequencer/EasingHandle", FVector2D(10.f,10.f) ) );
-
-		Set( "Sequencer.Section.PreRoll", new BORDER_BRUSH( TEXT("Sequencer/PreRoll"), FMargin(0.f, .5f, 0.f, .5f) ) );
-
-		Set( "Sequencer.Section.PinCusion", new IMAGE_BRUSH( TEXT("Sequencer/PinCusion"), Icon16x16, FLinearColor::White, ESlateBrushTileType::Both ) );
-		Set( "Sequencer.Section.OverlapBorder", new BORDER_BRUSH( TEXT("Sequencer/OverlapBorder"), FMargin(1.f/4.f, 0.f) ) );
-		Set( "Sequencer.Section.StripeOverlay", new BOX_BRUSH( "Sequencer/SectionStripeOverlay", FMargin(0.f, .5f) ) );
-		Set( "Sequencer.Section.BackgroundText", DEFAULT_FONT( "Bold", 24 ) );
-		Set( "Sequencer.Section.EmptySpace", new BOX_BRUSH( TEXT("Sequencer/EmptySpace"), FMargin(0.f, 7.f/14.f) ) );
-
-		Set( "Sequencer.AnimationOutliner.ColorStrip", FButtonStyle()
-			.SetNormal(FSlateNoResource())
-			.SetHovered(FSlateNoResource())
-			.SetPressed(FSlateNoResource())
-			.SetNormalPadding(FMargin(0,0,0,0))
-			.SetPressedPadding(FMargin(0,0,0,0))
-		);
-
-		Set( "Sequencer.AnimationOutliner.TopLevelBorder_Expanded", new BOX_BRUSH( "Sequencer/TopLevelNodeBorder_Expanded", FMargin(4.0f/16.0f) ) );
-		Set( "Sequencer.AnimationOutliner.TopLevelBorder_Collapsed", new BOX_BRUSH( "Sequencer/TopLevelNodeBorder_Collapsed", FMargin(4.0f/16.0f) ) );
-		Set( "Sequencer.AnimationOutliner.DefaultBorder", new FSlateColorBrush( FLinearColor::White ) );
-		Set( "Sequencer.AnimationOutliner.TransparentBorder", new FSlateColorBrush( FLinearColor::Transparent ) );
-		Set( "Sequencer.AnimationOutliner.BoldFont", DEFAULT_FONT( "Bold", 11 ) );
-		Set( "Sequencer.AnimationOutliner.RegularFont", DEFAULT_FONT( "Regular", 9 ) );
-		Set( "Sequencer.ShotFilter", new IMAGE_BRUSH( "Sequencer/FilteredArea", FVector2D(74,74), FLinearColor::White, ESlateBrushTileType::Both ) );
-		Set( "Sequencer.KeyMark", new IMAGE_BRUSH("Sequencer/KeyMark", FVector2D(3,21), FLinearColor::White, ESlateBrushTileType::NoTile ) );
-		Set( "Sequencer.SetAutoKey", new IMAGE_BRUSH( "Sequencer/Main_Icons/Icon_Sequencer_Auto_Key_24x", Icon48x48 ) );
-		Set( "Sequencer.SetAutoKey.Small", new IMAGE_BRUSH( "Sequencer/Main_Icons/Icon_Sequencer_Auto_Key_24x", Icon24x24 ) );
-		Set( "Sequencer.SetAutoTrack", new IMAGE_BRUSH("Sequencer/Main_Icons/Icon_Sequencer_Auto_Track_24x", Icon48x48 ) );
-		Set( "Sequencer.SetAutoTrack.Small", new IMAGE_BRUSH("Sequencer/Main_Icons/Icon_Sequencer_Auto_Track_24x", Icon24x24 ) );
-		Set( "Sequencer.SetAutoChangeAll", new IMAGE_BRUSH( "Sequencer/Main_Icons/Icon_Sequencer_Auto_Key_All_24x", Icon48x48 ) );
-		Set( "Sequencer.SetAutoChangeAll.Small", new IMAGE_BRUSH( "Sequencer/Main_Icons/Icon_Sequencer_Auto_Key_All_24x", Icon24x24 ) );
-		Set( "Sequencer.SetAutoChangeNone", new IMAGE_BRUSH("Sequencer/Main_Icons/Icon_Sequencer_Disable_Auto_Key_24x", Icon48x48));
-		Set( "Sequencer.SetAutoChangeNone.Small", new IMAGE_BRUSH("Sequencer/Main_Icons/Icon_Sequencer_Disable_Auto_Key_24x", Icon24x24 ) );
-		Set( "Sequencer.AllowAllEdits", new IMAGE_BRUSH( "Sequencer/Main_Icons/Icon_Sequencer_Allow_All_Edits_24x", Icon48x48 ) );
-		Set( "Sequencer.AllowAllEdits.Small", new IMAGE_BRUSH( "Sequencer/Main_Icons/Icon_Sequencer_Allow_All_Edits_24x", Icon24x24 ) );
-		Set( "Sequencer.AllowSequencerEditsOnly", new IMAGE_BRUSH( "Sequencer/Main_Icons/Icon_Sequencer_Allow_Sequencer_Edits_Only_24x", Icon48x48 ) );
-		Set( "Sequencer.AllowSequencerEditsOnly.Small", new IMAGE_BRUSH( "Sequencer/Main_Icons/Icon_Sequencer_Allow_Sequencer_Edits_Only_24x", Icon24x24 ) );
-		Set( "Sequencer.AllowLevelEditsOnly", new IMAGE_BRUSH( "Sequencer/Main_Icons/Icon_Sequencer_Allow_Level_Edits_Only_24x", Icon48x48 ) );
-		Set( "Sequencer.AllowLevelEditsOnly.Small", new IMAGE_BRUSH( "Sequencer/Main_Icons/Icon_Sequencer_Allow_Level_Edits_Only_24x", Icon24x24 ) );		
-		Set( "Sequencer.SetKeyAll", new IMAGE_BRUSH( "Sequencer/Main_Icons/Icon_Sequencer_Key_All_24x", Icon48x48 ) );
-		Set( "Sequencer.SetKeyAll.Small", new IMAGE_BRUSH( "Sequencer/Main_Icons/Icon_Sequencer_Key_All_24x", Icon24x24 ) );
-		Set( "Sequencer.SetKeyGroup", new IMAGE_BRUSH("Sequencer/Main_Icons/Icon_Sequencer_Key_Group_24x", Icon48x48));
-		Set( "Sequencer.SetKeyGroup.Small", new IMAGE_BRUSH("Sequencer/Main_Icons/Icon_Sequencer_Key_Group_24x", Icon24x24));
-		Set( "Sequencer.SetKeyChanged", new IMAGE_BRUSH( "Sequencer/Main_Icons/Icon_Sequencer_Key_Part_24x", Icon48x48 ) );
-		Set( "Sequencer.SetKeyChanged.Small", new IMAGE_BRUSH( "Sequencer/Main_Icons/Icon_Sequencer_Key_Part_24x", Icon24x24 ) );
-		Set( "Sequencer.ToggleIsSnapEnabled", new IMAGE_BRUSH( "Sequencer/Main_Icons/Icon_Sequencer_Snap_24x", Icon48x48 ) );
-		Set( "Sequencer.ToggleIsSnapEnabled.Small", new IMAGE_BRUSH( "Sequencer/Main_Icons/Icon_Sequencer_Snap_24x", Icon24x24 ) );
-		Set( "Sequencer.ToggleShowCurveEditor", new IMAGE_BRUSH("Icons/SequencerIcons/icon_Sequencer_CurveEditor_24x", Icon48x48) );
-		Set( "Sequencer.ToggleShowCurveEditor.Small", new IMAGE_BRUSH("Icons/SequencerIcons/icon_Sequencer_CurveEditor_24x", Icon24x24) );
-		Set( "Sequencer.ToggleAutoScroll", new IMAGE_BRUSH( "Icons/icon_Sequencer_ToggleAutoScroll_40x", Icon48x48 ) );
-		Set( "Sequencer.ToggleAutoScroll.Small", new IMAGE_BRUSH( "Icons/icon_Sequencer_ToggleAutoScroll_16x", Icon16x16 ) );
-		Set( "Sequencer.MoveTool.Small", new IMAGE_BRUSH( "Icons/SequencerIcons/icon_Sequencer_Move_24x", Icon16x16 ) );
-		Set( "Sequencer.MarqueeTool.Small", new IMAGE_BRUSH( "Icons/SequencerIcons/icon_Sequencer_Marquee_24x", Icon16x16 ) );
-		Set( "Sequencer.RenderMovie.Small", new IMAGE_BRUSH( "Sequencer/Main_Icons/Icon_Sequencer_Create_Movie_24x", Icon24x24 ) );
-		Set( "Sequencer.CreateCamera.Small", new IMAGE_BRUSH( "Sequencer/Main_Icons/Icon_Sequencer_Create_Camera_24x", Icon24x24 ) );
-		Set( "Sequencer.FindInContentBrowser.Small", new IMAGE_BRUSH( "Sequencer/Main_Icons/Icon_Sequencer_Find_In_Content_Browser_24x", Icon24x24 ) );
-		Set( "Sequencer.LockCamera", new IMAGE_BRUSH( "Sequencer/Main_Icons/Icon_Sequencer_Look_Thru_24x", Icon16x16 ) );
-		Set( "Sequencer.UnlockCamera", new IMAGE_BRUSH( "Sequencer/Main_Icons/Icon_Sequencer_Look_Thru_24x", Icon16x16, FLinearColor(1.f, 1.f, 1.f, 0.5f) ) );
-		Set( "Sequencer.Thumbnail.SectionHandle", new IMAGE_BRUSH( "Old/White", Icon16x16, FLinearColor::Black ) );
-		Set( "Sequencer.TrackHoverHighlight_Top", new IMAGE_BRUSH( TEXT("Sequencer/TrackHoverHighlight_Top"), FVector2D(4, 4) ) );
-		Set( "Sequencer.TrackHoverHighlight_Bottom", new IMAGE_BRUSH( TEXT("Sequencer/TrackHoverHighlight_Bottom"), FVector2D(4, 4) ) );
-		Set( "Sequencer.SpawnableIconOverlay", new IMAGE_BRUSH( TEXT("Sequencer/SpawnableIconOverlay"), FVector2D(13, 13) ) );
-		Set( "Sequencer.MultipleIconOverlay", new IMAGE_BRUSH(TEXT("Sequencer/MultipleIconOverlay"), FVector2D(13, 13) ) );
-		Set( "Sequencer.LockSequence", new IMAGE_BRUSH("Sequencer/Main_Icons/Icon_Sequencer_Locked_16x", Icon16x16) );
-		Set( "Sequencer.UnlockSequence", new IMAGE_BRUSH("Sequencer/Main_Icons/Icon_Sequencer_Unlocked_16x", Icon16x16) );
-
-		Set( "Sequencer.GeneralOptions", new IMAGE_BRUSH( "Sequencer/Main_Icons/Icon_Sequencer_General_Options_24x", Icon48x48 ) );
-		Set( "Sequencer.GeneralOptions.Small", new IMAGE_BRUSH( "Sequencer/Main_Icons/Icon_Sequencer_General_Options_24x", Icon24x24 ) );
-		Set( "Sequencer.PlaybackOptions", new IMAGE_BRUSH( "Sequencer/Main_Icons/Icon_Sequencer_Playback_Options_24x", Icon48x48 ) );
-		Set( "Sequencer.PlaybackOptions.Small", new IMAGE_BRUSH( "Sequencer/Main_Icons/Icon_Sequencer_Playback_Options_24x", Icon24x24 ) );
-		Set( "Sequencer.SelectEditOptions", new IMAGE_BRUSH( "Sequencer/Main_Icons/Icon_Sequencer_SelectEdit_Options_24x", Icon48x48 ) );
-		Set( "Sequencer.SelectEditOptions.Small", new IMAGE_BRUSH( "Sequencer/Main_Icons/Icon_Sequencer_SelectEdit_Options_24x", Icon24x24 ) );
-		Set( "Sequencer.Time", new IMAGE_BRUSH( "Sequencer/Main_Icons/Icon_Sequencer_Time_24x", Icon48x48 ) );
-		Set( "Sequencer.Time.Small", new IMAGE_BRUSH( "Sequencer/Main_Icons/Icon_Sequencer_Time_24x", Icon24x24 ) );
-		Set( "Sequencer.Value", new IMAGE_BRUSH( "Sequencer/Main_Icons/Icon_Sequencer_Value_24x", Icon48x48 ) );
-		Set( "Sequencer.Value.Small", new IMAGE_BRUSH( "Sequencer/Main_Icons/Icon_Sequencer_Value_24x", Icon24x24 ) );
-
-		Set( "Sequencer.OverlayPanel.Background", new BOX_BRUSH( "Sequencer/OverlayPanelBackground", FMargin(26.f/54.f) ) );
-
-		Set( "Sequencer.TrackArea.LaneColor", FLinearColor(0.3f, 0.3f, 0.3f, 0.3f) );
-
-		Set( "Sequencer.Tracks.Media", new IMAGE_BRUSH("Sequencer/Dropdown_Icons/Icon_Media_Track_16x", Icon16x16));
-		Set( "Sequencer.Tracks.Audio", new IMAGE_BRUSH("Sequencer/Dropdown_Icons/Icon_Audio_Track_16x", Icon16x16));
-		Set( "Sequencer.Tracks.Event", new IMAGE_BRUSH("Sequencer/Dropdown_Icons/Icon_Event_Track_16x", Icon16x16));
-		Set( "Sequencer.Tracks.Fade", new IMAGE_BRUSH("Sequencer/Dropdown_Icons/Icon_Fade_Track_16x", Icon16x16));
-		Set( "Sequencer.Tracks.CameraCut", new IMAGE_BRUSH("Sequencer/Dropdown_Icons/Icon_Camera_Cut_Track_16x", Icon16x16));
-		Set( "Sequencer.Tracks.CinematicShot", new IMAGE_BRUSH("Sequencer/Dropdown_Icons/Icon_Shot_Track_16x", Icon16x16));
-		Set( "Sequencer.Tracks.Slomo", new IMAGE_BRUSH("Sequencer/Dropdown_Icons/Icon_Play_Rate_Track_16x", Icon16x16));
-		Set( "Sequencer.Tracks.Sub", new IMAGE_BRUSH("Sequencer/Dropdown_Icons/Icon_Sub_Track_16x", Icon16x16));
-		Set( "Sequencer.Tracks.LevelVisibility", new IMAGE_BRUSH("Sequencer/Dropdown_Icons/Icon_Level_Visibility_Track_16x", Icon16x16));
-
-		Set( "Sequencer.CursorDecorator_MarqueeAdd", new IMAGE_BRUSH( "Sequencer/CursorDecorator_MarqueeAdd", Icon16x16));
-		Set( "Sequencer.CursorDecorator_MarqueeSubtract", new IMAGE_BRUSH( "Sequencer/CursorDecorator_MarqueeSubtract", Icon16x16));
-
-		Set( "Sequencer.BreadcrumbText", FTextBlockStyle(NormalText)
-			.SetFont( DEFAULT_FONT( "Bold", 11 ) )
-			.SetColorAndOpacity( FLinearColor( 1.0f, 1.0f, 1.0f ) )
-			.SetHighlightColor( FLinearColor( 1.0f, 1.0f, 1.0f ) )
-			.SetShadowOffset( FVector2D( 1,1 ) )
-			.SetShadowColorAndOpacity( FLinearColor(0,0,0,0.9f) ) );
-		Set( "Sequencer.BreadcrumbIcon", new IMAGE_BRUSH( "Common/SmallArrowRight", Icon10x10 ) );
-
-		const FButtonStyle DetailsKeyButton = FButtonStyle(NoBorder)
-			.SetNormal( IMAGE_BRUSH("Sequencer/AddKey_Details", FVector2D(11,11) )  )
-			.SetHovered( IMAGE_BRUSH("Sequencer/AddKey_Details", FVector2D(11,11), SelectionColor ) )
-			.SetPressed( IMAGE_BRUSH("Sequencer/AddKey_Details", FVector2D(11,11), SelectionColor_Pressed ) )
-			.SetNormalPadding(FMargin(0, 1))
-			.SetPressedPadding(FMargin(0, 2, 0, 0));
-		Set( "Sequencer.AddKey.Details", DetailsKeyButton );
-
-		const FSplitterStyle OutlinerSplitterStyle = FSplitterStyle()
-			.SetHandleNormalBrush( FSlateNoResource() )
-			.SetHandleHighlightBrush( FSlateNoResource() );
-		Set( "Sequencer.AnimationOutliner.Splitter", OutlinerSplitterStyle );
-
-		Set( "Sequencer.HyperlinkSpinBox", FSpinBoxStyle(GetWidgetStyle<FSpinBoxStyle>("SpinBox"))
-			.SetTextPadding(FMargin(0))
-			.SetBackgroundBrush(BORDER_BRUSH("Old/HyperlinkDotted", FMargin(0,0,0,3/16.0f), FSlateColor::UseSubduedForeground()))
-			.SetHoveredBackgroundBrush(FSlateNoResource())
-			.SetInactiveFillBrush(FSlateNoResource())
-			.SetActiveFillBrush(FSlateNoResource())
-			.SetForegroundColor(FSlateColor::UseSubduedForeground())
-			.SetArrowsImage(FSlateNoResource())
-		);
-
-		Set("Sequencer.PlayTimeSpinBox", FSpinBoxStyle(GetWidgetStyle<FSpinBoxStyle>("SpinBox"))
-			.SetTextPadding(FMargin(0))
-			.SetBackgroundBrush(FSlateNoResource())
-			.SetHoveredBackgroundBrush(FSlateNoResource())
-			.SetInactiveFillBrush(FSlateNoResource())
-			.SetActiveFillBrush(FSlateNoResource())
-			.SetForegroundColor(SelectionColor_Pressed)
-			.SetArrowsImage(FSlateNoResource())
-		);
-
-		Set( "Sequencer.HyperlinkTextBox", FEditableTextBoxStyle()
-			.SetFont( DEFAULT_FONT( "Regular", 9 ) )
-			.SetBackgroundImageNormal( FSlateNoResource() )
-			.SetBackgroundImageHovered( FSlateNoResource() )
-			.SetBackgroundImageFocused( FSlateNoResource() )
-			.SetBackgroundImageReadOnly( FSlateNoResource() )
-			.SetBackgroundColor( FLinearColor::Transparent )
-			.SetForegroundColor( FSlateColor::UseSubduedForeground() )
-			);
-		Set( "Sequencer.FixedFont", DEFAULT_FONT( "Mono", 9 ) );
-
-		Set( "Sequencer.RecordSelectedActors", new IMAGE_BRUSH( "SequenceRecorder/icon_tab_SequenceRecorder_16x", Icon16x16 ) );
-
-		FComboButtonStyle SequencerSectionComboButton = FComboButtonStyle()
-			.SetButtonStyle(
-				FButtonStyle()
-				.SetNormal(FSlateNoResource())
-				.SetHovered(FSlateNoResource())
-				.SetPressed(FSlateNoResource())
-				.SetNormalPadding(FMargin(0,0,0,0))
-				.SetPressedPadding(FMargin(0,1,0,0))
-			)
-			.SetDownArrowImage(IMAGE_BRUSH("Common/ComboArrow", Icon8x8));
-		Set( "Sequencer.SectionComboButton", SequencerSectionComboButton );
-
-		Set("Sequencer.CreateEventBinding", new IMAGE_BRUSH("Icons/icon_Blueprint_AddFunction_16px", Icon16x16));
-		Set("Sequencer.CreateQuickBinding", new IMAGE_BRUSH("Icons/icon_Blueprint_Node_16x", Icon16x16));
-		Set("Sequencer.ClearEventBinding", new IMAGE_BRUSH("Icons/Edit/icon_Edit_Delete_40x", Icon16x16));
-		Set("Sequencer.MultipleEvents", new IMAGE_BRUSH("Sequencer/MultipleEvents", Icon16x16));
-		Set("Sequencer.UnboundEvent", new IMAGE_BRUSH("Sequencer/UnboundEvent", Icon16x16));
-
-		// Sequencer Blending Iconography
-		Set( "EMovieSceneBlendType::Absolute", new IMAGE_BRUSH( "Sequencer/EMovieSceneBlendType_Absolute", FVector2D(32, 16) ) );
-		Set( "EMovieSceneBlendType::Relative", new IMAGE_BRUSH( "Sequencer/EMovieSceneBlendType_Relative", FVector2D(32, 16) ) );
-		Set( "EMovieSceneBlendType::Additive", new IMAGE_BRUSH( "Sequencer/EMovieSceneBlendType_Additive", FVector2D(32, 16) ) );
-	}
-
-
-	// Sequence recorder standalone UI
-	if (IncludeEditorSpecificStyles())
-	{
-		Set( "SequenceRecorder.TabIcon", new IMAGE_BRUSH( "SequenceRecorder/icon_tab_SequenceRecorder_16x", Icon16x16 ) );
-		Set( "SequenceRecorder.Common.RecordAll.Small", new IMAGE_BRUSH( "SequenceRecorder/icon_RecordAll_40x", Icon20x20 ) );
-		Set( "SequenceRecorder.Common.RecordAll", new IMAGE_BRUSH( "SequenceRecorder/icon_RecordAll_40x", Icon40x40 ) );
-		Set( "SequenceRecorder.Common.StopAll.Small", new IMAGE_BRUSH( "SequenceRecorder/icon_StopAll_40x", Icon20x20 ) );
-		Set( "SequenceRecorder.Common.StopAll", new IMAGE_BRUSH( "SequenceRecorder/icon_StopAll_40x", Icon40x40 ) );
-		Set( "SequenceRecorder.Common.AddRecording.Small", new IMAGE_BRUSH( "SequenceRecorder/icon_AddRecording_40x", Icon20x20 ) );
-		Set( "SequenceRecorder.Common.AddRecording", new IMAGE_BRUSH( "SequenceRecorder/icon_AddRecording_40x", Icon40x40 ) );
-		Set( "SequenceRecorder.Common.AddCurrentPlayerRecording.Small", new IMAGE_BRUSH( "SequenceRecorder/icon_AddCurrentPlayerRecording_40x", Icon20x20 ) );
-		Set( "SequenceRecorder.Common.AddCurrentPlayerRecording", new IMAGE_BRUSH( "SequenceRecorder/icon_AddCurrentPlayerRecording_40x", Icon40x40 ) );
-		Set( "SequenceRecorder.Common.RemoveRecording.Small", new IMAGE_BRUSH( "SequenceRecorder/icon_RemoveRecording_40x", Icon20x20 ) );
-		Set( "SequenceRecorder.Common.RemoveRecording", new IMAGE_BRUSH( "SequenceRecorder/icon_RemoveRecording_40x", Icon40x40 ) );
-		Set( "SequenceRecorder.Common.RemoveAllRecordings.Small", new IMAGE_BRUSH( "SequenceRecorder/icon_RemoveRecording_40x", Icon20x20 ) );
-		Set( "SequenceRecorder.Common.RemoveAllRecordings", new IMAGE_BRUSH( "SequenceRecorder/icon_RemoveRecording_40x", Icon40x40 ) );
-		Set( "SequenceRecorder.Common.RecordingActive", new IMAGE_BRUSH( "Common/SmallCheckBox_Checked", Icon14x14 ) );
-		Set( "SequenceRecorder.Common.RecordingInactive", new IMAGE_BRUSH( "Common/SmallCheckBox", Icon14x14 ) );
-	}
-
 	// Foliage Edit Mode
 	if (IncludeEditorSpecificStyles())
 	{	
@@ -2449,108 +2112,7 @@ void FSlateEditorStyle::FStyle::SetupGeneralStyles()
 			);
 	}
 
-	// MenuBar
-	{
-		Set( "Menu.Background",	new BOX_BRUSH( "Old/Menu_Background", FMargin(8.0f/64.0f) ) );
-		Set( "Menu.Icon", new IMAGE_BRUSH( "Icons/icon_tab_toolbar_16px", Icon16x16 ) );
-		Set( "Menu.Expand", new IMAGE_BRUSH( "Icons/toolbar_expand_16x", Icon16x16) );
-		Set( "Menu.SubMenuIndicator", new IMAGE_BRUSH( "Common/SubmenuArrow", Icon8x8 ) );
-		Set( "Menu.SToolBarComboButtonBlock.Padding", FMargin(4.0f));
-		Set( "Menu.SToolBarButtonBlock.Padding", FMargin(4.0f));
-		Set( "Menu.SToolBarCheckComboButtonBlock.Padding", FMargin(4.0f));
-		Set( "Menu.SToolBarButtonBlock.CheckBox.Padding", FMargin(0.0f) );
-		Set( "Menu.SToolBarComboButtonBlock.ComboButton.Color", DefaultForeground );
-
-		Set( "Menu.Block.IndentedPadding", FMargin( 18.0f, 2.0f, 4.0f, 4.0f ) );
-		Set( "Menu.Block.Padding", FMargin( 2.0f, 2.0f, 4.0f, 4.0f ) );
-
-		Set( "Menu.Separator", new BOX_BRUSH( "Old/Button", 4.0f/32.0f ) );
-		Set( "Menu.Separator.Padding", FMargin( 0.5f ) );
-
-		Set( "Menu.Label", FTextBlockStyle(NormalText) .SetFont( DEFAULT_FONT( "Regular", 9 ) ) );
-		Set( "Menu.Label.Padding", FMargin(0.0f, 0.0f, 0.0f, 0.0f) );
-		Set( "Menu.Label.ContentPadding", FMargin(10.0f, 2.0f) );
-		Set( "Menu.EditableText", FEditableTextBoxStyle(NormalEditableTextBoxStyle) .SetFont( DEFAULT_FONT( "Regular", 9 ) ) );
-		Set( "Menu.Keybinding", FTextBlockStyle(NormalText) .SetFont( DEFAULT_FONT( "Regular", 8 ) ) );
-
-		Set( "Menu.Heading", FTextBlockStyle(NormalText)
-			.SetFont( DEFAULT_FONT( "Regular", 8 ) )
-			.SetColorAndOpacity( FLinearColor( 0.4f, 0.4, 0.4f, 1.0f ) ) );
-
-		/* Set images for various SCheckBox states associated with menu check box items... */
-		const FCheckBoxStyle BasicMenuCheckBoxStyle = FCheckBoxStyle()
-			.SetUncheckedImage( IMAGE_BRUSH( "Common/SmallCheckBox", Icon14x14 ) )
-			.SetUncheckedHoveredImage( IMAGE_BRUSH( "Common/SmallCheckBox_Hovered", Icon14x14 ) )
-			.SetUncheckedPressedImage( IMAGE_BRUSH( "Common/SmallCheckBox_Hovered", Icon14x14, FLinearColor( 0.5f, 0.5f, 0.5f ) ) )
-			.SetCheckedImage( IMAGE_BRUSH( "Common/SmallCheckBox_Checked", Icon14x14 ) )
-			.SetCheckedHoveredImage( IMAGE_BRUSH( "Common/SmallCheckBox_Checked_Hovered", Icon14x14 ) )
-			.SetCheckedPressedImage( IMAGE_BRUSH( "Common/SmallCheckBox_Checked_Hovered", Icon14x14, FLinearColor( 0.5f, 0.5f, 0.5f ) ) )
-			.SetUndeterminedImage( IMAGE_BRUSH( "Common/CheckBox_Undetermined", Icon14x14 ) )
-			.SetUndeterminedHoveredImage( IMAGE_BRUSH( "Common/CheckBox_Undetermined_Hovered", Icon14x14 ) )
-			.SetUndeterminedPressedImage( IMAGE_BRUSH( "Common/CheckBox_Undetermined_Hovered", Icon14x14, FLinearColor( 0.5f, 0.5f, 0.5f ) ) );
- 
-		/* ...and add the new style */
-		Set( "Menu.CheckBox", BasicMenuCheckBoxStyle );
-						
-		/* Read-only checkbox that appears next to a menu item */
-		/* Set images for various SCheckBox states associated with read-only menu check box items... */
-		const FCheckBoxStyle BasicMenuCheckStyle = FCheckBoxStyle()
-			.SetUncheckedImage( IMAGE_BRUSH( "Icons/Empty_14x", Icon14x14 ) )
-			.SetUncheckedHoveredImage( IMAGE_BRUSH( "Icons/Empty_14x", Icon14x14 ) )
-			.SetUncheckedPressedImage( IMAGE_BRUSH( "Common/SmallCheckBox_Hovered", Icon14x14 ) )
-			.SetCheckedImage( IMAGE_BRUSH( "Common/SmallCheck", Icon14x14 ) )
-			.SetCheckedHoveredImage( IMAGE_BRUSH( "Common/SmallCheck", Icon14x14 ) )
-			.SetCheckedPressedImage( IMAGE_BRUSH( "Common/SmallCheck", Icon14x14 ) )
-			.SetUndeterminedImage( IMAGE_BRUSH( "Icons/Empty_14x", Icon14x14 ) )
-			.SetUndeterminedHoveredImage( FSlateNoResource() )
-			.SetUndeterminedPressedImage( FSlateNoResource() );
-
-		/* ...and add the new style */
-		Set( "Menu.Check", BasicMenuCheckStyle );
-
-		/* This radio button is actually just a check box with different images */
-		/* Set images for various Menu radio button (SCheckBox) states... */
-		const FCheckBoxStyle BasicMenuRadioButtonStyle = FCheckBoxStyle()
-			.SetUncheckedImage( IMAGE_BRUSH( "Common/RadioButton_Unselected_16x", Icon16x16 ) )
-			.SetUncheckedHoveredImage( IMAGE_BRUSH( "Common/RadioButton_Unselected_16x", Icon16x16 ) )
-			.SetUncheckedPressedImage( IMAGE_BRUSH( "Common/RadioButton_Unselected_16x", Icon16x16 ) )
-			.SetCheckedImage( IMAGE_BRUSH( "Common/RadioButton_Selected_16x", Icon16x16 ) )
-			.SetCheckedHoveredImage( IMAGE_BRUSH( "Common/RadioButton_Selected_16x", Icon16x16, SelectionColor ) )
-			.SetCheckedPressedImage( IMAGE_BRUSH( "Common/RadioButton_Unselected_16x", Icon16x16, SelectionColor_Pressed ) )
-			.SetUndeterminedImage( IMAGE_BRUSH( "Common/RadioButton_Unselected_16x", Icon16x16 ) )
-			.SetUndeterminedHoveredImage( IMAGE_BRUSH( "Common/RadioButton_Unselected_16x", Icon16x16, SelectionColor ) )
-			.SetUndeterminedPressedImage( IMAGE_BRUSH( "Common/RadioButton_Unselected_16x", Icon16x16, SelectionColor_Pressed ) );
-
-		/* ...and set new style */
-		Set( "Menu.RadioButton", BasicMenuRadioButtonStyle );
-
-		/* Create style for "Menu.ToggleButton" widget ... */
-		const FCheckBoxStyle MenuToggleButtonCheckBoxStyle = FCheckBoxStyle()
-			.SetCheckBoxType( ESlateCheckBoxType::ToggleButton )
-			.SetUncheckedImage( FSlateNoResource() )
-			.SetUncheckedPressedImage( BOX_BRUSH( "Common/RoundedSelection_16x", 4.0f/16.0f, SelectionColor_Pressed ) )
-			.SetUncheckedHoveredImage( BOX_BRUSH( "Common/RoundedSelection_16x", 4.0f/16.0f, SelectionColor ) )
-			.SetCheckedImage( BOX_BRUSH( "Common/RoundedSelection_16x",  4.0f/16.0f, SelectionColor_Pressed ) )
-			.SetCheckedHoveredImage( BOX_BRUSH( "Common/RoundedSelection_16x",  4.0f/16.0f, SelectionColor_Pressed ) )
-			.SetCheckedPressedImage( BOX_BRUSH( "Common/RoundedSelection_16x",  4.0f/16.0f, SelectionColor ) );
-		/* ... and add new style */
-		Set( "Menu.ToggleButton", MenuToggleButtonCheckBoxStyle );
-
-		Set( "Menu.Button", FButtonStyle( NoBorder )
-			.SetNormal ( FSlateNoResource() )
-			.SetPressed( BOX_BRUSH( "Common/RoundedSelection_16x", 4.0f/16.0f, SelectionColor_Pressed ) )
-			.SetHovered( BOX_BRUSH( "Common/RoundedSelection_16x", 4.0f/16.0f, SelectionColor ) )
-			.SetNormalPadding( FMargin(0,1) )
-			.SetPressedPadding( FMargin(0,2,0,0) )
-		);
-
-		Set( "Menu.Button.Checked", new BOX_BRUSH( "Common/RoundedSelection_16x",  4.0f/16.0f, SelectionColor_Pressed ) );
-		Set( "Menu.Button.Checked_Hovered", new BOX_BRUSH( "Common/RoundedSelection_16x",  4.0f/16.0f, SelectionColor_Pressed ) );
-		Set( "Menu.Button.Checked_Pressed", new BOX_BRUSH( "Common/RoundedSelection_16x",  4.0f/16.0f, SelectionColor ) );
-
-		/* The style of a menu bar button when it has a sub menu open */
-		Set( "Menu.Button.SubMenuOpen", new BORDER_BRUSH( "Common/Selection", FMargin(4.f/16.f), FLinearColor(0.10f, 0.10f, 0.10f) ) );
-	}
+	
 
 	// ViewportLayoutToolbar
 	{
@@ -2578,239 +2140,7 @@ void FSlateEditorStyle::FStyle::SetupGeneralStyles()
 		Set( "ViewportLayoutToolbar.SToolBarComboButtonBlock.ComboButton.Color", DefaultForeground );
 		}
 
-	// NotificationBar
-	{
-		Set( "NotificationBar.Background",	new FSlateNoResource() );
-		Set( "NotificationBar.Icon", new FSlateNoResource() );
-		Set( "NotificationBar.Expand", new IMAGE_BRUSH( "Icons/toolbar_expand_16x", Icon16x16) );
-		Set( "NotificationBar.SubMenuIndicator", new IMAGE_BRUSH( "Common/SubmenuArrow", Icon8x8 ) );
-
-		Set( "NotificationBar.Block.IndentedPadding", FMargin( 0 ) );
-		Set( "NotificationBar.Block.Padding", FMargin( 0 ) );
-
-		Set( "NotificationBar.Separator", new BOX_BRUSH( "Old/Button", 4.0f/32.0f ) );
-		Set( "NotificationBar.Separator.Padding", FMargin( 0.5f ) );
-
-		Set( "NotificationBar.Label", FTextBlockStyle(NormalText) .SetFont( DEFAULT_FONT( "Regular", 9 ) ) );
-		Set( "NotificationBar.EditableText", FEditableTextBoxStyle(NormalEditableTextBoxStyle) .SetFont( DEFAULT_FONT( "Regular", 9 ) ) );
-		Set( "NotificationBar.Keybinding", FTextBlockStyle(NormalText) .SetFont( DEFAULT_FONT( "Regular", 8 ) ) );
-
-		Set( "NotificationBar.Heading", FTextBlockStyle(NormalText)
-			.SetFont( DEFAULT_FONT( "Regular", 8 ) )
-			.SetColorAndOpacity( FLinearColor( 0.4f, 0.4, 0.4f, 1.0f ) ) );
-
-		const FCheckBoxStyle NotificationBarCheckBoxCheckBoxStyle = FCheckBoxStyle()
-			.SetUncheckedImage( IMAGE_BRUSH( "Common/SmallCheckBox", Icon14x14 ) )
-			.SetUncheckedPressedImage( IMAGE_BRUSH( "Common/SmallCheckBox_Hovered", Icon14x14, FLinearColor( 0.5f, 0.5f, 0.5f ) ) )
-			.SetUncheckedHoveredImage( IMAGE_BRUSH( "Common/SmallCheckBox_Hovered", Icon14x14 ) )
-			.SetCheckedHoveredImage( IMAGE_BRUSH( "Common/SmallCheckBox_Checked_Hovered", Icon14x14 ) )
-			.SetCheckedPressedImage( IMAGE_BRUSH( "Common/SmallCheckBox_Checked_Hovered", Icon14x14, FLinearColor( 0.5f, 0.5f, 0.5f ) ) )
-			.SetCheckedImage( IMAGE_BRUSH( "Common/SmallCheckBox_Checked",  Icon14x14 ) )
-			.SetUndeterminedImage( IMAGE_BRUSH( "Common/CheckBox_Undetermined", Icon14x14 ) )
-			.SetUndeterminedHoveredImage( IMAGE_BRUSH( "Common/CheckBox_Undetermined_Hovered", Icon14x14 ) )
-			.SetUndeterminedPressedImage( IMAGE_BRUSH( "Common/CheckBox_Undetermined_Hovered", Icon14x14, FLinearColor( 0.5f, 0.5f, 0.5f ) ) );
-		Set( "NotificationBar.CheckBox", NotificationBarCheckBoxCheckBoxStyle );
-
-		// Read-only checkbox that appears next to a menu item
-		const FCheckBoxStyle NotificationBarCheckCheckBoxStyle = FCheckBoxStyle()
-			.SetUncheckedImage( IMAGE_BRUSH( "Icons/Empty_14x", Icon14x14 ) )
-			.SetUncheckedPressedImage( FSlateNoResource() )
-			.SetUncheckedHoveredImage( FSlateNoResource() )
-			.SetCheckedHoveredImage( IMAGE_BRUSH( "Common/SmallCheck",  Icon14x14 ) )
-			.SetCheckedPressedImage( IMAGE_BRUSH( "Common/SmallCheck",  Icon14x14 ) )
-			.SetCheckedImage( IMAGE_BRUSH( "Common/SmallCheck",  Icon14x14 ) )
-			.SetUndeterminedImage( IMAGE_BRUSH( "Icons/Empty_14x", Icon14x14 ) )
-			.SetUndeterminedPressedImage( FSlateNoResource() )
-			.SetUndeterminedHoveredImage( FSlateNoResource() );
-		Set( "NotificationBar.Check", NotificationBarCheckCheckBoxStyle );
-
-		// This radio button is actually just a check box with different images
-		const FCheckBoxStyle NotificationBarRadioButtonCheckBoxStyle = FCheckBoxStyle()
-			.SetUncheckedImage( IMAGE_BRUSH( "Common/RadioButton_Unselected_16x", Icon16x16 ) )
-			.SetUncheckedPressedImage( IMAGE_BRUSH( "Common/RadioButton_Unselected_16x", Icon16x16, SelectionColor_Pressed ) )
-			.SetUncheckedHoveredImage( IMAGE_BRUSH( "Common/RadioButton_Unselected_16x", Icon16x16, SelectionColor ) )
-			.SetCheckedHoveredImage( IMAGE_BRUSH( "Common/RadioButton_Selected_16x", Icon16x16, SelectionColor ) )
-			.SetCheckedPressedImage( IMAGE_BRUSH( "Common/RadioButton_Selected_16x", Icon16x16, SelectionColor_Pressed ) )
-			.SetCheckedImage( IMAGE_BRUSH( "Common/RadioButton_Selected_16x",  Icon16x16 ) );
-		Set( "NotificationBar.RadioButton", NotificationBarRadioButtonCheckBoxStyle );
-
-		const FCheckBoxStyle NotificationBarToggleButtonCheckBoxStyle = FCheckBoxStyle()
-			.SetCheckBoxType( ESlateCheckBoxType::ToggleButton )
-			.SetUncheckedImage( FSlateNoResource() )
-			.SetUncheckedPressedImage( BOX_BRUSH( "Common/RoundedSelection_16x", 4.0f/16.0f, SelectionColor_Pressed ) )
-			.SetUncheckedHoveredImage( BOX_BRUSH( "Common/RoundedSelection_16x", 4.0f/16.0f, SelectionColor ) )
-			.SetCheckedHoveredImage( BOX_BRUSH( "Common/RoundedSelection_16x",  4.0f/16.0f, SelectionColor_Pressed ) )
-			.SetCheckedPressedImage( BOX_BRUSH( "Common/RoundedSelection_16x",  4.0f/16.0f, SelectionColor ) )
-			.SetCheckedImage( BOX_BRUSH( "Common/RoundedSelection_16x",  4.0f/16.0f, SelectionColor_Pressed ) );
-		Set( "NotificationBar.ToggleButton", NotificationBarToggleButtonCheckBoxStyle );
-
-		Set( "NotificationBar.Button", FButtonStyle( NoBorder )
-			.SetNormal ( FSlateNoResource() )
-			.SetPressed( BOX_BRUSH( "Common/RoundedSelection_16x", 4.0f/16.0f, SelectionColor_Pressed ) )
-			.SetHovered( BOX_BRUSH( "Common/RoundedSelection_16x", 4.0f/16.0f, SelectionColor ) )
-			.SetNormalPadding( FMargin(0,1) )
-			.SetPressedPadding( FMargin(0,2,0,0) )
-			);
-
-		Set( "NotificationBar.Button.Checked", new BOX_BRUSH( "Common/RoundedSelection_16x",  4.0f/16.0f, SelectionColor_Pressed ) );
-		Set( "NotificationBar.Button.Checked_Hovered", new BOX_BRUSH( "Common/RoundedSelection_16x",  4.0f/16.0f, SelectionColor_Pressed ) );
-		Set( "NotificationBar.Button.Checked_Pressed", new BOX_BRUSH( "Common/RoundedSelection_16x",  4.0f/16.0f, SelectionColor ) );
-
-		Set( "NotificationBar.SToolBarButtonBlock.CheckBox.Padding", FMargin(4.0f) );
-		Set( "NotificationBar.SToolBarButtonBlock.Button.Padding", FMargin(0.0f) );
-		Set( "NotificationBar.SToolBarComboButtonBlock.ComboButton.Color", DefaultForeground );
-	}
-
-	// Viewport ToolbarBar
-	{
-		Set( "ViewportMenu.Background",	new BOX_BRUSH( "Old/Menu_Background", FMargin(8.0f/64.0f), FLinearColor::Transparent ) );
-		Set( "ViewportMenu.Icon", new IMAGE_BRUSH( "Icons/icon_tab_toolbar_16px", Icon16x16 ) );
-		Set( "ViewportMenu.Expand", new IMAGE_BRUSH( "Icons/toolbar_expand_16x", Icon8x8) );
-		Set( "ViewportMenu.SubMenuIndicator", new IMAGE_BRUSH( "Common/SubmenuArrow", Icon8x8) );
-		Set( "ViewportMenu.SToolBarComboButtonBlock.Padding", FMargin(0));
-		Set( "ViewportMenu.SToolBarButtonBlock.Padding", FMargin(0));
-		Set( "ViewportMenu.SToolBarButtonBlock.Button.Padding", FMargin(0));
-		Set( "ViewportMenu.SToolBarCheckComboButtonBlock.Padding", FMargin(0));
-		Set( "ViewportMenu.SToolBarButtonBlock.CheckBox.Padding", FMargin(4.0f) );
-		Set( "ViewportMenu.SToolBarComboButtonBlock.ComboButton.Color", FLinearColor(0.f,0.f,0.f,0.75f) );
-
-		Set( "ViewportMenu.Separator", new BOX_BRUSH( "Old/Button", 8.0f/32.0f, FLinearColor::Transparent ) );
-		Set( "ViewportMenu.Separator.Padding", FMargin( 100.0f ) );
-
-		Set( "ViewportMenu.Label", FTextBlockStyle(NormalText)
-			.SetFont(DEFAULT_FONT("Bold", 9))
-			.SetColorAndOpacity(FLinearColor(0.0f, 0.0f, 0.0f, 1.0f)));
-		Set( "ViewportMenu.Label.Padding", FMargin(0.0f, 0.0f, 3.0f, 0.0f) );
-		Set( "ViewportMenu.Label.ContentPadding", FMargin(5.0f, 2.0f) );
-		Set( "ViewportMenu.EditableText", FEditableTextBoxStyle(NormalEditableTextBoxStyle) .SetFont( DEFAULT_FONT( "Regular", 9) ) );
-		Set( "ViewportMenu.Keybinding", FTextBlockStyle(NormalText) .SetFont( DEFAULT_FONT( "Regular", 8) ) );
-
-		Set( "ViewportMenu.Block.IndentedPadding", FMargin( 0 ) );
-		Set( "ViewportMenu.Block.Padding", FMargin( 0 ) );
-
-		Set( "ViewportMenu.Heading.Font", DEFAULT_FONT( "Regular", 8 ) );
-		Set( "ViewportMenu.Heading.ColorAndOpacity", FLinearColor( 0.4f, 0.4, 0.4f, 1.0f ) );
-
-		const FCheckBoxStyle ViewportMenuCheckBoxCheckBoxStyle = FCheckBoxStyle()
-			.SetUncheckedImage(			IMAGE_BRUSH( "Common/SmallCheckBox", Icon14x14 ) )
-			.SetUncheckedPressedImage(	IMAGE_BRUSH( "Common/SmallCheckBox_Hovered", Icon14x14, FLinearColor( 0.5f, 0.5f, 0.5f ) ) )
-			.SetUncheckedHoveredImage(	IMAGE_BRUSH( "Common/SmallCheckBox_Hovered", Icon14x14 ) )
-			.SetCheckedHoveredImage(	IMAGE_BRUSH( "Common/SmallCheckBox_Checked_Hovered", Icon14x14 ) )
-			.SetCheckedPressedImage(	IMAGE_BRUSH( "Common/SmallCheckBox_Checked_Hovered", Icon14x14, FLinearColor( 0.5f, 0.5f, 0.5f ) ) )
-			.SetCheckedImage(			IMAGE_BRUSH( "Common/SmallCheckBox_Checked",  Icon14x14 ) );
-		Set( "ViewportMenu.CheckBox", ViewportMenuCheckBoxCheckBoxStyle );
-
-		// Read-only checkbox that appears next to a menu item
-		const FCheckBoxStyle ViewportMenuCheckCheckBoxStyle = FCheckBoxStyle()
-			.SetUncheckedImage(			IMAGE_BRUSH( "Icons/Empty_14x", Icon14x14 ) )
-			.SetUncheckedPressedImage(	FSlateNoResource() )
-			.SetUncheckedHoveredImage(	FSlateNoResource() )
-			.SetCheckedHoveredImage(	IMAGE_BRUSH( "Common/SmallCheck",  Icon14x14 ) )
-			.SetCheckedPressedImage(	IMAGE_BRUSH( "Common/SmallCheck",  Icon14x14 ) )
-			.SetCheckedImage(			IMAGE_BRUSH( "Common/SmallCheck",  Icon14x14 ) );
-		Set( "ViewportMenu.Check", ViewportMenuCheckCheckBoxStyle );
-
-		const FString SmallRoundedButton			(TEXT("Common/SmallRoundedButton"));
-		const FString SmallRoundedButtonStart		(TEXT("Common/SmallRoundedButtonLeft"));
-		const FString SmallRoundedButtonMiddle	(TEXT("Common/SmallRoundedButtonCentre"));
-		const FString SmallRoundedButtonEnd		(TEXT("Common/SmallRoundedButtonRight"));
-
-		const FLinearColor NormalColor(1,1,1,0.75f);
-		const FLinearColor PressedColor(1,1,1,1.f);
-
-		const FCheckBoxStyle ViewportMenuRadioButtonCheckBoxStyle = FCheckBoxStyle()
-			.SetUncheckedImage(			IMAGE_BRUSH( "Common/MenuItemRadioButton_Off", Icon14x14 ) )
-			.SetUncheckedPressedImage(	IMAGE_BRUSH( "Common/MenuItemRadioButton_Off", Icon14x14, FLinearColor( 0.5f, 0.5f, 0.5f ) ) )
-			.SetUncheckedHoveredImage(	IMAGE_BRUSH( "Common/MenuItemRadioButton_Off", Icon14x14 ) )
-			.SetCheckedHoveredImage(	IMAGE_BRUSH( "Common/MenuItemRadioButton_On", Icon14x14 ) )
-			.SetCheckedPressedImage(	IMAGE_BRUSH( "Common/MenuItemRadioButton_On_Pressed", Icon14x14 ) )
-			.SetCheckedImage(			IMAGE_BRUSH( "Common/MenuItemRadioButton_On",  Icon14x14 ) );
-		Set( "ViewportMenu.RadioButton", ViewportMenuRadioButtonCheckBoxStyle );
-
-		/* Create style for "ViewportMenu.ToggleButton" ... */
-		const FCheckBoxStyle ViewportMenuToggleButtonStyle = FCheckBoxStyle()
-			.SetCheckBoxType(			ESlateCheckBoxType::ToggleButton )
-			.SetUncheckedImage(			BOX_BRUSH( *SmallRoundedButton, FMargin(7.f/16.f), NormalColor ) )
-			.SetUncheckedPressedImage(	BOX_BRUSH( *SmallRoundedButton, FMargin(7.f/16.f), PressedColor ) )
-			.SetUncheckedHoveredImage(	BOX_BRUSH( *SmallRoundedButton, FMargin(7.f/16.f), PressedColor ) )
-			.SetCheckedHoveredImage(	BOX_BRUSH( *SmallRoundedButton, FMargin(7.f/16.f), SelectionColor_Pressed ) )
-			.SetCheckedPressedImage(	BOX_BRUSH( *SmallRoundedButton, FMargin(7.f/16.f), SelectionColor_Pressed )  )
-			.SetCheckedImage(			BOX_BRUSH( *SmallRoundedButton, FMargin(7.f/16.f), SelectionColor_Pressed ) );
-		/* ... and add new style */
-		Set( "ViewportMenu.ToggleButton", ViewportMenuToggleButtonStyle );
-
-		/* Create style for "ViewportMenu.ToggleButton.Start" ... */
-		const FCheckBoxStyle ViewportMenuToggleStartButtonStyle = FCheckBoxStyle()
-			.SetCheckBoxType(			ESlateCheckBoxType::ToggleButton )
-			.SetUncheckedImage(			BOX_BRUSH( *SmallRoundedButtonStart, FMargin(7.f/16.f), NormalColor ) )
-			.SetUncheckedPressedImage(	BOX_BRUSH( *SmallRoundedButtonStart, FMargin(7.f/16.f), PressedColor ) )
-			.SetUncheckedHoveredImage(	BOX_BRUSH( *SmallRoundedButtonStart, FMargin(7.f/16.f), PressedColor ) )
-			.SetCheckedHoveredImage(	BOX_BRUSH( *SmallRoundedButtonStart, FMargin(7.f/16.f), SelectionColor_Pressed ) )
-			.SetCheckedPressedImage(	BOX_BRUSH( *SmallRoundedButtonStart, FMargin(7.f/16.f), SelectionColor_Pressed )  )
-			.SetCheckedImage(			BOX_BRUSH( *SmallRoundedButtonStart, FMargin(7.f/16.f), SelectionColor_Pressed ) );
-		/* ... and add new style */
-		Set( "ViewportMenu.ToggleButton.Start", ViewportMenuToggleStartButtonStyle );
-
-		/* Create style for "ViewportMenu.ToggleButton.Middle" ... */
-		const FCheckBoxStyle ViewportMenuToggleMiddleButtonStyle = FCheckBoxStyle()
-			.SetCheckBoxType(			ESlateCheckBoxType::ToggleButton )
-			.SetUncheckedImage(			BOX_BRUSH( *SmallRoundedButtonMiddle, FMargin(7.f/16.f), NormalColor ) )
-			.SetUncheckedPressedImage(	BOX_BRUSH( *SmallRoundedButtonMiddle, FMargin(7.f/16.f), PressedColor ) )
-			.SetUncheckedHoveredImage(	BOX_BRUSH( *SmallRoundedButtonMiddle, FMargin(7.f/16.f), PressedColor ) )
-			.SetCheckedHoveredImage(	BOX_BRUSH( *SmallRoundedButtonMiddle, FMargin(7.f/16.f), SelectionColor_Pressed ) )
-			.SetCheckedPressedImage(	BOX_BRUSH( *SmallRoundedButtonMiddle, FMargin(7.f/16.f), SelectionColor_Pressed )  )
-			.SetCheckedImage(			BOX_BRUSH( *SmallRoundedButtonMiddle, FMargin(7.f/16.f), SelectionColor_Pressed ) );
-		/* ... and add new style */
-		Set( "ViewportMenu.ToggleButton.Middle", ViewportMenuToggleMiddleButtonStyle );
-
-		/* Create style for "ViewportMenu.ToggleButton.End" ... */
-		const FCheckBoxStyle ViewportMenuToggleEndButtonStyle = FCheckBoxStyle()
-			.SetCheckBoxType(			ESlateCheckBoxType::ToggleButton )
-			.SetUncheckedImage(			BOX_BRUSH( *SmallRoundedButtonEnd, FMargin(7.f/16.f), NormalColor ) )
-			.SetUncheckedPressedImage(	BOX_BRUSH( *SmallRoundedButtonEnd, FMargin(7.f/16.f), PressedColor ) )
-			.SetUncheckedHoveredImage(	BOX_BRUSH( *SmallRoundedButtonEnd, FMargin(7.f/16.f), PressedColor ) )
-			.SetCheckedHoveredImage(	BOX_BRUSH( *SmallRoundedButtonEnd,  FMargin(7.f/16.f), SelectionColor_Pressed ) )
-			.SetCheckedPressedImage(	BOX_BRUSH( *SmallRoundedButtonEnd,  FMargin(7.f/16.f), SelectionColor_Pressed )  )
-			.SetCheckedImage(			BOX_BRUSH( *SmallRoundedButtonEnd,  FMargin(7.f/16.f), SelectionColor_Pressed ) );
-		/* ... and add new style */
-		Set( "ViewportMenu.ToggleButton.End", ViewportMenuToggleEndButtonStyle );
-
-		const FMargin NormalPadding = FMargin(4.0f, 4.0f, 4.0f, 4.0f);
-		const FMargin PressedPadding = FMargin(4.0f, 4.0f, 4.0f, 4.0f);
-
-		const FButtonStyle ViewportMenuButton = FButtonStyle(Button)
-			.SetNormal ( BOX_BRUSH( *SmallRoundedButton, 7.0f/16.0f, NormalColor))
-			.SetPressed( BOX_BRUSH( *SmallRoundedButton, 7.0f/16.0f, PressedColor ) )
-			.SetHovered(BOX_BRUSH(*SmallRoundedButton, 7.0f / 16.0f, PressedColor))
-			.SetPressedPadding(PressedPadding)
-			.SetNormalPadding(NormalPadding);
-
-		Set( "ViewportMenu.Button", ViewportMenuButton );
-
-		Set( "ViewportMenu.Button.Start", FButtonStyle(ViewportMenuButton)
-			.SetNormal ( BOX_BRUSH( *SmallRoundedButtonStart, 7.0f/16.0f, NormalColor))
-			.SetPressed( BOX_BRUSH( *SmallRoundedButtonStart, 7.0f/16.0f, PressedColor ) )
-			.SetHovered( BOX_BRUSH( *SmallRoundedButtonStart, 7.0f/16.0f, PressedColor ) )
-			);
-
-		Set( "ViewportMenu.Button.Middle", FButtonStyle(ViewportMenuButton)
-			.SetNormal ( BOX_BRUSH( *SmallRoundedButtonMiddle, 7.0f/16.0f, NormalColor))
-			.SetPressed( BOX_BRUSH( *SmallRoundedButtonMiddle, 7.0f/16.0f, PressedColor ) )
-			.SetHovered( BOX_BRUSH( *SmallRoundedButtonMiddle, 7.0f/16.0f, PressedColor ) )
-			);
-
-		Set( "ViewportMenu.Button.End", FButtonStyle(ViewportMenuButton)
-			.SetNormal ( BOX_BRUSH( *SmallRoundedButtonEnd, 7.0f/16.0f, NormalColor))
-			.SetPressed( BOX_BRUSH( *SmallRoundedButtonEnd, 7.0f/16.0f, PressedColor ) )
-			.SetHovered( BOX_BRUSH( *SmallRoundedButtonEnd, 7.0f/16.0f, PressedColor ) )
-			);
-	}
-
-	// Viewport actor preview's pin/unpin buttons
-	{
-		Set( "ViewportActorPreview.Pinned", new IMAGE_BRUSH( "Common/PushPin_Down", Icon16x16 ) );
-		Set( "ViewportActorPreview.Unpinned", new IMAGE_BRUSH( "Common/PushPin_Up", Icon16x16 ) );
-	}
+	
 
 	// Standard Dialog Settings
 	{
@@ -2911,6 +2241,729 @@ void FSlateEditorStyle::FStyle::SetupGeneralStyles()
 #endif // WITH_EDITOR || (IS_PROGRAM && WITH_UNREAL_DEVELOPER_TOOLS)
 }
 
+void FSlateEditorStyle::FStyle::SetupLevelGeneralStyles()
+{
+// Levels General
+	{
+	Set("Level.VisibleIcon16x", new IMAGE_BRUSH("Icons/icon_levels_visible_16px", Icon16x16));
+	Set("Level.VisibleHighlightIcon16x", new IMAGE_BRUSH("Icons/icon_levels_visible_hi_16px", Icon16x16));
+	Set("Level.NotVisibleIcon16x", new IMAGE_BRUSH("Icons/icon_levels_invisible_16px", Icon16x16));
+	Set("Level.NotVisibleHighlightIcon16x", new IMAGE_BRUSH("Icons/icon_levels_invisible_hi_16px", Icon16x16));
+	Set("Level.LightingScenarioIcon16x", new IMAGE_BRUSH("Icons/icon_levels_LightingScenario_16px", Icon16x16));
+	Set("Level.LightingScenarioNotIcon16x", new IMAGE_BRUSH("Icons/icon_levels_LightingScenarioNot_16px", Icon16x16));
+	Set("Level.LockedIcon16x", new IMAGE_BRUSH("Icons/icon_locked_16px", Icon16x16));
+	Set("Level.LockedHighlightIcon16x", new IMAGE_BRUSH("Icons/icon_locked_highlight_16px", Icon16x16));
+	Set("Level.UnlockedIcon16x", new IMAGE_BRUSH("Icons/icon_levels_unlocked_16px", Icon16x16));
+	Set("Level.UnlockedHighlightIcon16x", new IMAGE_BRUSH("Icons/icon_levels_unlocked_hi_16px", Icon16x16));
+	Set("Level.ReadOnlyLockedIcon16x", new IMAGE_BRUSH("Icons/icon_levels_LockedReadOnly_16px", Icon16x16));
+	Set("Level.ReadOnlyLockedHighlightIcon16x", new IMAGE_BRUSH("Icons/icon_levels_LockedReadOnly_hi_16px", Icon16x16));
+	Set("Level.SaveIcon16x", new IMAGE_BRUSH("Icons/icon_levels_Save_16px", Icon16x16));
+	Set("Level.SaveHighlightIcon16x", new IMAGE_BRUSH("Icons/icon_levels_Save_hi_16px", Icon16x16));
+	Set("Level.SaveModifiedIcon16x", new IMAGE_BRUSH("Icons/icon_levels_SaveModified_16px", Icon16x16));
+	Set("Level.SaveModifiedHighlightIcon16x", new IMAGE_BRUSH("Icons/icon_levels_SaveModified_hi_16px", Icon16x16));
+	Set("Level.SaveDisabledIcon16x", new IMAGE_BRUSH("Icons/icon_levels_SaveDisabled_16px", Icon16x16));
+	Set("Level.SaveDisabledHighlightIcon16x", new IMAGE_BRUSH("Icons/icon_levels_SaveDisabled_hi_16px", Icon16x16));
+	Set("Level.ScriptIcon16x", new IMAGE_BRUSH("Icons/icon_levels_Blueprint_16px", Icon16x16));
+	Set("Level.ScriptHighlightIcon16x", new IMAGE_BRUSH("Icons/icon_levels_Blueprint_hi_16px", Icon16x16));
+	Set("Level.EmptyIcon16x", new IMAGE_BRUSH("Icons/Empty_16x", Icon16x16));
+	Set("Level.ColorIcon40x", new IMAGE_BRUSH("Icons/icon_levels_back_16px", Icon16x16));
+	}
+}
+
+void FSlateEditorStyle::FStyle::SetupWorldBrowserStyles()
+{
+
+	// World Browser
+	{
+		Set("WorldBrowser.AddLayer", new IMAGE_BRUSH("Icons/icon_levels_addlayer_16x", Icon16x16));
+		Set("WorldBrowser.SimulationViewPositon", new IMAGE_BRUSH("Icons/icon_levels_simulationviewpos_16x", Icon16x16));
+		Set("WorldBrowser.MouseLocation", new IMAGE_BRUSH("Icons/icon_levels_mouselocation_16x", Icon16x16));
+		Set("WorldBrowser.MarqueeRectSize", new IMAGE_BRUSH("Icons/icon_levels_marqueerectsize_16x", Icon16x16));
+		Set("WorldBrowser.WorldSize", new IMAGE_BRUSH("Icons/icon_levels_worldsize_16x", Icon16x16));
+		Set("WorldBrowser.WorldOrigin", new IMAGE_BRUSH("Icons/icon_levels_worldorigin_16x", Icon16x16));
+		Set("WorldBrowser.DirectionXPositive", new IMAGE_BRUSH("Icons/icon_PanRight", Icon16x16));
+		Set("WorldBrowser.DirectionXNegative", new IMAGE_BRUSH("Icons/icon_PanLeft", Icon16x16));
+		Set("WorldBrowser.DirectionYPositive", new IMAGE_BRUSH("Icons/icon_PanUp", Icon16x16));
+		Set("WorldBrowser.DirectionYNegative", new IMAGE_BRUSH("Icons/icon_PanDown", Icon16x16));
+		Set("WorldBrowser.LevelStreamingAlwaysLoaded", new FSlateNoResource());
+		Set("WorldBrowser.LevelStreamingBlueprint", new IMAGE_BRUSH("Icons/icon_levels_blueprinttype_7x16", Icon7x16));
+		Set("WorldBrowser.LevelsMenuBrush", new IMAGE_BRUSH("Icons/icon_levels_levelsmenu_40x", Icon25x25));
+		Set("WorldBrowser.HierarchyButtonBrush", new IMAGE_BRUSH("Icons/icon_levels_hierarchybutton_16x", Icon16x16));
+		Set("WorldBrowser.DetailsButtonBrush", new IMAGE_BRUSH("Icons/icon_levels_detailsbutton_40x", Icon16x16));
+		Set("WorldBrowser.CompositionButtonBrush", new IMAGE_BRUSH("Icons/icon_levels_compositionbutton_16x", Icon16x16));
+
+		Set("WorldBrowser.FolderClosed", new IMAGE_BRUSH("Icons/FolderClosed", Icon16x16));
+		Set("WorldBrowser.FolderOpen", new IMAGE_BRUSH("Icons/FolderOpen", Icon16x16));
+		Set("WorldBrowser.NewFolderIcon", new IMAGE_BRUSH("Icons/icon_AddFolder_16x", Icon16x16));
+
+		Set("WorldBrowser.StatusBarText", FTextBlockStyle(NormalText)
+			.SetFont(DEFAULT_FONT("BoldCondensed", 12))
+			.SetColorAndOpacity(FLinearColor(0.9, 0.9f, 0.9f, 0.5f))
+			.SetShadowOffset(FVector2D::ZeroVector)
+		);
+
+		Set("WorldBrowser.LabelFont", DEFAULT_FONT("Regular", 9));
+		Set("WorldBrowser.LabelFontBold", DEFAULT_FONT("Bold", 10));
+	}
+}
+
+void FSlateEditorStyle::FStyle::SetupSequencerStyles()
+{
+	// Sequencer
+	if (IncludeEditorSpecificStyles())
+	{
+		Set("Sequencer.IconKeyAuto", new IMAGE_BRUSH("Sequencer/IconKeyAuto", Icon12x12));
+		Set("Sequencer.IconKeyBreak", new IMAGE_BRUSH("Sequencer/IconKeyBreak", Icon12x12));
+		Set("Sequencer.IconKeyConstant", new IMAGE_BRUSH("Sequencer/IconKeyConstant", Icon12x12));
+		Set("Sequencer.IconKeyLinear", new IMAGE_BRUSH("Sequencer/IconKeyLinear", Icon12x12));
+		Set("Sequencer.IconKeyUser", new IMAGE_BRUSH("Sequencer/IconKeyUser", Icon12x12));
+
+		Set("Sequencer.KeyCircle", new IMAGE_BRUSH("Sequencer/KeyCircle", Icon12x12));
+		Set("Sequencer.KeyDiamond", new IMAGE_BRUSH("Sequencer/KeyDiamond", Icon12x12));
+		Set("Sequencer.KeyDiamondBorder", new IMAGE_BRUSH("Sequencer/KeyDiamondBorder", Icon12x12));
+		Set("Sequencer.KeySquare", new IMAGE_BRUSH("Sequencer/KeySquare", Icon12x12));
+		Set("Sequencer.KeyTriangle", new IMAGE_BRUSH("Sequencer/KeyTriangle", Icon12x12));
+		Set("Sequencer.KeyLeft", new IMAGE_BRUSH("Sequencer/KeyLeft", Icon12x12));
+		Set("Sequencer.KeyRight", new IMAGE_BRUSH("Sequencer/KeyRight", Icon12x12));
+		Set("Sequencer.PartialKey", new IMAGE_BRUSH("Sequencer/PartialKey", FVector2D(11.f, 11.f)));
+		Set("Sequencer.Star", new IMAGE_BRUSH("Sequencer/Star", Icon12x12));
+		Set("Sequencer.Empty", new IMAGE_BRUSH("Sequencer/Empty", Icon12x12));
+		Set("Sequencer.TangentHandle", new IMAGE_BRUSH("Sequencer/TangentHandle", FVector2D(7, 7)));
+		Set("Sequencer.GenericDivider", new IMAGE_BRUSH("Sequencer/GenericDivider", FVector2D(2.f, 2.f), FLinearColor::White, ESlateBrushTileType::Vertical));
+
+		Set("Sequencer.Timeline.ScrubHandleDown", new BOX_BRUSH("Sequencer/ScrubHandleDown", FMargin(6.f / 13.f, 5 / 12.f, 6 / 13.f, 8 / 12.f)));
+		Set("Sequencer.Timeline.ScrubHandleUp", new BOX_BRUSH("Sequencer/ScrubHandleUp", FMargin(6.f / 13.f, 8 / 12.f, 6 / 13.f, 5 / 12.f)));
+		Set("Sequencer.Timeline.ScrubFill", new BOX_BRUSH("Sequencer/ScrubFill", FMargin(2.f / 4.f, 0.f)));
+		Set("Sequencer.Timeline.FrameBlockScrubHandleDown", new BOX_BRUSH("Sequencer/ScrubHandleDown", FMargin(6.f / 13.f, 5 / 12.f, 6 / 13.f, 8 / 12.f)));
+		Set("Sequencer.Timeline.FrameBlockScrubHandleUp", new BOX_BRUSH("Sequencer/ScrubHandleUp", FMargin(6.f / 13.f, 8 / 12.f, 6 / 13.f, 5 / 12.f)));
+		Set("Sequencer.Timeline.VanillaScrubHandleDown", new BOX_BRUSH("Sequencer/ScrubHandleDown_Clamped", FMargin(6.f / 13.f, 3.f / 12.f, 6.f / 13.f, 7.f / 12.f)));
+		Set("Sequencer.Timeline.VanillaScrubHandleUp", new BOX_BRUSH("Sequencer/ScrubHandleUp_Clamped", FMargin(6.f / 13.f, 8 / 12.f, 6 / 13.f, 5 / 12.f)));
+		Set("Sequencer.Timeline.ScrubHandleWhole", new BOX_BRUSH("Sequencer/ScrubHandleWhole", FMargin(6.f / 13.f, 10 / 24.f, 6 / 13.f, 10 / 24.f)));
+		Set("Sequencer.Timeline.RangeHandleLeft", new BOX_BRUSH("Sequencer/GenericGripLeft", FMargin(5.f / 16.f)));
+		Set("Sequencer.Timeline.RangeHandleRight", new BOX_BRUSH("Sequencer/GenericGripRight", FMargin(5.f / 16.f)));
+		Set("Sequencer.Timeline.RangeHandle", new BOX_BRUSH("Sequencer/GenericSectionBackground", FMargin(5.f / 16.f)));
+		Set("Sequencer.Timeline.NotifyAlignmentMarker", new IMAGE_BRUSH("Sequencer/NotifyAlignmentMarker", FVector2D(10, 19)));
+		Set("Sequencer.Timeline.PlayRange_Top_L", new BOX_BRUSH("Sequencer/PlayRange_Top_L", FMargin(1.f, 0.5f, 0.f, 0.5f)));
+		Set("Sequencer.Timeline.PlayRange_Top_R", new BOX_BRUSH("Sequencer/PlayRange_Top_R", FMargin(0.f, 0.5f, 1.f, 0.5f)));
+		Set("Sequencer.Timeline.PlayRange_L", new BOX_BRUSH("Sequencer/PlayRange_L", FMargin(1.f, 0.5f, 0.f, 0.5f)));
+		Set("Sequencer.Timeline.PlayRange_R", new BOX_BRUSH("Sequencer/PlayRange_R", FMargin(0.f, 0.5f, 1.f, 0.5f)));
+		Set("Sequencer.Timeline.PlayRange_Bottom_L", new BOX_BRUSH("Sequencer/PlayRange_Bottom_L", FMargin(1.f, 0.5f, 0.f, 0.5f)));
+		Set("Sequencer.Timeline.PlayRange_Bottom_R", new BOX_BRUSH("Sequencer/PlayRange_Bottom_R", FMargin(0.f, 0.5f, 1.f, 0.5f)));
+
+		Set("Sequencer.Timeline.SubSequenceRangeHashL", new BORDER_BRUSH("Sequencer/SubSequenceRangeHashL", FMargin(1.f, 0.f, 0.f, 0.f)));
+		Set("Sequencer.Timeline.SubSequenceRangeHashR", new BORDER_BRUSH("Sequencer/SubSequenceRangeHashR", FMargin(1.f, 0.f, 0.f, 0.f)));
+		Set("Sequencer.Timeline.EaseInOut", new IMAGE_BRUSH("Sequencer/EaseInOut", FVector2D(128, 128)));
+		Set("Sequencer.InterpLine", new BOX_BRUSH("Sequencer/InterpLine", FMargin(5.f / 7.f, 0.f, 0.f, 0.f)));
+
+		Set("Sequencer.Transport.JumpToPreviousKey", FButtonStyle()
+			.SetNormal(IMAGE_BRUSH("/Sequencer/Transport_Bar/Previous_Frame_OFF", Icon24x24))
+			.SetPressed(IMAGE_BRUSH("/Sequencer/Transport_Bar/Previous_Frame", Icon24x24))
+			.SetHovered(IMAGE_BRUSH("/Sequencer/Transport_Bar/Previous_Frame_OFF", Icon24x24)));
+		Set("Sequencer.Transport.JumpToNextKey", FButtonStyle()
+			.SetNormal(IMAGE_BRUSH("/Sequencer/Transport_Bar/Next_Frame_24x_OFF", Icon24x24))
+			.SetPressed(IMAGE_BRUSH("/Sequencer/Transport_Bar/Next_Frame_24x", Icon24x24))
+			.SetHovered(IMAGE_BRUSH("/Sequencer/Transport_Bar/Next_Frame_24x_OFF", Icon24x24)));
+		Set("Sequencer.Transport.SetPlayStart", FButtonStyle()
+			.SetNormal(IMAGE_BRUSH("/Sequencer/Transport_Bar/Bracket_In_16x24_OFF", FVector2D(16, 24)))
+			.SetPressed(IMAGE_BRUSH("/Sequencer/Transport_Bar/Bracket_In_16x24", FVector2D(16, 24)))
+			.SetHovered(IMAGE_BRUSH("/Sequencer/Transport_Bar/Bracket_In_16x24_OFF", FVector2D(16, 24))));
+		Set("Sequencer.Transport.SetPlayEnd", FButtonStyle()
+			.SetNormal(IMAGE_BRUSH("/Sequencer/Transport_Bar/Bracket_Out_16x24_OFF", FVector2D(16, 24)))
+			.SetPressed(IMAGE_BRUSH("/Sequencer/Transport_Bar/Bracket_Out_16x24", FVector2D(16, 24)))
+			.SetHovered(IMAGE_BRUSH("/Sequencer/Transport_Bar/Bracket_Out_16x24_OFF", FVector2D(16, 24))));
+
+		Set("Sequencer.Transport.CloseButton", FButtonStyle()
+			.SetNormal(IMAGE_BRUSH("/Docking/CloseApp_Normal", Icon16x16))
+			.SetPressed(IMAGE_BRUSH("/Docking/CloseApp_Pressed", Icon16x16))
+			.SetHovered(IMAGE_BRUSH("/Docking/CloseApp_Hovered", Icon16x16)));
+
+		Set("Sequencer.NotificationImage_AddedPlayMovieSceneEvent", new IMAGE_BRUSH("Old/Checkbox_checked", Icon16x16));
+
+		Set("Sequencer.Save", new IMAGE_BRUSH("Sequencer/Main_Icons/Icon_Sequencer_Save_24x", Icon48x48));
+		Set("Sequencer.Save.Small", new IMAGE_BRUSH("Sequencer/Main_Icons/Icon_Sequencer_Save_24x", Icon24x24));
+		Set("Sequencer.SaveAs", new IMAGE_BRUSH("Sequencer/Main_Icons/Icon_Sequencer_Save_As_24x", Icon48x48));
+		Set("Sequencer.SaveAs.Small", new IMAGE_BRUSH("Sequencer/Main_Icons/Icon_Sequencer_Save_As_24x", Icon24x24));
+		Set("Sequencer.DiscardChanges", new IMAGE_BRUSH("Sequencer/Main_Icons/Icon_Sequencer_Revert_24x", Icon48x48));
+		Set("Sequencer.DiscardChanges.Small", new IMAGE_BRUSH("Sequencer/Main_Icons/Icon_Sequencer_Revert_24x", Icon24x24));
+		Set("Sequencer.RestoreAnimatedState", new IMAGE_BRUSH("Sequencer/Main_Icons/Icon_Sequencer_RestoreAnimatedState_24x", Icon48x48));
+		Set("Sequencer.RestoreAnimatedState.Small", new IMAGE_BRUSH("Sequencer/Main_Icons/Icon_Sequencer_RestoreAnimatedState_24x", Icon24x24));
+		Set("Sequencer.GenericGripLeft", new BOX_BRUSH("Sequencer/GenericGripLeft", FMargin(5.f / 16.f)));
+		Set("Sequencer.GenericGripRight", new BOX_BRUSH("Sequencer/GenericGripRight", FMargin(5.f / 16.f)));
+		Set("Sequencer.SectionArea.Background", new FSlateColorBrush(FColor::White));
+
+		Set("Sequencer.Section.Background", new BORDER_BRUSH(TEXT("Sequencer/SectionBackground"), FMargin(4.f / 16.f)));
+		Set("Sequencer.Section.BackgroundTint", new BOX_BRUSH(TEXT("Sequencer/SectionBackgroundTint"), FMargin(4 / 16.f)));
+		Set("Sequencer.Section.SelectedSectionOverlay", new IMAGE_BRUSH(TEXT("Sequencer/SelectedSectionOverlay"), Icon16x16, FLinearColor::White, ESlateBrushTileType::Both));
+		Set("Sequencer.Section.SelectedTrackTint", new BOX_BRUSH(TEXT("Sequencer/SelectedTrackTint"), FMargin(0.f, 0.5f)));
+		Set("Sequencer.Section.SelectionBorder", new BORDER_BRUSH(TEXT("Sequencer/SectionHighlight"), FMargin(7.f / 16.f)));
+		Set("Sequencer.Section.LockedBorder", new BORDER_BRUSH(TEXT("Sequencer/SectionLocked"), FMargin(7.f / 16.f)));
+		Set("Sequencer.Section.SelectedSectionOverlay", new IMAGE_BRUSH(TEXT("Sequencer/SelectedSectionOverlay"), Icon16x16, FLinearColor::White, ESlateBrushTileType::Both));
+		Set("Sequencer.Section.FilmBorder", new IMAGE_BRUSH(TEXT("Sequencer/SectionFilmBorder"), FVector2D(10, 7), FLinearColor::White, ESlateBrushTileType::Horizontal));
+		Set("Sequencer.Section.GripLeft", new BOX_BRUSH("Sequencer/SectionGripLeft", FMargin(5.f / 16.f)));
+		Set("Sequencer.Section.GripRight", new BOX_BRUSH("Sequencer/SectionGripRight", FMargin(5.f / 16.f)));
+		Set("Sequencer.Section.EasingHandle", new IMAGE_BRUSH("Sequencer/EasingHandle", FVector2D(10.f, 10.f)));
+
+		Set("Sequencer.Section.PreRoll", new BORDER_BRUSH(TEXT("Sequencer/PreRoll"), FMargin(0.f, .5f, 0.f, .5f)));
+
+		Set("Sequencer.Section.PinCusion", new IMAGE_BRUSH(TEXT("Sequencer/PinCusion"), Icon16x16, FLinearColor::White, ESlateBrushTileType::Both));
+		Set("Sequencer.Section.OverlapBorder", new BORDER_BRUSH(TEXT("Sequencer/OverlapBorder"), FMargin(1.f / 4.f, 0.f)));
+		Set("Sequencer.Section.StripeO.verlay", new BOX_BRUSH("Sequencer/SectionStripeOverlay", FMargin(0.f, .5f)));
+		Set("Sequencer.Section.BackgroundText", DEFAULT_FONT("Bold", 24));
+		Set("Sequencer.Section.EmptySpace", new BOX_BRUSH(TEXT("Sequencer/EmptySpace"), FMargin(0.f, 7.f / 14.f)));
+
+		Set("Sequencer.AnimationOutliner.ColorStrip", FButtonStyle()
+			.SetNormal(FSlateNoResource())
+			.SetHovered(FSlateNoResource())
+			.SetPressed(FSlateNoResource())
+			.SetNormalPadding(FMargin(0, 0, 0, 0))
+			.SetPressedPadding(FMargin(0, 0, 0, 0))
+		);
+
+		Set("Sequencer.AnimationOutliner.TopLevelBorder_Expanded", new BOX_BRUSH("Sequencer/TopLevelNodeBorder_Expanded", FMargin(4.0f / 16.0f)));
+		Set("Sequencer.AnimationOutliner.TopLevelBorder_Collapsed", new BOX_BRUSH("Sequencer/TopLevelNodeBorder_Collapsed", FMargin(4.0f / 16.0f)));
+		Set("Sequencer.AnimationOutliner.DefaultBorder", new FSlateColorBrush(FLinearColor::White));
+		Set("Sequencer.AnimationOutliner.TransparentBorder", new FSlateColorBrush(FLinearColor::Transparent));
+		Set("Sequencer.AnimationOutliner.BoldFont", DEFAULT_FONT("Bold", 11));
+		Set("Sequencer.AnimationOutliner.RegularFont", DEFAULT_FONT("Regular", 9));
+		Set("Sequencer.ShotFilter", new IMAGE_BRUSH("Sequencer/FilteredArea", FVector2D(74, 74), FLinearColor::White, ESlateBrushTileType::Both));
+		Set("Sequencer.KeyMark", new IMAGE_BRUSH("Sequencer/KeyMark", FVector2D(3, 21), FLinearColor::White, ESlateBrushTileType::NoTile));
+		Set("Sequencer.SetAutoKey", new IMAGE_BRUSH("Sequencer/Main_Icons/Icon_Sequencer_Auto_Key_24x", Icon48x48));
+		Set("Sequencer.SetAutoKey.Small", new IMAGE_BRUSH("Sequencer/Main_Icons/Icon_Sequencer_Auto_Key_24x", Icon24x24));
+		Set("Sequencer.SetAutoTrack", new IMAGE_BRUSH("Sequencer/Main_Icons/Icon_Sequencer_Auto_Track_24x", Icon48x48));
+		Set("Sequencer.SetAutoTrack.Small", new IMAGE_BRUSH("Sequencer/Main_Icons/Icon_Sequencer_Auto_Track_24x", Icon24x24));
+		Set("Sequencer.SetAutoChangeAll", new IMAGE_BRUSH("Sequencer/Main_Icons/Icon_Sequencer_Auto_Key_All_24x", Icon48x48));
+		Set("Sequencer.SetAutoChangeAll.Small", new IMAGE_BRUSH("Sequencer/Main_Icons/Icon_Sequencer_Auto_Key_All_24x", Icon24x24));
+		Set("Sequencer.SetAutoChangeNone", new IMAGE_BRUSH("Sequencer/Main_Icons/Icon_Sequencer_Disable_Auto_Key_24x", Icon48x48));
+		Set("Sequencer.SetAutoChangeNone.Small", new IMAGE_BRUSH("Sequencer/Main_Icons/Icon_Sequencer_Disable_Auto_Key_24x", Icon24x24));
+		Set("Sequencer.AllowAllEdits", new IMAGE_BRUSH("Sequencer/Main_Icons/Icon_Sequencer_Allow_All_Edits_24x", Icon48x48));
+		Set("Sequencer.AllowAllEdits.Small", new IMAGE_BRUSH("Sequencer/Main_Icons/Icon_Sequencer_Allow_All_Edits_24x", Icon24x24));
+		Set("Sequencer.AllowSequencerEditsOnly", new IMAGE_BRUSH("Sequencer/Main_Icons/Icon_Sequencer_Allow_Sequencer_Edits_Only_24x", Icon48x48));
+		Set("Sequencer.AllowSequencerEditsOnly.Small", new IMAGE_BRUSH("Sequencer/Main_Icons/Icon_Sequencer_Allow_Sequencer_Edits_Only_24x", Icon24x24));
+		Set("Sequencer.AllowLevelEditsOnly", new IMAGE_BRUSH("Sequencer/Main_Icons/Icon_Sequencer_Allow_Level_Edits_Only_24x", Icon48x48));
+		Set("Sequencer.AllowLevelEditsOnly.Small", new IMAGE_BRUSH("Sequencer/Main_Icons/Icon_Sequencer_Allow_Level_Edits_Only_24x", Icon24x24));
+		Set("Sequencer.SetKeyAll", new IMAGE_BRUSH("Sequencer/Main_Icons/Icon_Sequencer_Key_All_24x", Icon48x48));
+		Set("Sequencer.SetKeyAll.Small", new IMAGE_BRUSH("Sequencer/Main_Icons/Icon_Sequencer_Key_All_24x", Icon24x24));
+		Set("Sequencer.SetKeyGroup", new IMAGE_BRUSH("Sequencer/Main_Icons/Icon_Sequencer_Key_Group_24x", Icon48x48));
+		Set("Sequencer.SetKeyGroup.Small", new IMAGE_BRUSH("Sequencer/Main_Icons/Icon_Sequencer_Key_Group_24x", Icon24x24));
+		Set("Sequencer.SetKeyChanged", new IMAGE_BRUSH("Sequencer/Main_Icons/Icon_Sequencer_Key_Part_24x", Icon48x48));
+		Set("Sequencer.SetKeyChanged.Small", new IMAGE_BRUSH("Sequencer/Main_Icons/Icon_Sequencer_Key_Part_24x", Icon24x24));
+		Set("Sequencer.ToggleIsSnapEnabled", new IMAGE_BRUSH("Sequencer/Main_Icons/Icon_Sequencer_Snap_24x", Icon48x48));
+		Set("Sequencer.ToggleIsSnapEnabled.Small", new IMAGE_BRUSH("Sequencer/Main_Icons/Icon_Sequencer_Snap_24x", Icon24x24));
+		Set("Sequencer.ToggleShowCurveEditor", new IMAGE_BRUSH("Icons/SequencerIcons/icon_Sequencer_CurveEditor_24x", Icon48x48));
+		Set("Sequencer.ToggleShowCurveEditor.Small", new IMAGE_BRUSH("Icons/SequencerIcons/icon_Sequencer_CurveEditor_24x", Icon24x24));
+		Set("Sequencer.ToggleAutoScroll", new IMAGE_BRUSH("Icons/icon_Sequencer_ToggleAutoScroll_40x", Icon48x48));
+		Set("Sequencer.ToggleAutoScroll.Small", new IMAGE_BRUSH("Icons/icon_Sequencer_ToggleAutoScroll_16x", Icon16x16));
+		Set("Sequencer.MoveTool.Small", new IMAGE_BRUSH("Icons/SequencerIcons/icon_Sequencer_Move_24x", Icon16x16));
+		Set("Sequencer.MarqueeTool.Small", new IMAGE_BRUSH("Icons/SequencerIcons/icon_Sequencer_Marquee_24x", Icon16x16));
+		Set("Sequencer.RenderMovie.Small", new IMAGE_BRUSH("Sequencer/Main_Icons/Icon_Sequencer_Create_Movie_24x", Icon24x24));
+		Set("Sequencer.CreateCamera.Small", new IMAGE_BRUSH("Sequencer/Main_Icons/Icon_Sequencer_Create_Camera_24x", Icon24x24));
+		Set("Sequencer.FindInContentBrowser.Small", new IMAGE_BRUSH("Sequencer/Main_Icons/Icon_Sequencer_Find_In_Content_Browser_24x", Icon24x24));
+		Set("Sequencer.LockCamera", new IMAGE_BRUSH("Sequencer/Main_Icons/Icon_Sequencer_Look_Thru_24x", Icon16x16));
+		Set("Sequencer.UnlockCamera", new IMAGE_BRUSH("Sequencer/Main_Icons/Icon_Sequencer_Look_Thru_24x", Icon16x16, FLinearColor(1.f, 1.f, 1.f, 0.5f)));
+		Set("Sequencer.Thumbnail.SectionHandle", new IMAGE_BRUSH("Old/White", Icon16x16, FLinearColor::Black));
+		Set("Sequencer.TrackHoverHighlight_Top", new IMAGE_BRUSH(TEXT("Sequencer/TrackHoverHighlight_Top"), FVector2D(4, 4)));
+		Set("Sequencer.TrackHoverHighlight_Bottom", new IMAGE_BRUSH(TEXT("Sequencer/TrackHoverHighlight_Bottom"), FVector2D(4, 4)));
+		Set("Sequencer.SpawnableIconOverlay", new IMAGE_BRUSH(TEXT("Sequencer/SpawnableIconOverlay"), FVector2D(13, 13)));
+		Set("Sequencer.MultipleIconOverlay", new IMAGE_BRUSH(TEXT("Sequencer/MultipleIconOverlay"), FVector2D(13, 13)));
+		Set("Sequencer.LockSequence", new IMAGE_BRUSH("Sequencer/Main_Icons/Icon_Sequencer_Locked_16x", Icon16x16));
+		Set("Sequencer.UnlockSequence", new IMAGE_BRUSH("Sequencer/Main_Icons/Icon_Sequencer_Unlocked_16x", Icon16x16));
+
+		Set("Sequencer.GeneralOptions", new IMAGE_BRUSH("Sequencer/Main_Icons/Icon_Sequencer_General_Options_24x", Icon48x48));
+		Set("Sequencer.GeneralOptions.Small", new IMAGE_BRUSH("Sequencer/Main_Icons/Icon_Sequencer_General_Options_24x", Icon24x24));
+		Set("Sequencer.PlaybackOptions", new IMAGE_BRUSH("Sequencer/Main_Icons/Icon_Sequencer_Playback_Options_24x", Icon48x48));
+		Set("Sequencer.PlaybackOptions.Small", new IMAGE_BRUSH("Sequencer/Main_Icons/Icon_Sequencer_Playback_Options_24x", Icon24x24));
+		Set("Sequencer.SelectEditOptions", new IMAGE_BRUSH("Sequencer/Main_Icons/Icon_Sequencer_SelectEdit_Options_24x", Icon48x48));
+		Set("Sequencer.SelectEditOptions.Small", new IMAGE_BRUSH("Sequencer/Main_Icons/Icon_Sequencer_SelectEdit_Options_24x", Icon24x24));
+		Set("Sequencer.Time", new IMAGE_BRUSH("Sequencer/Main_Icons/Icon_Sequencer_Time_24x", Icon48x48));
+		Set("Sequencer.Time.Small", new IMAGE_BRUSH("Sequencer/Main_Icons/Icon_Sequencer_Time_24x", Icon24x24));
+		Set("Sequencer.Value", new IMAGE_BRUSH("Sequencer/Main_Icons/Icon_Sequencer_Value_24x", Icon48x48));
+		Set("Sequencer.Value.Small", new IMAGE_BRUSH("Sequencer/Main_Icons/Icon_Sequencer_Value_24x", Icon24x24));
+
+		Set("Sequencer.OverlayPanel.Background", new BOX_BRUSH("Sequencer/OverlayPanelBackground", FMargin(26.f / 54.f)));
+
+		Set("Sequencer.TrackArea.LaneColor", FLinearColor(0.3f, 0.3f, 0.3f, 0.3f));
+
+		Set("Sequencer.Tracks.Media", new IMAGE_BRUSH("Sequencer/Dropdown_Icons/Icon_Media_Track_16x", Icon16x16));
+		Set("Sequencer.Tracks.Audio", new IMAGE_BRUSH("Sequencer/Dropdown_Icons/Icon_Audio_Track_16x", Icon16x16));
+		Set("Sequencer.Tracks.Event", new IMAGE_BRUSH("Sequencer/Dropdown_Icons/Icon_Event_Track_16x", Icon16x16));
+		Set("Sequencer.Tracks.Fade", new IMAGE_BRUSH("Sequencer/Dropdown_Icons/Icon_Fade_Track_16x", Icon16x16));
+		Set("Sequencer.Tracks.CameraCut", new IMAGE_BRUSH("Sequencer/Dropdown_Icons/Icon_Camera_Cut_Track_16x", Icon16x16));
+		Set("Sequencer.Tracks.CinematicShot", new IMAGE_BRUSH("Sequencer/Dropdown_Icons/Icon_Shot_Track_16x", Icon16x16));
+		Set("Sequencer.Tracks.Slomo", new IMAGE_BRUSH("Sequencer/Dropdown_Icons/Icon_Play_Rate_Track_16x", Icon16x16));
+		Set("Sequencer.Tracks.Sub", new IMAGE_BRUSH("Sequencer/Dropdown_Icons/Icon_Sub_Track_16x", Icon16x16));
+		Set("Sequencer.Tracks.LevelVisibility", new IMAGE_BRUSH("Sequencer/Dropdown_Icons/Icon_Level_Visibility_Track_16x", Icon16x16));
+
+		Set("Sequencer.CursorDecorator_MarqueeAdd", new IMAGE_BRUSH("Sequencer/CursorDecorator_MarqueeAdd", Icon16x16));
+		Set("Sequencer.CursorDecorator_MarqueeSubtract", new IMAGE_BRUSH("Sequencer/CursorDecorator_MarqueeSubtract", Icon16x16));
+
+	Set("Sequencer.BreadcrumbText", FTextBlockStyle(NormalText)
+		.SetFont(DEFAULT_FONT("Bold", 11))
+		.SetColorAndOpacity(FLinearColor(1.0f, 1.0f, 1.0f))
+		.SetHighlightColor(FLinearColor(1.0f, 1.0f, 1.0f))
+		.SetShadowOffset(FVector2D(1, 1))
+		.SetShadowColorAndOpacity(FLinearColor(0, 0, 0, 0.9f)));
+	Set("Sequencer.BreadcrumbIcon", new IMAGE_BRUSH("Common/SmallArrowRight", Icon10x10));
+
+	const FButtonStyle NoBorder = FButtonStyle()
+	.SetNormal(FSlateNoResource())
+	.SetHovered(FSlateNoResource())
+	.SetPressed(FSlateNoResource())
+	.SetNormalPadding(FMargin(0, 0, 0, 1))
+	.SetPressedPadding(FMargin(0, 1, 0, 0));
+
+	const FButtonStyle DetailsKeyButton = FButtonStyle(NoBorder)
+	.SetNormal(IMAGE_BRUSH("Sequencer/AddKey_Details", FVector2D(11, 11)))
+	.SetHovered(IMAGE_BRUSH("Sequencer/AddKey_Details", FVector2D(11, 11), SelectionColor))
+	.SetPressed(IMAGE_BRUSH("Sequencer/AddKey_Details", FVector2D(11, 11), SelectionColor_Pressed))
+	.SetNormalPadding(FMargin(0, 1))
+	.SetPressedPadding(FMargin(0, 2, 0, 0));
+	Set("Sequencer.AddKey.Details", DetailsKeyButton);
+
+	const FSplitterStyle OutlinerSplitterStyle = FSplitterStyle()
+	.SetHandleNormalBrush(FSlateNoResource())
+	.SetHandleHighlightBrush(FSlateNoResource());
+	Set("Sequencer.AnimationOutliner.Splitter", OutlinerSplitterStyle);
+
+	Set("Sequencer.HyperlinkSpinBox", FSpinBoxStyle(GetWidgetStyle<FSpinBoxStyle>("SpinBox"))
+		.SetTextPadding(FMargin(0))
+		.SetBackgroundBrush(BORDER_BRUSH("Old/HyperlinkDotted", FMargin(0, 0, 0, 3 / 16.0f), FSlateColor::UseSubduedForeground()))
+		.SetHoveredBackgroundBrush(FSlateNoResource())
+		.SetInactiveFillBrush(FSlateNoResource())
+		.SetActiveFillBrush(FSlateNoResource())
+		.SetForegroundColor(FSlateColor::UseSubduedForeground())
+		.SetArrowsImage(FSlateNoResource())
+	);
+
+	Set("Sequencer.PlayTimeSpinBox", FSpinBoxStyle(GetWidgetStyle<FSpinBoxStyle>("SpinBox"))
+		.SetTextPadding(FMargin(0))
+		.SetBackgroundBrush(FSlateNoResource())
+		.SetHoveredBackgroundBrush(FSlateNoResource())
+		.SetInactiveFillBrush(FSlateNoResource())
+		.SetActiveFillBrush(FSlateNoResource())
+		.SetForegroundColor(SelectionColor_Pressed)
+		.SetArrowsImage(FSlateNoResource())
+	);
+
+	Set("Sequencer.HyperlinkTextBox", FEditableTextBoxStyle()
+		.SetFont(DEFAULT_FONT("Regular", 9))
+		.SetBackgroundImageNormal(FSlateNoResource())
+		.SetBackgroundImageHovered(FSlateNoResource())
+		.SetBackgroundImageFocused(FSlateNoResource())
+		.SetBackgroundImageReadOnly(FSlateNoResource())
+		.SetBackgroundColor(FLinearColor::Transparent)
+		.SetForegroundColor(FSlateColor::UseSubduedForeground())
+	);
+	Set("Sequencer.FixedFont", DEFAULT_FONT("Mono", 9));
+
+	Set("Sequencer.RecordSelectedActors", new IMAGE_BRUSH("SequenceRecorder/icon_tab_SequenceRecorder_16x", Icon16x16));
+
+	FComboButtonStyle SequencerSectionComboButton = FComboButtonStyle()
+	.SetButtonStyle(
+		FButtonStyle()
+		.SetNormal(FSlateNoResource())
+		.SetHovered(FSlateNoResource())
+		.SetPressed(FSlateNoResource())
+		.SetNormalPadding(FMargin(0, 0, 0, 0))
+		.SetPressedPadding(FMargin(0, 1, 0, 0))
+	)
+	.SetDownArrowImage(IMAGE_BRUSH("Common/ComboArrow", Icon8x8));
+	Set("Sequencer.SectionComboButton", SequencerSectionComboButton);
+
+	Set("Sequencer.CreateEventBinding", new IMAGE_BRUSH("Icons/icon_Blueprint_AddFunction_16px", Icon16x16));
+	Set("Sequencer.CreateQuickBinding", new IMAGE_BRUSH("Icons/icon_Blueprint_Node_16x", Icon16x16));
+	Set("Sequencer.ClearEventBinding", new IMAGE_BRUSH("Icons/Edit/icon_Edit_Delete_40x", Icon16x16));
+	Set("Sequencer.MultipleEvents", new IMAGE_BRUSH("Sequencer/MultipleEvents", Icon16x16));
+	Set("Sequencer.UnboundEvent", new IMAGE_BRUSH("Sequencer/UnboundEvent", Icon16x16));
+
+	// Sequencer Blending Iconography
+	Set("EMovieSceneBlendType::Absolute", new IMAGE_BRUSH("Sequencer/EMovieSceneBlendType_Absolute", FVector2D(32, 16)));
+	Set("EMovieSceneBlendType::Relative", new IMAGE_BRUSH("Sequencer/EMovieSceneBlendType_Relative", FVector2D(32, 16)));
+	Set("EMovieSceneBlendType::Additive", new IMAGE_BRUSH("Sequencer/EMovieSceneBlendType_Additive", FVector2D(32, 16)));
+		}
+
+
+	// Sequence recorder standalone UI
+	if (IncludeEditorSpecificStyles())
+	{
+		Set("SequenceRecorder.TabIcon", new IMAGE_BRUSH("SequenceRecorder/icon_tab_SequenceRecorder_16x", Icon16x16));
+		Set("SequenceRecorder.Common.RecordAll.Small", new IMAGE_BRUSH("SequenceRecorder/icon_RecordAll_40x", Icon20x20));
+		Set("SequenceRecorder.Common.RecordAll", new IMAGE_BRUSH("SequenceRecorder/icon_RecordAll_40x", Icon40x40));
+		Set("SequenceRecorder.Common.StopAll.Small", new IMAGE_BRUSH("SequenceRecorder/icon_StopAll_40x", Icon20x20));
+		Set("SequenceRecorder.Common.StopAll", new IMAGE_BRUSH("SequenceRecorder/icon_StopAll_40x", Icon40x40));
+		Set("SequenceRecorder.Common.AddRecording.Small", new IMAGE_BRUSH("SequenceRecorder/icon_AddRecording_40x", Icon20x20));
+		Set("SequenceRecorder.Common.AddRecording", new IMAGE_BRUSH("SequenceRecorder/icon_AddRecording_40x", Icon40x40));
+		Set("SequenceRecorder.Common.AddCurrentPlayerRecording.Small", new IMAGE_BRUSH("SequenceRecorder/icon_AddCurrentPlayerRecording_40x", Icon20x20));
+		Set("SequenceRecorder.Common.AddCurrentPlayerRecording", new IMAGE_BRUSH("SequenceRecorder/icon_AddCurrentPlayerRecording_40x", Icon40x40));
+		Set("SequenceRecorder.Common.RemoveRecording.Small", new IMAGE_BRUSH("SequenceRecorder/icon_RemoveRecording_40x", Icon20x20));
+		Set("SequenceRecorder.Common.RemoveRecording", new IMAGE_BRUSH("SequenceRecorder/icon_RemoveRecording_40x", Icon40x40));
+		Set("SequenceRecorder.Common.RemoveAllRecordings.Small", new IMAGE_BRUSH("SequenceRecorder/icon_RemoveRecording_40x", Icon20x20));
+		Set("SequenceRecorder.Common.RemoveAllRecordings", new IMAGE_BRUSH("SequenceRecorder/icon_RemoveRecording_40x", Icon40x40));
+		Set("SequenceRecorder.Common.RecordingActive", new IMAGE_BRUSH("Common/SmallCheckBox_Checked", Icon14x14));
+		Set("SequenceRecorder.Common.RecordingInactive", new IMAGE_BRUSH("Common/SmallCheckBox", Icon14x14));
+	}
+
+}
+
+void FSlateEditorStyle::FStyle::SetupViewportStyles()
+{
+// Viewport ToolbarBar
+	{
+		Set("ViewportMenu.Background", new BOX_BRUSH("Old/Menu_Background", FMargin(8.0f / 64.0f), FLinearColor::Transparent));
+		Set("ViewportMenu.Icon", new IMAGE_BRUSH("Icons/icon_tab_toolbar_16px", Icon16x16));
+		Set("ViewportMenu.Expand", new IMAGE_BRUSH("Icons/toolbar_expand_16x", Icon8x8));
+		Set("ViewportMenu.SubMenuIndicator", new IMAGE_BRUSH("Common/SubmenuArrow", Icon8x8));
+		Set("ViewportMenu.SToolBarComboButtonBlock.Padding", FMargin(0));
+		Set("ViewportMenu.SToolBarButtonBlock.Padding", FMargin(0));
+		Set("ViewportMenu.SToolBarButtonBlock.Button.Padding", FMargin(0));
+		Set("ViewportMenu.SToolBarCheckComboButtonBlock.Padding", FMargin(0));
+		Set("ViewportMenu.SToolBarButtonBlock.CheckBox.Padding", FMargin(4.0f));
+		Set("ViewportMenu.SToolBarComboButtonBlock.ComboButton.Color", FLinearColor(0.f, 0.f, 0.f, 0.75f));
+
+		Set("ViewportMenu.Separator", new BOX_BRUSH("Old/Button", 8.0f / 32.0f, FLinearColor::Transparent));
+		Set("ViewportMenu.Separator.Padding", FMargin(100.0f));
+
+		Set("ViewportMenu.Label", FTextBlockStyle(NormalText)
+			.SetFont(DEFAULT_FONT("Bold", 9))
+			.SetColorAndOpacity(FLinearColor(0.0f, 0.0f, 0.0f, 1.0f)));
+		Set("ViewportMenu.Label.Padding", FMargin(0.0f, 0.0f, 3.0f, 0.0f));
+		Set("ViewportMenu.Label.ContentPadding", FMargin(5.0f, 2.0f));
+		Set("ViewportMenu.EditableText", FEditableTextBoxStyle(NormalEditableTextBoxStyle).SetFont(DEFAULT_FONT("Regular", 9)));
+		Set("ViewportMenu.Keybinding", FTextBlockStyle(NormalText).SetFont(DEFAULT_FONT("Regular", 8)));
+
+		Set("ViewportMenu.Block.IndentedPadding", FMargin(0));
+		Set("ViewportMenu.Block.Padding", FMargin(0));
+
+		Set("ViewportMenu.Heading.Font", DEFAULT_FONT("Regular", 8));
+		Set("ViewportMenu.Heading.ColorAndOpacity", FLinearColor(0.4f, 0.4, 0.4f, 1.0f));
+
+		const FCheckBoxStyle ViewportMenuCheckBoxCheckBoxStyle = FCheckBoxStyle()
+			.SetUncheckedImage(IMAGE_BRUSH("Common/SmallCheckBox", Icon14x14))
+			.SetUncheckedPressedImage(IMAGE_BRUSH("Common/SmallCheckBox_Hovered", Icon14x14, FLinearColor(0.5f, 0.5f, 0.5f)))
+			.SetUncheckedHoveredImage(IMAGE_BRUSH("Common/SmallCheckBox_Hovered", Icon14x14))
+			.SetCheckedHoveredImage(IMAGE_BRUSH("Common/SmallCheckBox_Checked_Hovered", Icon14x14))
+			.SetCheckedPressedImage(IMAGE_BRUSH("Common/SmallCheckBox_Checked_Hovered", Icon14x14, FLinearColor(0.5f, 0.5f, 0.5f)))
+			.SetCheckedImage(IMAGE_BRUSH("Common/SmallCheckBox_Checked", Icon14x14));
+		Set("ViewportMenu.CheckBox", ViewportMenuCheckBoxCheckBoxStyle);
+
+		// Read-only checkbox that appears next to a menu item
+		const FCheckBoxStyle ViewportMenuCheckCheckBoxStyle = FCheckBoxStyle()
+			.SetUncheckedImage(IMAGE_BRUSH("Icons/Empty_14x", Icon14x14))
+			.SetUncheckedPressedImage(FSlateNoResource())
+			.SetUncheckedHoveredImage(FSlateNoResource())
+			.SetCheckedHoveredImage(IMAGE_BRUSH("Common/SmallCheck", Icon14x14))
+			.SetCheckedPressedImage(IMAGE_BRUSH("Common/SmallCheck", Icon14x14))
+			.SetCheckedImage(IMAGE_BRUSH("Common/SmallCheck", Icon14x14));
+		Set("ViewportMenu.Check", ViewportMenuCheckCheckBoxStyle);
+
+		const FString SmallRoundedButton(TEXT("Common/SmallRoundedButton"));
+		const FString SmallRoundedButtonStart(TEXT("Common/SmallRoundedButtonLeft"));
+		const FString SmallRoundedButtonMiddle(TEXT("Common/SmallRoundedButtonCentre"));
+const FString SmallRoundedButtonEnd(TEXT("Common/SmallRoundedButtonRight"));
+
+const FLinearColor NormalColor(1, 1, 1, 0.75f);
+const FLinearColor PressedColor(1, 1, 1, 1.f);
+
+const FCheckBoxStyle ViewportMenuRadioButtonCheckBoxStyle = FCheckBoxStyle()
+.SetUncheckedImage(IMAGE_BRUSH("Common/MenuItemRadioButton_Off", Icon14x14))
+.SetUncheckedPressedImage(IMAGE_BRUSH("Common/MenuItemRadioButton_Off", Icon14x14, FLinearColor(0.5f, 0.5f, 0.5f)))
+.SetUncheckedHoveredImage(IMAGE_BRUSH("Common/MenuItemRadioButton_Off", Icon14x14))
+.SetCheckedHoveredImage(IMAGE_BRUSH("Common/MenuItemRadioButton_On", Icon14x14))
+.SetCheckedPressedImage(IMAGE_BRUSH("Common/MenuItemRadioButton_On_Pressed", Icon14x14))
+.SetCheckedImage(IMAGE_BRUSH("Common/MenuItemRadioButton_On", Icon14x14));
+Set("ViewportMenu.RadioButton", ViewportMenuRadioButtonCheckBoxStyle);
+
+/* Create style for "ViewportMenu.ToggleButton" ... */
+const FCheckBoxStyle ViewportMenuToggleButtonStyle = FCheckBoxStyle()
+.SetCheckBoxType(ESlateCheckBoxType::ToggleButton)
+.SetUncheckedImage(BOX_BRUSH(*SmallRoundedButton, FMargin(7.f / 16.f), NormalColor))
+.SetUncheckedPressedImage(BOX_BRUSH(*SmallRoundedButton, FMargin(7.f / 16.f), PressedColor))
+.SetUncheckedHoveredImage(BOX_BRUSH(*SmallRoundedButton, FMargin(7.f / 16.f), PressedColor))
+.SetCheckedHoveredImage(BOX_BRUSH(*SmallRoundedButton, FMargin(7.f / 16.f), SelectionColor_Pressed))
+.SetCheckedPressedImage(BOX_BRUSH(*SmallRoundedButton, FMargin(7.f / 16.f), SelectionColor_Pressed))
+.SetCheckedImage(BOX_BRUSH(*SmallRoundedButton, FMargin(7.f / 16.f), SelectionColor_Pressed));
+/* ... and add new style */
+Set("ViewportMenu.ToggleButton", ViewportMenuToggleButtonStyle);
+
+/* Create style for "ViewportMenu.ToggleButton.Start" ... */
+const FCheckBoxStyle ViewportMenuToggleStartButtonStyle = FCheckBoxStyle()
+.SetCheckBoxType(ESlateCheckBoxType::ToggleButton)
+.SetUncheckedImage(BOX_BRUSH(*SmallRoundedButtonStart, FMargin(7.f / 16.f), NormalColor))
+.SetUncheckedPressedImage(BOX_BRUSH(*SmallRoundedButtonStart, FMargin(7.f / 16.f), PressedColor))
+.SetUncheckedHoveredImage(BOX_BRUSH(*SmallRoundedButtonStart, FMargin(7.f / 16.f), PressedColor))
+.SetCheckedHoveredImage(BOX_BRUSH(*SmallRoundedButtonStart, FMargin(7.f / 16.f), SelectionColor_Pressed))
+.SetCheckedPressedImage(BOX_BRUSH(*SmallRoundedButtonStart, FMargin(7.f / 16.f), SelectionColor_Pressed))
+.SetCheckedImage(BOX_BRUSH(*SmallRoundedButtonStart, FMargin(7.f / 16.f), SelectionColor_Pressed));
+/* ... and add new style */
+Set("ViewportMenu.ToggleButton.Start", ViewportMenuToggleStartButtonStyle);
+
+/* Create style for "ViewportMenu.ToggleButton.Middle" ... */
+const FCheckBoxStyle ViewportMenuToggleMiddleButtonStyle = FCheckBoxStyle()
+.SetCheckBoxType(ESlateCheckBoxType::ToggleButton)
+.SetUncheckedImage(BOX_BRUSH(*SmallRoundedButtonMiddle, FMargin(7.f / 16.f), NormalColor))
+.SetUncheckedPressedImage(BOX_BRUSH(*SmallRoundedButtonMiddle, FMargin(7.f / 16.f), PressedColor))
+.SetUncheckedHoveredImage(BOX_BRUSH(*SmallRoundedButtonMiddle, FMargin(7.f / 16.f), PressedColor))
+.SetCheckedHoveredImage(BOX_BRUSH(*SmallRoundedButtonMiddle, FMargin(7.f / 16.f), SelectionColor_Pressed))
+.SetCheckedPressedImage(BOX_BRUSH(*SmallRoundedButtonMiddle, FMargin(7.f / 16.f), SelectionColor_Pressed))
+.SetCheckedImage(BOX_BRUSH(*SmallRoundedButtonMiddle, FMargin(7.f / 16.f), SelectionColor_Pressed));
+/* ... and add new style */
+Set("ViewportMenu.ToggleButton.Middle", ViewportMenuToggleMiddleButtonStyle);
+
+/* Create style for "ViewportMenu.ToggleButton.End" ... */
+const FCheckBoxStyle ViewportMenuToggleEndButtonStyle = FCheckBoxStyle()
+.SetCheckBoxType(ESlateCheckBoxType::ToggleButton)
+.SetUncheckedImage(BOX_BRUSH(*SmallRoundedButtonEnd, FMargin(7.f / 16.f), NormalColor))
+.SetUncheckedPressedImage(BOX_BRUSH(*SmallRoundedButtonEnd, FMargin(7.f / 16.f), PressedColor))
+.SetUncheckedHoveredImage(BOX_BRUSH(*SmallRoundedButtonEnd, FMargin(7.f / 16.f), PressedColor))
+.SetCheckedHoveredImage(BOX_BRUSH(*SmallRoundedButtonEnd, FMargin(7.f / 16.f), SelectionColor_Pressed))
+.SetCheckedPressedImage(BOX_BRUSH(*SmallRoundedButtonEnd, FMargin(7.f / 16.f), SelectionColor_Pressed))
+.SetCheckedImage(BOX_BRUSH(*SmallRoundedButtonEnd, FMargin(7.f / 16.f), SelectionColor_Pressed));
+/* ... and add new style */
+Set("ViewportMenu.ToggleButton.End", ViewportMenuToggleEndButtonStyle);
+
+const FMargin NormalPadding = FMargin(4.0f, 4.0f, 4.0f, 4.0f);
+const FMargin PressedPadding = FMargin(4.0f, 4.0f, 4.0f, 4.0f);
+
+const FButtonStyle ViewportMenuButton = FButtonStyle(Button)
+.SetNormal(BOX_BRUSH(*SmallRoundedButton, 7.0f / 16.0f, NormalColor))
+.SetPressed(BOX_BRUSH(*SmallRoundedButton, 7.0f / 16.0f, PressedColor))
+.SetHovered(BOX_BRUSH(*SmallRoundedButton, 7.0f / 16.0f, PressedColor))
+.SetPressedPadding(PressedPadding)
+.SetNormalPadding(NormalPadding);
+
+Set("ViewportMenu.Button", ViewportMenuButton);
+
+Set("ViewportMenu.Button.Start", FButtonStyle(ViewportMenuButton)
+	.SetNormal(BOX_BRUSH(*SmallRoundedButtonStart, 7.0f / 16.0f, NormalColor))
+	.SetPressed(BOX_BRUSH(*SmallRoundedButtonStart, 7.0f / 16.0f, PressedColor))
+	.SetHovered(BOX_BRUSH(*SmallRoundedButtonStart, 7.0f / 16.0f, PressedColor))
+);
+
+Set("ViewportMenu.Button.Middle", FButtonStyle(ViewportMenuButton)
+	.SetNormal(BOX_BRUSH(*SmallRoundedButtonMiddle, 7.0f / 16.0f, NormalColor))
+	.SetPressed(BOX_BRUSH(*SmallRoundedButtonMiddle, 7.0f / 16.0f, PressedColor))
+	.SetHovered(BOX_BRUSH(*SmallRoundedButtonMiddle, 7.0f / 16.0f, PressedColor))
+);
+
+Set("ViewportMenu.Button.End", FButtonStyle(ViewportMenuButton)
+	.SetNormal(BOX_BRUSH(*SmallRoundedButtonEnd, 7.0f / 16.0f, NormalColor))
+	.SetPressed(BOX_BRUSH(*SmallRoundedButtonEnd, 7.0f / 16.0f, PressedColor))
+	.SetHovered(BOX_BRUSH(*SmallRoundedButtonEnd, 7.0f / 16.0f, PressedColor))
+);
+	}
+
+	// Viewport actor preview's pin/unpin buttons
+	{
+	Set("ViewportActorPreview.Pinned", new IMAGE_BRUSH("Common/PushPin_Down", Icon16x16));
+	Set("ViewportActorPreview.Unpinned", new IMAGE_BRUSH("Common/PushPin_Up", Icon16x16));
+	}
+}
+
+void FSlateEditorStyle::FStyle::SetupNotificationBarStyles()
+{
+// NotificationBar
+	{
+		Set("NotificationBar.Background", new FSlateNoResource());
+		Set("NotificationBar.Icon", new FSlateNoResource());
+		Set("NotificationBar.Expand", new IMAGE_BRUSH("Icons/toolbar_expand_16x", Icon16x16));
+		Set("NotificationBar.SubMenuIndicator", new IMAGE_BRUSH("Common/SubmenuArrow", Icon8x8));
+
+		Set("NotificationBar.Block.IndentedPadding", FMargin(0));
+		Set("NotificationBar.Block.Padding", FMargin(0));
+
+		Set("NotificationBar.Separator", new BOX_BRUSH("Old/Button", 4.0f / 32.0f));
+		Set("NotificationBar.Separator.Padding", FMargin(0.5f));
+
+		Set("NotificationBar.Label", FTextBlockStyle(NormalText).SetFont(DEFAULT_FONT("Regular", 9)));
+		Set("NotificationBar.EditableText", FEditableTextBoxStyle(NormalEditableTextBoxStyle).SetFont(DEFAULT_FONT("Regular", 9)));
+		Set("NotificationBar.Keybinding", FTextBlockStyle(NormalText).SetFont(DEFAULT_FONT("Regular", 8)));
+
+		Set("NotificationBar.Heading", FTextBlockStyle(NormalText)
+			.SetFont(DEFAULT_FONT("Regular", 8))
+			.SetColorAndOpacity(FLinearColor(0.4f, 0.4, 0.4f, 1.0f)));
+
+		const FCheckBoxStyle NotificationBarCheckBoxCheckBoxStyle = FCheckBoxStyle()
+			.SetUncheckedImage(IMAGE_BRUSH("Common/SmallCheckBox", Icon14x14))
+			.SetUncheckedPressedImage(IMAGE_BRUSH("Common/SmallCheckBox_Hovered", Icon14x14, FLinearColor(0.5f, 0.5f, 0.5f)))
+			.SetUncheckedHoveredImage(IMAGE_BRUSH("Common/SmallCheckBox_Hovered", Icon14x14))
+			.SetCheckedHoveredImage(IMAGE_BRUSH("Common/SmallCheckBox_Checked_Hovered", Icon14x14))
+			.SetCheckedPressedImage(IMAGE_BRUSH("Common/SmallCheckBox_Checked_Hovered", Icon14x14, FLinearColor(0.5f, 0.5f, 0.5f)))
+			.SetCheckedImage(IMAGE_BRUSH("Common/SmallCheckBox_Checked", Icon14x14))
+			.SetUndeterminedImage(IMAGE_BRUSH("Common/CheckBox_Undetermined", Icon14x14))
+			.SetUndeterminedHoveredImage(IMAGE_BRUSH("Common/CheckBox_Undetermined_Hovered", Icon14x14))
+			.SetUndeterminedPressedImage(IMAGE_BRUSH("Common/CheckBox_Undetermined_Hovered", Icon14x14, FLinearColor(0.5f, 0.5f, 0.5f)));
+		Set("NotificationBar.CheckBox", NotificationBarCheckBoxCheckBoxStyle);
+
+		// Read-only checkbox that appears next to a menu item
+		const FCheckBoxStyle NotificationBarCheckCheckBoxStyle = FCheckBoxStyle()
+			.SetUncheckedImage(IMAGE_BRUSH("Icons/Empty_14x", Icon14x14))
+			.SetUncheckedPressedImage(FSlateNoResource())
+			.SetUncheckedHoveredImage(FSlateNoResource())
+			.SetCheckedHoveredImage(IMAGE_BRUSH("Common/SmallCheck", Icon14x14))
+			.SetCheckedPressedImage(IMAGE_BRUSH("Common/SmallCheck", Icon14x14))
+			.SetCheckedImage(IMAGE_BRUSH("Common/SmallCheck", Icon14x14))
+			.SetUndeterminedImage(IMAGE_BRUSH("Icons/Empty_14x", Icon14x14))
+			.SetUndeterminedPressedImage(FSlateNoResource())
+			.SetUndeterminedHoveredImage(FSlateNoResource());
+		Set("NotificationBar.Check", NotificationBarCheckCheckBoxStyle);
+
+		// This radio button is actually just a check box with different images
+		const FCheckBoxStyle NotificationBarRadioButtonCheckBoxStyle = FCheckBoxStyle()
+			.SetUncheckedImage(IMAGE_BRUSH("Common/RadioButton_Unselected_16x", Icon16x16))
+			.SetUncheckedPressedImage(IMAGE_BRUSH("Common/RadioButton_Unselected_16x", Icon16x16, SelectionColor_Pressed))
+			.SetUncheckedHoveredImage(IMAGE_BRUSH("Common/RadioButton_Unselected_16x", Icon16x16, SelectionColor))
+			.SetCheckedHoveredImage(IMAGE_BRUSH("Common/RadioButton_Selected_16x", Icon16x16, SelectionColor))
+			.SetCheckedPressedImage(IMAGE_BRUSH("Common/RadioButton_Selected_16x", Icon16x16, SelectionColor_Pressed))
+			.SetCheckedImage(IMAGE_BRUSH("Common/RadioButton_Selected_16x", Icon16x16));
+		Set("NotificationBar.RadioButton", NotificationBarRadioButtonCheckBoxStyle);
+
+		const FCheckBoxStyle NotificationBarToggleButtonCheckBoxStyle = FCheckBoxStyle()
+			.SetCheckBoxType(ESlateCheckBoxType::ToggleButton)
+			.SetUncheckedImage(FSlateNoResource())
+			.SetUncheckedPressedImage(BOX_BRUSH("Common/RoundedSelection_16x", 4.0f / 16.0f, SelectionColor_Pressed))
+			.SetUncheckedHoveredImage(BOX_BRUSH("Common/RoundedSelection_16x", 4.0f / 16.0f, SelectionColor))
+			.SetCheckedHoveredImage(BOX_BRUSH("Common/RoundedSelection_16x", 4.0f / 16.0f, SelectionColor_Pressed))
+			.SetCheckedPressedImage(BOX_BRUSH("Common/RoundedSelection_16x", 4.0f / 16.0f, SelectionColor))
+			.SetCheckedImage(BOX_BRUSH("Common/RoundedSelection_16x", 4.0f / 16.0f, SelectionColor_Pressed));
+		Set("NotificationBar.ToggleButton", NotificationBarToggleButtonCheckBoxStyle);
+
+		const FButtonStyle NoBorder = FButtonStyle()
+			.SetNormal(FSlateNoResource())
+			.SetHovered(FSlateNoResource())
+			.SetPressed(FSlateNoResource())
+			.SetNormalPadding(FMargin(0, 0, 0, 1))
+			.SetPressedPadding(FMargin(0, 1, 0, 0));
+
+		Set("NotificationBar.Button", FButtonStyle(NoBorder)
+			.SetNormal(FSlateNoResource())
+			.SetPressed(BOX_BRUSH("Common/RoundedSelection_16x", 4.0f / 16.0f, SelectionColor_Pressed))
+			.SetHovered(BOX_BRUSH("Common/RoundedSelection_16x", 4.0f / 16.0f, SelectionColor))
+			.SetNormalPadding(FMargin(0, 1))
+			.SetPressedPadding(FMargin(0, 2, 0, 0))
+		);
+
+		Set("NotificationBar.Button.Checked", new BOX_BRUSH("Common/RoundedSelection_16x", 4.0f / 16.0f, SelectionColor_Pressed));
+		Set("NotificationBar.Button.Checked_Hovered", new BOX_BRUSH("Common/RoundedSelection_16x", 4.0f / 16.0f, SelectionColor_Pressed));
+		Set("NotificationBar.Button.Checked_Pressed", new BOX_BRUSH("Common/RoundedSelection_16x", 4.0f / 16.0f, SelectionColor));
+
+		Set("NotificationBar.SToolBarButtonBlock.CheckBox.Padding", FMargin(4.0f));
+		Set("NotificationBar.SToolBarButtonBlock.Button.Padding", FMargin(0.0f));
+		Set("NotificationBar.SToolBarComboButtonBlock.ComboButton.Color", DefaultForeground);
+	}
+}
+
+void FSlateEditorStyle::FStyle::SetupMenuBarStyles()
+{
+	// MenuBar
+	{
+		Set("Menu.Background", new BOX_BRUSH("Old/Menu_Background", FMargin(8.0f / 64.0f)));
+		Set("Menu.Icon", new IMAGE_BRUSH("Icons/icon_tab_toolbar_16px", Icon16x16));
+		Set("Menu.Expand", new IMAGE_BRUSH("Icons/toolbar_expand_16x", Icon16x16));
+		Set("Menu.SubMenuIndicator", new IMAGE_BRUSH("Common/SubmenuArrow", Icon8x8));
+		Set("Menu.SToolBarComboButtonBlock.Padding", FMargin(4.0f));
+		Set("Menu.SToolBarButtonBlock.Padding", FMargin(4.0f));
+		Set("Menu.SToolBarCheckComboButtonBlock.Padding", FMargin(4.0f));
+		Set("Menu.SToolBarButtonBlock.CheckBox.Padding", FMargin(0.0f));
+		Set("Menu.SToolBarComboButtonBlock.ComboButton.Color", DefaultForeground);
+
+		Set("Menu.Block.IndentedPadding", FMargin(18.0f, 2.0f, 4.0f, 4.0f));
+		Set("Menu.Block.Padding", FMargin(2.0f, 2.0f, 4.0f, 4.0f));
+
+		Set("Menu.Separator", new BOX_BRUSH("Old/Button", 4.0f / 32.0f));
+		Set("Menu.Separator.Padding", FMargin(0.5f));
+
+		Set("Menu.Label", FTextBlockStyle(NormalText).SetFont(DEFAULT_FONT("Regular", 9)));
+		Set("Menu.Label.Padding", FMargin(0.0f, 0.0f, 0.0f, 0.0f));
+		Set("Menu.Label.ContentPadding", FMargin(10.0f, 2.0f));
+		Set("Menu.EditableText", FEditableTextBoxStyle(NormalEditableTextBoxStyle).SetFont(DEFAULT_FONT("Regular", 9)));
+		Set("Menu.Keybinding", FTextBlockStyle(NormalText).SetFont(DEFAULT_FONT("Regular", 8)));
+
+		Set("Menu.Heading", FTextBlockStyle(NormalText)
+			.SetFont(DEFAULT_FONT("Regular", 8))
+			.SetColorAndOpacity(FLinearColor(0.4f, 0.4, 0.4f, 1.0f)));
+
+		/* Set images for various SCheckBox states associated with menu check box items... */
+		const FCheckBoxStyle BasicMenuCheckBoxStyle = FCheckBoxStyle()
+			.SetUncheckedImage(IMAGE_BRUSH("Common/SmallCheckBox", Icon14x14))
+			.SetUncheckedHoveredImage(IMAGE_BRUSH("Common/SmallCheckBox_Hovered", Icon14x14))
+			.SetUncheckedPressedImage(IMAGE_BRUSH("Common/SmallCheckBox_Hovered", Icon14x14, FLinearColor(0.5f, 0.5f, 0.5f)))
+			.SetCheckedImage(IMAGE_BRUSH("Common/SmallCheckBox_Checked", Icon14x14))
+			.SetCheckedHoveredImage(IMAGE_BRUSH("Common/SmallCheckBox_Checked_Hovered", Icon14x14))
+			.SetCheckedPressedImage(IMAGE_BRUSH("Common/SmallCheckBox_Checked_Hovered", Icon14x14, FLinearColor(0.5f, 0.5f, 0.5f)))
+			.SetUndeterminedImage(IMAGE_BRUSH("Common/CheckBox_Undetermined", Icon14x14))
+			.SetUndeterminedHoveredImage(IMAGE_BRUSH("Common/CheckBox_Undetermined_Hovered", Icon14x14))
+			.SetUndeterminedPressedImage(IMAGE_BRUSH("Common/CheckBox_Undetermined_Hovered", Icon14x14, FLinearColor(0.5f, 0.5f, 0.5f)));
+
+		/* ...and add the new style */
+		Set("Menu.CheckBox", BasicMenuCheckBoxStyle);
+
+		/* Read-only checkbox that appears next to a menu item */
+		/* Set images for various SCheckBox states associated with read-only menu check box items... */
+		const FCheckBoxStyle BasicMenuCheckStyle = FCheckBoxStyle()
+			.SetUncheckedImage(IMAGE_BRUSH("Icons/Empty_14x", Icon14x14))
+			.SetUncheckedHoveredImage(IMAGE_BRUSH("Icons/Empty_14x", Icon14x14))
+			.SetUncheckedPressedImage(IMAGE_BRUSH("Common/SmallCheckBox_Hovered", Icon14x14))
+			.SetCheckedImage(IMAGE_BRUSH("Common/SmallCheck", Icon14x14))
+			.SetCheckedHoveredImage(IMAGE_BRUSH("Common/SmallCheck", Icon14x14))
+			.SetCheckedPressedImage(IMAGE_BRUSH("Common/SmallCheck", Icon14x14))
+			.SetUndeterminedImage(IMAGE_BRUSH("Icons/Empty_14x", Icon14x14))
+			.SetUndeterminedHoveredImage(FSlateNoResource())
+			.SetUndeterminedPressedImage(FSlateNoResource());
+
+		/* ...and add the new style */
+		Set("Menu.Check", BasicMenuCheckStyle);
+
+		/* This radio button is actually just a check box with different images */
+		/* Set images for various Menu radio button (SCheckBox) states... */
+		const FCheckBoxStyle BasicMenuRadioButtonStyle = FCheckBoxStyle()
+			.SetUncheckedImage(IMAGE_BRUSH("Common/RadioButton_Unselected_16x", Icon16x16))
+			.SetUncheckedHoveredImage(IMAGE_BRUSH("Common/RadioButton_Unselected_16x", Icon16x16))
+			.SetUncheckedPressedImage(IMAGE_BRUSH("Common/RadioButton_Unselected_16x", Icon16x16))
+			.SetCheckedImage(IMAGE_BRUSH("Common/RadioButton_Selected_16x", Icon16x16))
+			.SetCheckedHoveredImage(IMAGE_BRUSH("Common/RadioButton_Selected_16x", Icon16x16, SelectionColor))
+			.SetCheckedPressedImage(IMAGE_BRUSH("Common/RadioButton_Unselected_16x", Icon16x16, SelectionColor_Pressed))
+			.SetUndeterminedImage(IMAGE_BRUSH("Common/RadioButton_Unselected_16x", Icon16x16))
+			.SetUndeterminedHoveredImage(IMAGE_BRUSH("Common/RadioButton_Unselected_16x", Icon16x16, SelectionColor))
+			.SetUndeterminedPressedImage(IMAGE_BRUSH("Common/RadioButton_Unselected_16x", Icon16x16, SelectionColor_Pressed));
+
+		/* ...and set new style */
+		Set("Menu.RadioButton", BasicMenuRadioButtonStyle);
+
+		/* Create style for "Menu.ToggleButton" widget ... */
+		const FCheckBoxStyle MenuToggleButtonCheckBoxStyle = FCheckBoxStyle()
+			.SetCheckBoxType(ESlateCheckBoxType::ToggleButton)
+			.SetUncheckedImage(FSlateNoResource())
+			.SetUncheckedPressedImage(BOX_BRUSH("Common/RoundedSelection_16x", 4.0f / 16.0f, SelectionColor_Pressed))
+			.SetUncheckedHoveredImage(BOX_BRUSH("Common/RoundedSelection_16x", 4.0f / 16.0f, SelectionColor))
+			.SetCheckedImage(BOX_BRUSH("Common/RoundedSelection_16x", 4.0f / 16.0f, SelectionColor_Pressed))
+			.SetCheckedHoveredImage(BOX_BRUSH("Common/RoundedSelection_16x", 4.0f / 16.0f, SelectionColor_Pressed))
+			.SetCheckedPressedImage(BOX_BRUSH("Common/RoundedSelection_16x", 4.0f / 16.0f, SelectionColor));
+		/* ... and add new style */
+		Set("Menu.ToggleButton", MenuToggleButtonCheckBoxStyle);
+
+		const FButtonStyle NoBorder = FButtonStyle()
+			.SetNormal(FSlateNoResource())
+			.SetHovered(FSlateNoResource())
+			.SetPressed(FSlateNoResource())
+			.SetNormalPadding(FMargin(0, 0, 0, 1))
+			.SetPressedPadding(FMargin(0, 1, 0, 0));
+
+		Set("Menu.Button", FButtonStyle(NoBorder)
+			.SetNormal(FSlateNoResource())
+			.SetPressed(BOX_BRUSH("Common/RoundedSelection_16x", 4.0f / 16.0f, SelectionColor_Pressed))
+			.SetHovered(BOX_BRUSH("Common/RoundedSelection_16x", 4.0f / 16.0f, SelectionColor))
+			.SetNormalPadding(FMargin(0, 1))
+			.SetPressedPadding(FMargin(0, 2, 0, 0))
+		);
+
+		Set("Menu.Button.Checked", new BOX_BRUSH("Common/RoundedSelection_16x", 4.0f / 16.0f, SelectionColor_Pressed));
+		Set("Menu.Button.Checked_Hovered", new BOX_BRUSH("Common/RoundedSelection_16x", 4.0f / 16.0f, SelectionColor_Pressed));
+		Set("Menu.Button.Checked_Pressed", new BOX_BRUSH("Common/RoundedSelection_16x", 4.0f / 16.0f, SelectionColor));
+
+		/* The style of a menu bar button when it has a sub menu open */
+		Set("Menu.Button.SubMenuOpen", new BORDER_BRUSH("Common/Selection", FMargin(4.f / 16.f), FLinearColor(0.10f, 0.10f, 0.10f)));
+	}
+}
+
 void FSlateEditorStyle::FStyle::SetupGeneralIcons()
 {
 	Set("Plus", new IMAGE_BRUSH("Icons/PlusSymbol_12x", Icon12x12));
diff --git a/Engine/Source/Editor/EditorStyle/Private/SlateEditorStyle.h b/Engine/Source/Editor/EditorStyle/Private/SlateEditorStyle.h
--- a/Engine/Source/Editor/EditorStyle/Private/SlateEditorStyle.h
+++ b/Engine/Source/Editor/EditorStyle/Private/SlateEditorStyle.h
@@ -70,6 +70,12 @@ public:
 
 		void Initialize();
 		void SetupGeneralStyles();
+		void SetupLevelGeneralStyles();
+		void SetupWorldBrowserStyles();
+		void SetupSequencerStyles();
+		void SetupViewportStyles();
+		void SetupNotificationBarStyles();
+		void SetupMenuBarStyles();
 		void SetupGeneralIcons();
 		void SetupWindowStyles();
 		void SetupProjectBadgeStyle();
