Curve Editor: MultiScaleTool icon

diff --git a/Engine/Source/Editor/EditorStyle/Private/SlateEditorStyle.cpp b/Engine/Source/Editor/EditorStyle/Private/SlateEditorStyle.cpp
--- a/Engine/Source/Editor/EditorStyle/Private/SlateEditorStyle.cpp
+++ b/Engine/Source/Editor/EditorStyle/Private/SlateEditorStyle.cpp
@@ -7319,6 +7319,8 @@ void FSlateEditorStyle::FStyle::SetupToolkitStyles()
 		Set("CurveEditorTools.ActivateTransformTool.Small", new IMAGE_BRUSH("GenericCurveEditor/Icons/ToolsTransform_48x", Icon24x24));
 		Set("CurveEditorTools.ActivateRetimeTool", new IMAGE_BRUSH("GenericCurveEditor/Icons/ToolsRetime_48x", Icon48x48));
 		Set("CurveEditorTools.ActivateRetimeTool.Small", new IMAGE_BRUSH("GenericCurveEditor/Icons/ToolsRetime_48x", Icon24x24));
+		Set("CurveEditorTools.ActivateMultiScaleTool", new IMAGE_BRUSH("GenericCurveEditor/Icons/ToolsMultiScale_48x", Icon48x48));
+		Set("CurveEditorTools.ActivateMultiScaleTool.Small", new IMAGE_BRUSH("GenericCurveEditor/Icons/ToolsMultiScale_48x", Icon24x24));
 	}
 
 	// PhysicsAssetEditor
