UX: Cleaner Frames, Splitter Bug Fix

- splitters would not show their hover brush if the normal brush was non-transparent

diff --git a/Engine/Source/Editor/ContentBrowser/Private/SContentBrowser.cpp b/Engine/Source/Editor/ContentBrowser/Private/SContentBrowser.cpp
--- a/Engine/Source/Editor/ContentBrowser/Private/SContentBrowser.cpp
+++ b/Engine/Source/Editor/ContentBrowser/Private/SContentBrowser.cpp
@@ -61,6 +61,10 @@
 #include "Toolkits/GlobalEditorCommonCommands.h"
 #include "Subsystems/AssetEditorSubsystem.h"
 
+
+#include "Brushes/SlateColorBrush.h"
+
+
 #define LOCTEXT_NAMESPACE "ContentBrowser"
 
 const FString SContentBrowser::SettingsIniSection = TEXT("ContentBrowser");
@@ -152,7 +156,7 @@ void SContentBrowser::Construct( const FArguments& InArgs, const FName& InInstan
 				[
 					SNew( SBorder )
 					.Padding( FMargin( 3 ) )
-					.BorderImage( FEditorStyle::GetBrush( "ToolPanel.GroupBorder" ) )
+					.BorderImage( FEditorStyle::GetBrush( "ContentBrowser.TopBar.GroupBorder" ) )
 					[
 						SNew( SHorizontalBox )
 
@@ -300,7 +304,7 @@ void SContentBrowser::Construct( const FArguments& InArgs, const FName& InInstan
 			[
 				SNew(SBorder)
 				.Padding(FMargin(3))
-				.BorderImage( FEditorStyle::GetBrush("ToolPanel.GroupBorder") )
+				.BorderImage( FEditorStyle::GetBrush("ContentBrowser.TopBar.GroupBorder") )
 				[
 					SNew(SHorizontalBox)
 
@@ -438,6 +442,18 @@ void SContentBrowser::Construct( const FArguments& InArgs, const FName& InInstan
 			]
 		]
 
+		+ SVerticalBox::Slot()
+		.AutoHeight()
+		.Padding(0,0,0,0)
+		[
+			SNew(SBox)
+			.HeightOverride(2.0f)
+			[
+				SNew(SImage)
+				.Image(new FSlateColorBrush(FLinearColor( FColor( 34, 34, 34) ) ) )
+			]
+		]
+
 		// Assets/tree
 		+ SVerticalBox::Slot()
 		.FillHeight(1.0f)
@@ -445,13 +461,16 @@ void SContentBrowser::Construct( const FArguments& InArgs, const FName& InInstan
 		[
 			// The tree/assets splitter
 			SAssignNew(PathAssetSplitterPtr, SSplitter)
+			.Style(FEditorStyle::Get(), "ContentBrowser.Splitter")
+			.PhysicalSplitterHandleSize(2.0f)
+
 			// Sources View
 			+ SSplitter::Slot()
 			.Value(0.25f)
 			[
 				SNew(SBorder)
 				.Padding(FMargin(3))
-				.BorderImage(FEditorStyle::GetBrush("ToolPanel.GroupBorder"))
+				.BorderImage(FEditorStyle::GetBrush("ContentBrowser.TopBar.GroupBorder"))
 				.Visibility(this, &SContentBrowser::GetSourcesViewVisibility)
 				[
 					SNew(SVerticalBox)
@@ -611,7 +630,7 @@ void SContentBrowser::Construct( const FArguments& InArgs, const FName& InInstan
 			[
 				SNew(SBorder)
 				.Padding(FMargin(3))
-				.BorderImage( FEditorStyle::GetBrush("ToolPanel.GroupBorder") )
+				.BorderImage( FEditorStyle::GetBrush("ContentBrowser.TopBar.GroupBorder") )
 				[
 					SNew(SVerticalBox)
 
diff --git a/Engine/Source/Editor/CurveEditor/Private/SCurveEditorPanel.cpp b/Engine/Source/Editor/CurveEditor/Private/SCurveEditorPanel.cpp
--- a/Engine/Source/Editor/CurveEditor/Private/SCurveEditorPanel.cpp
+++ b/Engine/Source/Editor/CurveEditor/Private/SCurveEditorPanel.cpp
@@ -228,6 +228,8 @@ void SCurveEditorPanel::Construct(const FArguments& InArgs, TSharedRef<FCurveEdi
 		[
 			SNew(SSplitter)
 			.Orientation(Orient_Horizontal)
+			.Style(FEditorStyle::Get(), "SplitterDark")
+			.PhysicalSplitterHandleSize(2.0f)
 
 			+ SSplitter::Slot()
 			.Value(0.30f)
diff --git a/Engine/Source/Editor/EditorStyle/Private/SlateEditorStyle.cpp b/Engine/Source/Editor/EditorStyle/Private/SlateEditorStyle.cpp
--- a/Engine/Source/Editor/EditorStyle/Private/SlateEditorStyle.cpp
+++ b/Engine/Source/Editor/EditorStyle/Private/SlateEditorStyle.cpp
@@ -14,6 +14,12 @@ BEGIN_SLATE_FUNCTION_BUILD_OPTIMIZATION
 #define IMAGE_BRUSH( RelativePath, ... ) FSlateImageBrush( RootToContentDir( RelativePath, TEXT(".png") ), __VA_ARGS__ )
 #define BOX_BRUSH( RelativePath, ... ) FSlateBoxBrush( RootToContentDir( RelativePath, TEXT(".png") ), __VA_ARGS__ )
 #define BORDER_BRUSH( RelativePath, ... ) FSlateBorderBrush( RootToContentDir( RelativePath, TEXT(".png") ), __VA_ARGS__ )
+
+#define CORE_IMAGE_BRUSH( RelativePath, ... ) FSlateImageBrush( RootToCoreContentDir( RelativePath, TEXT(".png") ), __VA_ARGS__ )
+#define CORE_BOX_BRUSH( RelativePath, ... ) FSlateBoxBrush( RootToCoreContentDir( RelativePath, TEXT(".png") ), __VA_ARGS__ )
+#define CORE_BORDER_BRUSH( RelativePath, ... ) FSlateBorderBrush( RootToCoreContentDir( RelativePath, TEXT(".png") ), __VA_ARGS__ )
+
+
 #define DEFAULT_FONT(...) FCoreStyle::GetDefaultFontStyle(__VA_ARGS__)
 #define ICON_FONT(...) FSlateFontInfo(RootToContentDir("Fonts/FontAwesome", TEXT(".ttf")), __VA_ARGS__)
 
@@ -903,6 +909,11 @@ void FSlateEditorStyle::FStyle::SetupGeneralStyles()
 			.SetHandleNormalBrush( FSlateNoResource() )
 			.SetHandleHighlightBrush( IMAGE_BRUSH( "Common/SplitterHandleHighlight", Icon8x8, FLinearColor::White ) )
 			);
+
+		Set ("SplitterDark", FSplitterStyle()
+			.SetHandleNormalBrush( FSlateColorBrush( FLinearColor(FColor( 32, 32, 32) ) ) )
+			.SetHandleHighlightBrush( FSlateColorBrush( FLinearColor(FColor( 96, 96, 96) ) ) ) 
+			);
 	}
 #endif // WITH_EDITOR || (IS_PROGRAM && WITH_UNREAL_DEVELOPER_TOOLS)
 
@@ -6503,8 +6514,8 @@ void FSlateEditorStyle::FStyle::SetupContentBrowserStyle()
 		Set( "ContentBrowser.AddCollectionButtonIcon", new IMAGE_BRUSH( "Icons/PlusSymbol_12x", Icon12x12 ) );
 
 		Set( "ContentBrowser.Splitter", FSplitterStyle()
-			.SetHandleNormalBrush( FSlateNoResource() )
-			.SetHandleHighlightBrush( IMAGE_BRUSH( "Common/SplitterHandleHighlight", Icon8x8, FLinearColor::White ) )
+			.SetHandleNormalBrush( FSlateColorBrush( FLinearColor(FColor( 32, 32, 32) ) ) )
+			.SetHandleHighlightBrush( FSlateColorBrush( FLinearColor(FColor( 96, 96, 96) ) ) ) 
 			);
 
 		// Asset list view
@@ -6640,6 +6651,8 @@ void FSlateEditorStyle::FStyle::SetupContentBrowserStyle()
 		Set( "ContentBrowser.PrimitiveCylinder", new IMAGE_BRUSH( "ContentBrowser/ThumbnailCylinder", Icon32x32 ) );
 		Set( "ContentBrowser.ResetPrimitiveToDefault", new IMAGE_BRUSH("ContentBrowser/ThumbnailReset", Icon32x32) );
 
+		Set( "ContentBrowser.TopBar.GroupBorder", new CORE_BOX_BRUSH( "Common/GroupBorder_Shape", FMargin(4.0f/16.0f), FLinearColor(FColor(62, 62, 62)) ));
+
 		Set( "ContentBrowser.TopBar.Font", FTextBlockStyle( NormalText )
 			.SetFont( DEFAULT_FONT( "Bold", 11 ) )
 			.SetColorAndOpacity( FLinearColor( 1.0f, 1.0f, 1.0f ) )
diff --git a/Engine/Source/Editor/LevelEditor/Private/SActorDetails.cpp b/Engine/Source/Editor/LevelEditor/Private/SActorDetails.cpp
--- a/Engine/Source/Editor/LevelEditor/Private/SActorDetails.cpp
+++ b/Engine/Source/Editor/LevelEditor/Private/SActorDetails.cpp
@@ -153,6 +153,8 @@ void SActorDetails::Construct(const FArguments& InArgs, const FName TabIdentifie
 		[
 			SAssignNew(DetailsSplitter, SSplitter)
 			.Orientation(Orient_Vertical)
+			.Style(FEditorStyle::Get(), "SplitterDark")
+			.PhysicalSplitterHandleSize(2.0f)
 			+ SSplitter::Slot()
 			[
 				SNew( SVerticalBox )
diff --git a/Engine/Source/Runtime/Slate/Private/Widgets/Layout/SSplitter.cpp b/Engine/Source/Runtime/Slate/Private/Widgets/Layout/SSplitter.cpp
--- a/Engine/Source/Runtime/Slate/Private/Widgets/Layout/SSplitter.cpp
+++ b/Engine/Source/Runtime/Slate/Private/Widgets/Layout/SSplitter.cpp
@@ -166,7 +166,7 @@ int32 SSplitter::OnPaint( const FPaintArgs& Args, const FGeometry& AllottedGeome
 	// Draw the splitter above any children
 	MaxLayerId += 1;
 
-	for( int32 ChildIndex = 0; ChildIndex < ArrangedChildren.Num(); ++ChildIndex )
+	for( int32 ChildIndex = 0; ChildIndex < ArrangedChildren.Num() - 1; ++ChildIndex )
 	{
 		const FGeometry& GeometryAfterSplitter = ArrangedChildren[ FMath::Clamp(ChildIndex + 1, 0, ArrangedChildren.Num()-1) ].Geometry;
 
diff --git a/Engine/Source/Runtime/SlateCore/Private/Styling/CoreStyle.cpp b/Engine/Source/Runtime/SlateCore/Private/Styling/CoreStyle.cpp
--- a/Engine/Source/Runtime/SlateCore/Private/Styling/CoreStyle.cpp
+++ b/Engine/Source/Runtime/SlateCore/Private/Styling/CoreStyle.cpp
@@ -847,7 +847,7 @@ TSharedRef<ISlateStyle> FCoreStyle::Create( const FName& InStyleSetName )
 		Style->Set( "ToolBar.Block.IndentedPadding", FMargin( 18.0f, 2.0f, 4.0f, 4.0f ) );
 		Style->Set( "ToolBar.Block.Padding", FMargin( 2.0f, 2.0f, 4.0f, 4.0f ) );
 
-		Style->Set( "ToolBar.Separator", new BOX_BRUSH( "Old/Button", 4.0f/32.0f ) );
+		Style->Set( "ToolBar.Separator", new FSlateColorBrush( FLinearColor(FColor(48, 48, 48)) ) );
 		Style->Set( "ToolBar.Separator.Padding", FMargin( 0.5f ) );
 
 		Style->Set( "ToolBar.Label", FTextBlockStyle(NormalText) .SetFont( DEFAULT_FONT( "Regular", 9 ) ) );
@@ -1109,16 +1109,20 @@ TSharedRef<ISlateStyle> FCoreStyle::Create( const FName& InStyleSetName )
 			.SetPressed( IMAGE_BRUSH( "/Docking/CloseApp_Pressed", Icon16x16 ) )
 			.SetHovered( IMAGE_BRUSH( "/Docking/CloseApp_Hovered", Icon16x16 ) );
 
+
+		FLinearColor DockColor(FColor(62, 62, 62));
+		FLinearColor InactiveDockColor(FColor(50, 50, 50));
+
 		// Panel Tab
 		Style->Set( "Docking.Tab", FDockTabStyle()
 			.SetCloseButtonStyle( CloseButton )
-			.SetNormalBrush( BOX_BRUSH( "/Docking/Tab_Inactive", 4/16.0f ) )
-			.SetActiveBrush( BOX_BRUSH( "/Docking/Tab_Active", 4/16.0f ) )
-			.SetColorOverlayTabBrush(BOX_BRUSH("/Docking/Tab_ColorOverlay", 4 / 16.0f))
+			.SetNormalBrush( BOX_BRUSH( "/Docking/Tab_Shape", 2.f /8.0f, InactiveDockColor ) )
+			.SetActiveBrush( BOX_BRUSH( "/Docking/Tab_Active", 4./8.0f) ) // 
+			.SetColorOverlayTabBrush( BOX_BRUSH( "/Docking/Tab_ColorOverlay", 4 / 16.0f))
 			.SetColorOverlayIconBrush( BOX_BRUSH( "/Docking/Tab_ColorOverlayIcon", 4/16.0f ) )
-			.SetForegroundBrush( BOX_BRUSH( "/Docking/Tab_Foreground", 4/16.0f ) )
-			.SetHoveredBrush( BOX_BRUSH( "/Docking/Tab_Hovered", 4/16.0f ) )
-			.SetContentAreaBrush( BOX_BRUSH( "/Docking/TabContentArea", FMargin(4/16.0f) ) )
+			.SetForegroundBrush( BOX_BRUSH( "/Docking/Tab_Shape", 2.f /8.0f, DockColor ) )
+			.SetHoveredBrush( BOX_BRUSH( "/Docking/Tab_Shape", 2.f /8.0f, InactiveDockColor ) )
+			.SetContentAreaBrush( FSlateColorBrush( DockColor) )
 			.SetTabWellBrush( FSlateNoResource() )
 			.SetTabPadding( FMargin(5, 2, 5, 2) )
 			.SetOverlapWidth( -1.0f )
