Mode ToolBars - New Mode Toolpalette Style and Adds ability to programmatically change the Active Palette Tab

diff --git a/Engine/Source/Editor/EditorStyle/Private/SlateEditorStyle.cpp b/Engine/Source/Editor/EditorStyle/Private/SlateEditorStyle.cpp
--- a/Engine/Source/Editor/EditorStyle/Private/SlateEditorStyle.cpp
+++ b/Engine/Source/Editor/EditorStyle/Private/SlateEditorStyle.cpp
@@ -2112,6 +2112,55 @@ void FSlateEditorStyle::FStyle::SetupGeneralStyles()
 		Set( "ToolBar.Button.Checked_Pressed", new BOX_BRUSH( "Common/RoundedSelection_16x",  4.0f/16.0f, SelectionColor ) );
 	}
 
+	// Mode ToolPalette 
+	{
+		Set( "PaletteToolBar.Background", 				new BOX_BRUSH( "Common/GroupBorder", FMargin(4.0f/16.0f) ));
+		Set( "PaletteToolBar.Icon", 					new IMAGE_BRUSH( "Icons/icon_tab_Toolbars_16x", Icon16x16 ) );
+		Set( "PaletteToolBar.Expand", 					new IMAGE_BRUSH( "Icons/toolbar_expand_16x", Icon16x16) );
+		Set( "PaletteToolBar.SubMenuIndicator", 		new IMAGE_BRUSH( "Common/SubmenuArrow", Icon8x8 ) );
+
+		// Set( "PaletteToolBar.Content.Padding", FMargin(2.0) );
+		Set( "PaletteToolBar.Label.Padding", FMargin(0.0f, 4.0f, 0.0f, 0.0f) );
+
+		// "SToolBarButtonBlock.Padding" appears as space between the buttons (Outside the Orange selection/hover region)
+		// Outside Margin for Buttons
+		Set( "PaletteToolBar.SToolBarButtonBlock.Padding", 				FMargin(2.f, 0.0f) );
+
+		// SToolBarButtonBlock.CheckBox.Padding is the space inside the toggle button around the icon
+		// Inside Margin for Buttons
+		Set( "PaletteToolBar.SToolBarButtonBlock.CheckBox.Padding", 	FMargin(4.0f, 2.0f, 4.0f, 2.0f) );
+
+		// Outside Margin for Combo Buttons.   
+		// Combo Buttons with Arrows already have an annoying 2.0f built in padding, so we subtract to compensate for that.
+		Set( "PaletteToolBar.SToolBarComboButtonBlock.Padding", 		FMargin(2.0f, 0.0f) );
+
+		// Outside Margin for Generic Widget Blocks.  
+		Set( "PaletteToolBar.Block.IndentedPadding", 					FMargin(4.0f, 2.0f, 4.0f, 2.0f) );
+		Set( "PaletteToolBar.Block.Padding", 							FMargin(4.0f, 2.0f, 4.0f, 2.0f) );
+
+		Set( "PaletteToolBar.SToolBarComboButtonBlock.ComboButton.Color", DefaultForeground );
+
+		Set( "PaletteToolBar.Separator", new FSlateColorBrush( FLinearColor(FColor(47, 47, 47)) ) );
+		Set( "PaletteToolBar.Separator.Padding", FMargin( 2.f, 0.f, 2.f, 0.f) );
+
+		Set( "PaletteToolBar.Label", FTextBlockStyle(NormalText) .SetFont( DEFAULT_FONT( "Regular", 9 ) ) );
+
+		Set( "PaletteToolBar.EditableText",  			GetWidgetStyle<FEditableTextBoxStyle>("ToolBar.EditableText"));
+		Set( "PaletteToolBar.Keybinding", 				GetWidgetStyle<FTextBlockStyle>("ToolBar.Keybinding"));
+		Set( "PaletteToolBar.Heading", 					GetWidgetStyle<FTextBlockStyle>("ToolBar.Heading"));
+		Set( "PaletteToolBar.CheckBox", 				GetWidgetStyle<FCheckBoxStyle>("ToolBar.CheckBox"));
+		Set( "PaletteToolbar.Check", 					GetWidgetStyle<FCheckBoxStyle>("ToolBar.Check"));
+		Set( "PaletteToolBar.RadioButton", 				GetWidgetStyle<FCheckBoxStyle>("ToolBar.RadioButton"));
+		Set( "PaletteToolBar.ToggleButton", 			GetWidgetStyle<FCheckBoxStyle>("ToolBar.ToggleButton"));
+		Set( "PaletteToolBar.Button", 					GetWidgetStyle<FButtonStyle>("ToolBar.Button"));
+		Set( "PaletteToolBar.Button.Normal", 			new FSlateNoResource() );
+		Set( "PaletteToolBar.Button.Pressed", 			new BOX_BRUSH( "Common/RoundedSelection_16x", 4.0f/16.0f, SelectionColor_Pressed ) );
+		Set( "PaletteToolBar.Button.Hovered", 			new BOX_BRUSH( "Common/RoundedSelection_16x", 4.0f/16.0f, SelectionColor ) );
+		Set( "PaletteToolBar.Button.Checked", 			new BOX_BRUSH( "Common/RoundedSelection_16x",  4.0f/16.0f, SelectionColor_Pressed ) );
+		Set( "PaletteToolBar.Button.Checked_Hovered", 	new BOX_BRUSH( "Common/RoundedSelection_16x",  4.0f/16.0f, SelectionColor_Pressed ) );
+		Set( "PaletteToolBar.Button.Checked_Pressed", 	new BOX_BRUSH( "Common/RoundedSelection_16x",  4.0f/16.0f, SelectionColor ) );
+	}
+	
 	// Ctrl+Tab menu
 	{
 		Set("ControlTabMenu.Background", new BOX_BRUSH("Old/Menu_Background", FMargin(8.0f / 64.0f)));
diff --git a/Engine/Source/Editor/UnrealEd/Private/EditorModeManager.cpp b/Engine/Source/Editor/UnrealEd/Private/EditorModeManager.cpp
--- a/Engine/Source/Editor/UnrealEd/Private/EditorModeManager.cpp
+++ b/Engine/Source/Editor/UnrealEd/Private/EditorModeManager.cpp
@@ -786,8 +786,10 @@ void FEditorModeTools::RebuildModeToolBar()
 			ModeToolbarBoxPinned->AddSlot()
 			.Padding(1.f)
 			[
-				PaletteSwitcher	
+				PaletteSwitcher 
 			];
+
+			ModeToolbarPaletteSwitcher = PaletteSwitcher;
 		}
 		else
 		{
@@ -811,6 +813,33 @@ void FEditorModeTools::SpawnOrUpdateModeToolbar()
 	}
 }
 
+void FEditorModeTools::InvokeToolPaletteTab(FEditorModeID InModeID, FName InPaletteName)
+{
+	if (!ModeToolbarPaletteSwitcher.Pin()) 
+	{
+		return;
+	}
+
+	for (auto Row: ActiveToolBarRows)
+	{
+		if (Row.ModeID == InModeID && Row.PaletteName == InPaletteName)
+		{
+			TSharedRef<SWidget> PaletteWidget = Row.ToolbarWidget.ToSharedRef();
+
+			FEdMode* Mode = GetActiveMode(InModeID);
+			TSharedPtr<FModeToolkit> RowToolkit = Mode->GetToolkit();
+
+			TSharedPtr<SWidget> ActiveWidget = ModeToolbarPaletteSwitcher.Pin()->GetActiveWidget();
+			if (ActiveWidget.Get() != Row.ToolbarWidget.Get())
+			{
+				ModeToolbarPaletteSwitcher.Pin()->SetActiveWidget(Row.ToolbarWidget.ToSharedRef());
+				RowToolkit->OnToolPaletteChanged(Row.PaletteName);
+			}
+			break;	
+		}
+	}	
+}
+
 void FEditorModeTools::DeactivateMode( FEditorModeID InID )
 {
 	// Find the mode from the ID and exit it.
@@ -1021,6 +1050,7 @@ void FEditorModeTools::ActivateMode(FEditorModeID InID, bool bToggle)
 				for(auto Palette : PaletteNames)
 				{
 					FToolBarBuilder ModeToolbarBuilder(CommandList, FMultiBoxCustomization(Mode->GetModeInfo().ToolbarCustomizationName), TSharedPtr<FExtender>(), Orient_Horizontal, false);
+					ModeToolbarBuilder.SetStyle(&FEditorStyle::Get(), "PaletteToolBar");
 					Toolkit->BuildToolPalette(Palette, ModeToolbarBuilder);
 
 					ActiveToolBarRows.Emplace(Mode->GetID(), Palette, Toolkit->GetToolPaletteDisplayName(Palette), ModeToolbarBuilder.MakeWidget());
diff --git a/Engine/Source/Editor/UnrealEd/Public/EditorModeManager.h b/Engine/Source/Editor/UnrealEd/Public/EditorModeManager.h
--- a/Engine/Source/Editor/UnrealEd/Public/EditorModeManager.h
+++ b/Engine/Source/Editor/UnrealEd/Public/EditorModeManager.h
@@ -8,6 +8,7 @@
 #include "UnrealWidget.h"
 #include "Editor.h"
 #include "EditorUndoClient.h"
+#include "Widgets/Layout/SWidgetSwitcher.h"
 #include "EdMode.h"
 
 class FCanvas;
@@ -554,6 +555,11 @@ PRAGMA_ENABLE_DEPRECATION_WARNINGS
 	/** returns true if all active EdModes are OK with an AutoSave happening now  */
 	bool CanAutoSave() const;
 
+	/*
+	* Sets the active Modes ToolBar Palette Tab to the named Palette
+	*/
+	void  InvokeToolPaletteTab(FEditorModeID InMode, FName InPaletteName);
+
 protected:
 	/** 
 	 * Delegate handlers
@@ -629,6 +635,9 @@ private:
 	/** The actual toolbar rows will be placed in this vertical box */
 	TWeakPtr<SVerticalBox> ModeToolbarBox;
 
+	/** The modes palette toolbar **/	
+	TWeakPtr<SWidgetSwitcher> ModeToolbarPaletteSwitcher;
+
 	/** Flag set between calls to StartTracking() and EndTracking() */
 	bool bIsTracking;
 };
