Landscape Mode - Migrates tools to Mode Tool Palettes with new icons

diff --git a/Engine/Source/Editor/EditorStyle/Private/SlateEditorStyle.cpp b/Engine/Source/Editor/EditorStyle/Private/SlateEditorStyle.cpp
--- a/Engine/Source/Editor/EditorStyle/Private/SlateEditorStyle.cpp
+++ b/Engine/Source/Editor/EditorStyle/Private/SlateEditorStyle.cpp
@@ -3,6 +3,9 @@
 #include "SlateEditorStyle.h"
 #include "Misc/CommandLine.h"
 #include "Styling/CoreStyle.h"
+
+#include "Classes/EditorStyleSettings.h"
+
 #include "SlateOptMacros.h"
 
 #if (WITH_EDITOR || (IS_PROGRAM && PLATFORM_DESKTOP))
@@ -6942,79 +6945,159 @@ void FSlateEditorStyle::FStyle::SetupLandscapeEditorStyle()
 		Set("LandscapeEditor.ManageMode.Small", new IMAGE_BRUSH("Icons/icon_Landscape_Mode_Manage_20x", Icon20x20));
 		Set("LandscapeEditor.SculptMode.Small", new IMAGE_BRUSH("Icons/icon_Landscape_Mode_Sculpt_20x", Icon20x20));
 		Set("LandscapeEditor.PaintMode.Small",  new IMAGE_BRUSH("Icons/icon_Landscape_Mode_Paint_20x",  Icon20x20));
-		
-		// Tools
-		Set("LandscapeEditor.NewLandscape",       new IMAGE_BRUSH("Icons/icon_Landscape_New_Landscape_40x", Icon40x40));
-		Set("LandscapeEditor.NewLandscape.Small", new IMAGE_BRUSH("Icons/icon_Landscape_New_Landscape_20x", Icon20x20));
-		Set("LandscapeEditor.ResizeLandscape",       new IMAGE_BRUSH("Icons/icon_Landscape_Tool_Copy_40x", Icon40x40));
-		Set("LandscapeEditor.ResizeLandscape.Small", new IMAGE_BRUSH("Icons/icon_Landscape_Tool_Copy_20x", Icon20x20));
-
-		Set("LandscapeEditor.SculptTool",       new IMAGE_BRUSH("Icons/icon_Landscape_Tool_Sculpt_40x",           Icon40x40));
-		Set("LandscapeEditor.EraseTool",		new IMAGE_BRUSH("Icons/icon_Landscape_Tool_Erase_40x", Icon40x40));
-		Set("LandscapeEditor.PaintTool",        new IMAGE_BRUSH("Icons/icon_Landscape_Tool_Paint_40x",            Icon40x40));
-		Set("LandscapeEditor.SmoothTool",       new IMAGE_BRUSH("Icons/icon_Landscape_Tool_Smooth_40x",           Icon40x40));
-		Set("LandscapeEditor.FlattenTool",      new IMAGE_BRUSH("Icons/icon_Landscape_Tool_Flatten_40x",          Icon40x40));
-		Set("LandscapeEditor.RampTool",         new IMAGE_BRUSH("Icons/icon_Landscape_Tool_Ramp_40x",             Icon40x40));
-		Set("LandscapeEditor.ErosionTool",      new IMAGE_BRUSH("Icons/icon_Landscape_Tool_Erosion_40x",          Icon40x40));
-		Set("LandscapeEditor.HydroErosionTool", new IMAGE_BRUSH("Icons/icon_Landscape_Tool_HydraulicErosion_40x", Icon40x40));
-		Set("LandscapeEditor.NoiseTool",        new IMAGE_BRUSH("Icons/icon_Landscape_Tool_Noise_40x",            Icon40x40));
-		Set("LandscapeEditor.RetopologizeTool", new IMAGE_BRUSH("Icons/icon_Landscape_Tool_Retopologize_40x",     Icon40x40));
-		Set("LandscapeEditor.VisibilityTool",   new IMAGE_BRUSH("Icons/icon_Landscape_Tool_Visibility_40x",       Icon40x40));
-		Set("LandscapeEditor.BlueprintBrushTool", new IMAGE_BRUSH("Icons/icon_Landscape_Tool_BlueprintBrush_40x", Icon40x40));
-		Set("LandscapeEditor.SculptTool.Small",       new IMAGE_BRUSH("Icons/icon_Landscape_Tool_Sculpt_20x",           Icon20x20));
-		Set("LandscapeEditor.EraseTool.Small", new IMAGE_BRUSH("Icons/icon_Landscape_Tool_Erase_20x", Icon20x20));
-		Set("LandscapeEditor.PaintTool.Small",        new IMAGE_BRUSH("Icons/icon_Landscape_Tool_Paint_20x",            Icon20x20));
-		Set("LandscapeEditor.SmoothTool.Small",       new IMAGE_BRUSH("Icons/icon_Landscape_Tool_Smooth_20x",           Icon20x20));
-		Set("LandscapeEditor.FlattenTool.Small",      new IMAGE_BRUSH("Icons/icon_Landscape_Tool_Flatten_20x",          Icon20x20));
-		Set("LandscapeEditor.RampTool.Small",         new IMAGE_BRUSH("Icons/icon_Landscape_Tool_Ramp_20x",             Icon20x20));
-		Set("LandscapeEditor.ErosionTool.Small",      new IMAGE_BRUSH("Icons/icon_Landscape_Tool_Erosion_20x",          Icon20x20));
-		Set("LandscapeEditor.HydroErosionTool.Small", new IMAGE_BRUSH("Icons/icon_Landscape_Tool_HydraulicErosion_20x", Icon20x20));
-		Set("LandscapeEditor.NoiseTool.Small",        new IMAGE_BRUSH("Icons/icon_Landscape_Tool_Noise_20x",            Icon20x20));
-		Set("LandscapeEditor.RetopologizeTool.Small", new IMAGE_BRUSH("Icons/icon_Landscape_Tool_Retopologize_20x",     Icon20x20));
-		Set("LandscapeEditor.VisibilityTool.Small",   new IMAGE_BRUSH("Icons/icon_Landscape_Tool_Visibility_20x",       Icon20x20));
-		Set("LandscapeEditor.BlueprintBrushTool.Small", new IMAGE_BRUSH("Icons/icon_Landscape_Tool_BlueprintBrush_20x", Icon20x20));
-
-		Set("LandscapeEditor.SelectComponentTool", new IMAGE_BRUSH("Icons/icon_Landscape_Tool_Selection_40x",       Icon40x40));
-		Set("LandscapeEditor.AddComponentTool",    new IMAGE_BRUSH("Icons/icon_Landscape_Tool_AddComponent_40x",    Icon40x40));
-		Set("LandscapeEditor.DeleteComponentTool", new IMAGE_BRUSH("Icons/icon_Landscape_Tool_DeleteComponent_40x", Icon40x40));
-		Set("LandscapeEditor.MoveToLevelTool",     new IMAGE_BRUSH("Icons/icon_Landscape_Tool_MoveToLevel_40x",     Icon40x40));
-		Set("LandscapeEditor.SelectComponentTool.Small", new IMAGE_BRUSH("Icons/icon_Landscape_Tool_Selection_20x",       Icon20x20));
-		Set("LandscapeEditor.AddComponentTool.Small",    new IMAGE_BRUSH("Icons/icon_Landscape_Tool_AddComponent_20x",    Icon20x20));
-		Set("LandscapeEditor.DeleteComponentTool.Small", new IMAGE_BRUSH("Icons/icon_Landscape_Tool_DeleteComponent_20x", Icon20x20));
-		Set("LandscapeEditor.MoveToLevelTool.Small",     new IMAGE_BRUSH("Icons/icon_Landscape_Tool_MoveToLevel_20x",     Icon20x20));
-
-		Set("LandscapeEditor.RegionSelectTool",    new IMAGE_BRUSH("Icons/icon_Landscape_Tool_Mask_40x", Icon40x40));
-		Set("LandscapeEditor.RegionCopyPasteTool", new IMAGE_BRUSH("Icons/icon_Landscape_Tool_Copy_40x", Icon40x40));
-		Set("LandscapeEditor.RegionSelectTool.Small",    new IMAGE_BRUSH("Icons/icon_Landscape_Tool_Mask_20x", Icon20x20));
-		Set("LandscapeEditor.RegionCopyPasteTool.Small", new IMAGE_BRUSH("Icons/icon_Landscape_Tool_Copy_20x", Icon20x20));
-
-		Set("LandscapeEditor.MirrorTool",       new IMAGE_BRUSH("Icons/icon_Landscape_Tool_Mirror_40x", Icon40x40));
-		Set("LandscapeEditor.MirrorTool.Small", new IMAGE_BRUSH("Icons/icon_Landscape_Tool_Mirror_20x", Icon20x20));
-
-		Set("LandscapeEditor.SplineTool",       new IMAGE_BRUSH("Icons/icon_Landscape_Tool_Splines_40x", Icon40x40));
-		Set("LandscapeEditor.SplineTool.Small", new IMAGE_BRUSH("Icons/icon_Landscape_Tool_Splines_20x", Icon20x20));
-
-		// Brush Sets
-		Set("LandscapeEditor.CircleBrush",        new IMAGE_BRUSH("Icons/icon_Landscape_Brush_Circle_smooth_40x", Icon40x40));
-		Set("LandscapeEditor.AlphaBrush",         new IMAGE_BRUSH("Icons/icon_Landscape_Brush_Alpha_40x",     Icon40x40));
-		Set("LandscapeEditor.AlphaBrush_Pattern", new IMAGE_BRUSH("Icons/icon_Landscape_Brush_Pattern_40x",   Icon40x40));
-		Set("LandscapeEditor.ComponentBrush",     new IMAGE_BRUSH("Icons/icon_Landscape_Brush_Component_40x", Icon40x40));
-		Set("LandscapeEditor.GizmoBrush",         new IMAGE_BRUSH("Icons/icon_Landscape_Brush_Gizmo_40x",     Icon40x40));
-		Set("LandscapeEditor.CircleBrush.Small",        new IMAGE_BRUSH("Icons/icon_Landscape_Brush_Circle_smooth_20x", Icon20x20));
-		Set("LandscapeEditor.AlphaBrush.Small",         new IMAGE_BRUSH("Icons/icon_Landscape_Brush_Alpha_20x",     Icon20x20));
-		Set("LandscapeEditor.AlphaBrush_Pattern.Small", new IMAGE_BRUSH("Icons/icon_Landscape_Brush_Pattern_20x",   Icon20x20));
-		Set("LandscapeEditor.ComponentBrush.Small",     new IMAGE_BRUSH("Icons/icon_Landscape_Brush_Component_20x", Icon20x20));
-		Set("LandscapeEditor.GizmoBrush.Small",         new IMAGE_BRUSH("Icons/icon_Landscape_Brush_Gizmo_20x",     Icon20x20));
-
-		// Brushes
-		Set("LandscapeEditor.CircleBrush_Smooth",    new IMAGE_BRUSH("Icons/icon_Landscape_Brush_Circle_smooth_40x",    Icon40x40));
-		Set("LandscapeEditor.CircleBrush_Linear",    new IMAGE_BRUSH("Icons/icon_Landscape_Brush_Circle_linear_40x",    Icon40x40));
-		Set("LandscapeEditor.CircleBrush_Spherical", new IMAGE_BRUSH("Icons/icon_Landscape_Brush_Circle_spherical_40x", Icon40x40));
-		Set("LandscapeEditor.CircleBrush_Tip",       new IMAGE_BRUSH("Icons/icon_Landscape_Brush_Circle_tip_40x",       Icon40x40));
-		Set("LandscapeEditor.CircleBrush_Smooth.Small",    new IMAGE_BRUSH("Icons/icon_Landscape_Brush_Circle_smooth_20x",    Icon20x20));
-		Set("LandscapeEditor.CircleBrush_Linear.Small",    new IMAGE_BRUSH("Icons/icon_Landscape_Brush_Circle_linear_20x",    Icon20x20));
-		Set("LandscapeEditor.CircleBrush_Spherical.Small", new IMAGE_BRUSH("Icons/icon_Landscape_Brush_Circle_spherical_20x", Icon20x20));
-		Set("LandscapeEditor.CircleBrush_Tip.Small",       new IMAGE_BRUSH("Icons/icon_Landscape_Brush_Circle_tip_20x",       Icon20x20));
+
+		if (GetDefault<UEditorStyleSettings>()->bEnableLegacyEditorModeUI)
+		{
+			// Tools
+			Set("LandscapeEditor.NewLandscape",       new IMAGE_BRUSH("Icons/icon_Landscape_New_Landscape_40x", Icon40x40));
+			Set("LandscapeEditor.NewLandscape.Small", new IMAGE_BRUSH("Icons/icon_Landscape_New_Landscape_20x", Icon20x20));
+			Set("LandscapeEditor.ResizeLandscape",       new IMAGE_BRUSH("Icons/icon_Landscape_Tool_Copy_40x", Icon40x40));
+			Set("LandscapeEditor.ResizeLandscape.Small", new IMAGE_BRUSH("Icons/icon_Landscape_Tool_Copy_20x", Icon20x20));
+
+			Set("LandscapeEditor.SculptTool",       new IMAGE_BRUSH("Icons/icon_Landscape_Tool_Sculpt_40x",           Icon40x40));
+			Set("LandscapeEditor.EraseTool",		new IMAGE_BRUSH("Icons/icon_Landscape_Tool_Erase_40x", Icon40x40));
+			Set("LandscapeEditor.PaintTool",        new IMAGE_BRUSH("Icons/icon_Landscape_Tool_Paint_40x",            Icon40x40));
+			Set("LandscapeEditor.SmoothTool",       new IMAGE_BRUSH("Icons/icon_Landscape_Tool_Smooth_40x",           Icon40x40));
+			Set("LandscapeEditor.FlattenTool",      new IMAGE_BRUSH("Icons/icon_Landscape_Tool_Flatten_40x",          Icon40x40));
+			Set("LandscapeEditor.RampTool",         new IMAGE_BRUSH("Icons/icon_Landscape_Tool_Ramp_40x",             Icon40x40));
+			Set("LandscapeEditor.ErosionTool",      new IMAGE_BRUSH("Icons/icon_Landscape_Tool_Erosion_40x",          Icon40x40));
+			Set("LandscapeEditor.HydroErosionTool", new IMAGE_BRUSH("Icons/icon_Landscape_Tool_HydraulicErosion_40x", Icon40x40));
+			Set("LandscapeEditor.NoiseTool",        new IMAGE_BRUSH("Icons/icon_Landscape_Tool_Noise_40x",            Icon40x40));
+			Set("LandscapeEditor.RetopologizeTool", new IMAGE_BRUSH("Icons/icon_Landscape_Tool_Retopologize_40x",     Icon40x40));
+			Set("LandscapeEditor.VisibilityTool",   new IMAGE_BRUSH("Icons/icon_Landscape_Tool_Visibility_40x",       Icon40x40));
+			Set("LandscapeEditor.BlueprintBrushTool", new IMAGE_BRUSH("Icons/icon_Landscape_Tool_BlueprintBrush_40x", Icon40x40));
+			Set("LandscapeEditor.SculptTool.Small",       new IMAGE_BRUSH("Icons/icon_Landscape_Tool_Sculpt_20x",           Icon20x20));
+			Set("LandscapeEditor.EraseTool.Small", new IMAGE_BRUSH("Icons/icon_Landscape_Tool_Erase_20x", Icon20x20));
+			Set("LandscapeEditor.PaintTool.Small",        new IMAGE_BRUSH("Icons/icon_Landscape_Tool_Paint_20x",            Icon20x20));
+			Set("LandscapeEditor.SmoothTool.Small",       new IMAGE_BRUSH("Icons/icon_Landscape_Tool_Smooth_20x",           Icon20x20));
+			Set("LandscapeEditor.FlattenTool.Small",      new IMAGE_BRUSH("Icons/icon_Landscape_Tool_Flatten_20x",          Icon20x20));
+			Set("LandscapeEditor.RampTool.Small",         new IMAGE_BRUSH("Icons/icon_Landscape_Tool_Ramp_20x",             Icon20x20));
+			Set("LandscapeEditor.ErosionTool.Small",      new IMAGE_BRUSH("Icons/icon_Landscape_Tool_Erosion_20x",          Icon20x20));
+			Set("LandscapeEditor.HydroErosionTool.Small", new IMAGE_BRUSH("Icons/icon_Landscape_Tool_HydraulicErosion_20x", Icon20x20));
+			Set("LandscapeEditor.NoiseTool.Small",        new IMAGE_BRUSH("Icons/icon_Landscape_Tool_Noise_20x",            Icon20x20));
+			Set("LandscapeEditor.RetopologizeTool.Small", new IMAGE_BRUSH("Icons/icon_Landscape_Tool_Retopologize_20x",     Icon20x20));
+			Set("LandscapeEditor.VisibilityTool.Small",   new IMAGE_BRUSH("Icons/icon_Landscape_Tool_Visibility_20x",       Icon20x20));
+			Set("LandscapeEditor.BlueprintBrushTool.Small", new IMAGE_BRUSH("Icons/icon_Landscape_Tool_BlueprintBrush_20x", Icon20x20));
+
+			Set("LandscapeEditor.SelectComponentTool", new IMAGE_BRUSH("Icons/icon_Landscape_Tool_Selection_40x",       Icon40x40));
+			Set("LandscapeEditor.AddComponentTool",    new IMAGE_BRUSH("Icons/icon_Landscape_Tool_AddComponent_40x",    Icon40x40));
+			Set("LandscapeEditor.DeleteComponentTool", new IMAGE_BRUSH("Icons/icon_Landscape_Tool_DeleteComponent_40x", Icon40x40));
+			Set("LandscapeEditor.MoveToLevelTool",     new IMAGE_BRUSH("Icons/icon_Landscape_Tool_MoveToLevel_40x",     Icon40x40));
+			Set("LandscapeEditor.SelectComponentTool.Small", new IMAGE_BRUSH("Icons/icon_Landscape_Tool_Selection_20x",       Icon20x20));
+			Set("LandscapeEditor.AddComponentTool.Small",    new IMAGE_BRUSH("Icons/icon_Landscape_Tool_AddComponent_20x",    Icon20x20));
+			Set("LandscapeEditor.DeleteComponentTool.Small", new IMAGE_BRUSH("Icons/icon_Landscape_Tool_DeleteComponent_20x", Icon20x20));
+			Set("LandscapeEditor.MoveToLevelTool.Small",     new IMAGE_BRUSH("Icons/icon_Landscape_Tool_MoveToLevel_20x",     Icon20x20));
+
+			Set("LandscapeEditor.RegionSelectTool",    new IMAGE_BRUSH("Icons/icon_Landscape_Tool_Mask_40x", Icon40x40));
+			Set("LandscapeEditor.RegionCopyPasteTool", new IMAGE_BRUSH("Icons/icon_Landscape_Tool_Copy_40x", Icon40x40));
+			Set("LandscapeEditor.RegionSelectTool.Small",    new IMAGE_BRUSH("Icons/icon_Landscape_Tool_Mask_20x", Icon20x20));
+			Set("LandscapeEditor.RegionCopyPasteTool.Small", new IMAGE_BRUSH("Icons/icon_Landscape_Tool_Copy_20x", Icon20x20));
+
+			Set("LandscapeEditor.MirrorTool",       new IMAGE_BRUSH("Icons/icon_Landscape_Tool_Mirror_40x", Icon40x40));
+			Set("LandscapeEditor.MirrorTool.Small", new IMAGE_BRUSH("Icons/icon_Landscape_Tool_Mirror_20x", Icon20x20));
+
+			Set("LandscapeEditor.SplineTool",       new IMAGE_BRUSH("Icons/icon_Landscape_Tool_Splines_40x", Icon40x40));
+			Set("LandscapeEditor.SplineTool.Small", new IMAGE_BRUSH("Icons/icon_Landscape_Tool_Splines_20x", Icon20x20));
+
+			// Brush Sets
+			Set("LandscapeEditor.CircleBrush",        new IMAGE_BRUSH("Icons/icon_Landscape_Brush_Circle_smooth_40x", Icon40x40));
+			Set("LandscapeEditor.AlphaBrush",         new IMAGE_BRUSH("Icons/icon_Landscape_Brush_Alpha_40x",     Icon40x40));
+			Set("LandscapeEditor.AlphaBrush_Pattern", new IMAGE_BRUSH("Icons/icon_Landscape_Brush_Pattern_40x",   Icon40x40));
+			Set("LandscapeEditor.ComponentBrush",     new IMAGE_BRUSH("Icons/icon_Landscape_Brush_Component_40x", Icon40x40));
+			Set("LandscapeEditor.GizmoBrush",         new IMAGE_BRUSH("Icons/icon_Landscape_Brush_Gizmo_40x",     Icon40x40));
+			Set("LandscapeEditor.CircleBrush.Small",        new IMAGE_BRUSH("Icons/icon_Landscape_Brush_Circle_smooth_20x", Icon20x20));
+			Set("LandscapeEditor.AlphaBrush.Small",         new IMAGE_BRUSH("Icons/icon_Landscape_Brush_Alpha_20x",     Icon20x20));
+			Set("LandscapeEditor.AlphaBrush_Pattern.Small", new IMAGE_BRUSH("Icons/icon_Landscape_Brush_Pattern_20x",   Icon20x20));
+			Set("LandscapeEditor.ComponentBrush.Small",     new IMAGE_BRUSH("Icons/icon_Landscape_Brush_Component_20x", Icon20x20));
+			Set("LandscapeEditor.GizmoBrush.Small",         new IMAGE_BRUSH("Icons/icon_Landscape_Brush_Gizmo_20x",     Icon20x20));
+
+			// Brushes
+			Set("LandscapeEditor.CircleBrush_Smooth",    new IMAGE_BRUSH("Icons/icon_Landscape_Brush_Circle_smooth_40x",    Icon40x40));
+			Set("LandscapeEditor.CircleBrush_Linear",    new IMAGE_BRUSH("Icons/icon_Landscape_Brush_Circle_linear_40x",    Icon40x40));
+			Set("LandscapeEditor.CircleBrush_Spherical", new IMAGE_BRUSH("Icons/icon_Landscape_Brush_Circle_spherical_40x", Icon40x40));
+			Set("LandscapeEditor.CircleBrush_Tip",       new IMAGE_BRUSH("Icons/icon_Landscape_Brush_Circle_tip_40x",       Icon40x40));
+			Set("LandscapeEditor.CircleBrush_Smooth.Small",    new IMAGE_BRUSH("Icons/icon_Landscape_Brush_Circle_smooth_20x",    Icon20x20));
+			Set("LandscapeEditor.CircleBrush_Linear.Small",    new IMAGE_BRUSH("Icons/icon_Landscape_Brush_Circle_linear_20x",    Icon20x20));
+			Set("LandscapeEditor.CircleBrush_Spherical.Small", new IMAGE_BRUSH("Icons/icon_Landscape_Brush_Circle_spherical_20x", Icon20x20));
+			Set("LandscapeEditor.CircleBrush_Tip.Small",       new IMAGE_BRUSH("Icons/icon_Landscape_Brush_Circle_tip_20x",       Icon20x20));
+		}
+
+		else
+		{
+			// Tools
+			Set("LandscapeEditor.NewLandscape",       			new IMAGE_BRUSH("Icons/Landscape/Landscape_NewLandscape_x40", 		Icon20x20));
+			Set("LandscapeEditor.NewLandscape.Small",       	new IMAGE_BRUSH("Icons/Landscape/Landscape_NewLandscape_x40", 		Icon20x20));
+			// Set("LandscapeEditor.ResizeLandscape		",      new IMAGE_BRUSH("Icons/Landscape/Landscape_Copy_x40", 			Icon40x40));
+			// Set("LandscapeEditor.ResizeLandscape.Small",     new IMAGE_BRUSH("Icons/Landscape/Landscape_Copy_x40", 			Icon40x40));
+
+			Set("LandscapeEditor.SculptTool",       			new IMAGE_BRUSH("Icons/Landscape/Landscape_Sculpt_x40",            Icon20x20));
+			Set("LandscapeEditor.SculptTool.Small",       		new IMAGE_BRUSH("Icons/Landscape/Landscape_Sculpt_x40",            Icon20x20));
+			Set("LandscapeEditor.EraseTool",					new IMAGE_BRUSH("Icons/Landscape/Landscape_Erase_x40",       	   Icon20x20));
+			Set("LandscapeEditor.EraseTool.Small",				new IMAGE_BRUSH("Icons/Landscape/Landscape_Erase_x40",       	   Icon20x20));
+			Set("LandscapeEditor.PaintTool",        			new IMAGE_BRUSH("Icons/Landscape/Landscape_PaintTool_x40",         Icon20x20));
+			Set("LandscapeEditor.PaintTool.Small",        		new IMAGE_BRUSH("Icons/Landscape/Landscape_PaintTool_x40",         Icon20x20));
+			Set("LandscapeEditor.SmoothTool",       			new IMAGE_BRUSH("Icons/Landscape/Landscape_Smooth_x40",            Icon20x20));
+			Set("LandscapeEditor.SmoothTool.Small",       		new IMAGE_BRUSH("Icons/Landscape/Landscape_Smooth_x40",            Icon20x20));
+			Set("LandscapeEditor.FlattenTool",      			new IMAGE_BRUSH("Icons/Landscape/Landscape_Flatten_x40",           Icon20x20));
+			Set("LandscapeEditor.FlattenTool.Small",      		new IMAGE_BRUSH("Icons/Landscape/Landscape_Flatten_x40",           Icon20x20));
+			Set("LandscapeEditor.RampTool",         			new IMAGE_BRUSH("Icons/Landscape/Landscape_Ramp_x40",              Icon20x20));
+			Set("LandscapeEditor.RampTool.Small",         		new IMAGE_BRUSH("Icons/Landscape/Landscape_Ramp_x40",              Icon20x20));
+			Set("LandscapeEditor.ErosionTool",      			new IMAGE_BRUSH("Icons/Landscape/Landscape_Erosion_x40",           Icon20x20));
+			Set("LandscapeEditor.ErosionTool.Small",      		new IMAGE_BRUSH("Icons/Landscape/Landscape_Erosion_x40",           Icon20x20));
+			Set("LandscapeEditor.HydroErosionTool", 			new IMAGE_BRUSH("Icons/Landscape/Landscape_HydroErosion_x40",      Icon20x20));
+			Set("LandscapeEditor.HydroErosionTool.Small", 		new IMAGE_BRUSH("Icons/Landscape/Landscape_HydroErosion_x40",      Icon20x20));
+			Set("LandscapeEditor.NoiseTool",        			new IMAGE_BRUSH("Icons/Landscape/Landscape_Noise_x40",             Icon20x20));
+			Set("LandscapeEditor.NoiseTool.Small",        		new IMAGE_BRUSH("Icons/Landscape/Landscape_Noise_x40",             Icon20x20));
+			Set("LandscapeEditor.RetopologizeTool", 			new IMAGE_BRUSH("Icons/Landscape/Landscape_Retopologize_x40",      Icon20x20));
+			Set("LandscapeEditor.RetopologizeTool.Small", 		new IMAGE_BRUSH("Icons/Landscape/Landscape_Retopologize_x40",      Icon20x20));
+			Set("LandscapeEditor.VisibilityTool",   			new IMAGE_BRUSH("Icons/Landscape/Landscape_VisibilityTool_x40",    Icon20x20));
+			Set("LandscapeEditor.VisibilityTool.Small",   		new IMAGE_BRUSH("Icons/Landscape/Landscape_VisibilityTool_x40",    Icon20x20));
+			Set("LandscapeEditor.BlueprintBrushTool", 			new IMAGE_BRUSH("Icons/Landscape/Landscape_BlueprintTool_x40",     Icon20x20));
+			Set("LandscapeEditor.BlueprintBrushTool.Small", 	new IMAGE_BRUSH("Icons/Landscape/Landscape_BlueprintTool_x40",     Icon20x20));
+
+			Set("LandscapeEditor.SelectComponentTool", 			new IMAGE_BRUSH("Icons/Landscape/Landscape_ComponentSelect_x40",   Icon20x20));
+			Set("LandscapeEditor.SelectComponentTool.Small", 	new IMAGE_BRUSH("Icons/Landscape/Landscape_ComponentSelect_x40",   Icon20x20));
+			Set("LandscapeEditor.AddComponentTool",    			new IMAGE_BRUSH("Icons/Landscape/Landscape_ComponentAdd_x40",      Icon20x20));
+			Set("LandscapeEditor.AddComponentTool.Small",    	new IMAGE_BRUSH("Icons/Landscape/Landscape_ComponentAdd_x40",      Icon20x20));
+			Set("LandscapeEditor.DeleteComponentTool", 			new IMAGE_BRUSH("Icons/Landscape/Landscape_ComponentDelete_x40",   Icon20x20));
+			Set("LandscapeEditor.DeleteComponentTool.Small", 	new IMAGE_BRUSH("Icons/Landscape/Landscape_ComponentDelete_x40",   Icon20x20));
+			Set("LandscapeEditor.MoveToLevelTool",     			new IMAGE_BRUSH("Icons/Landscape/Landscape_ComponentMove_x40",     Icon20x20));
+			Set("LandscapeEditor.MoveToLevelTool.Small",     	new IMAGE_BRUSH("Icons/Landscape/Landscape_ComponentMove_x40",     Icon20x20));
+
+			Set("LandscapeEditor.RegionSelectTool",    			new IMAGE_BRUSH("Icons/Landscape/Landscape_RegionSelect_x40",      Icon20x20));
+			Set("LandscapeEditor.RegionSelectTool.Small",    	new IMAGE_BRUSH("Icons/Landscape/Landscape_RegionSelect_x40",      Icon20x20));
+			Set("LandscapeEditor.RegionCopyPasteTool",			new IMAGE_BRUSH("Icons/Landscape/Landscape_CopyPasteTool_x40",     Icon20x20));
+			Set("LandscapeEditor.RegionCopyPasteTool.Small",	new IMAGE_BRUSH("Icons/Landscape/Landscape_CopyPasteTool_x40",     Icon20x20));
+
+			Set("LandscapeEditor.MirrorTool",       			new IMAGE_BRUSH("Icons/Landscape/Landscape_Mirror_x40",            Icon20x20));
+			Set("LandscapeEditor.MirrorTool.Small",       		new IMAGE_BRUSH("Icons/Landscape/Landscape_Mirror_x40",            Icon20x20));
+
+			Set("LandscapeEditor.SplineTool",       			new IMAGE_BRUSH("Icons/Landscape/Landscape_EditSplines_x40",       Icon20x20));
+			Set("LandscapeEditor.SplineTool.Small",       		new IMAGE_BRUSH("Icons/Landscape/Landscape_EditSplines_x40",       Icon20x20));
+
+			// Brush Sets
+			Set("LandscapeEditor.CircleBrush",        			new IMAGE_BRUSH("Icons/Landscape/Landscape_BrushCircle_x40",       Icon20x20));
+			Set("LandscapeEditor.CircleBrush.Small",        	new IMAGE_BRUSH("Icons/Landscape/Landscape_BrushCircle_x40",       Icon20x20));
+			Set("LandscapeEditor.AlphaBrush",         			new IMAGE_BRUSH("Icons/Landscape/Landscape_BrushAlpha_x40",        Icon20x20));
+			Set("LandscapeEditor.AlphaBrush.Small",         	new IMAGE_BRUSH("Icons/Landscape/Landscape_BrushAlpha_x40",        Icon20x20));
+			Set("LandscapeEditor.AlphaBrush_Pattern", 			new IMAGE_BRUSH("Icons/Landscape/Landscape_BrushPattern_x40",      Icon20x20));
+			Set("LandscapeEditor.AlphaBrush_Pattern.Small", 	new IMAGE_BRUSH("Icons/Landscape/Landscape_BrushPattern_x40",      Icon20x20));
+			Set("LandscapeEditor.ComponentBrush",     			new IMAGE_BRUSH("Icons/Landscape/Landscape_BrushComponent_x40",    Icon20x20));
+			Set("LandscapeEditor.ComponentBrush.Small",     	new IMAGE_BRUSH("Icons/Landscape/Landscape_BrushComponent_x40",    Icon20x20));
+			Set("LandscapeEditor.GizmoBrush",         			new IMAGE_BRUSH("Icons/Landscape/Landscape_BrushGizmo_x40",        Icon20x20));
+			Set("LandscapeEditor.GizmoBrush.Small",         	new IMAGE_BRUSH("Icons/Landscape/Landscape_BrushGizmo_x40",        Icon20x20));
+
+			// Brushes
+			Set("LandscapeEditor.CircleBrush_Smooth",   		new IMAGE_BRUSH("Icons/Landscape/Landscape_FalloffSmooth_x40",     Icon20x20));
+			Set("LandscapeEditor.CircleBrush_Smooth.Small",   	new IMAGE_BRUSH("Icons/Landscape/Landscape_FalloffSmooth_x40",     Icon20x20));
+			Set("LandscapeEditor.CircleBrush_Linear",   		new IMAGE_BRUSH("Icons/Landscape/Landscape_FalloffLinear_x40",     Icon20x20));
+			Set("LandscapeEditor.CircleBrush_Linear.Small",   	new IMAGE_BRUSH("Icons/Landscape/Landscape_FalloffLinear_x40",     Icon20x20));
+			Set("LandscapeEditor.CircleBrush_Spherical",		new IMAGE_BRUSH("Icons/Landscape/Landscape_FalloffSpherical_x40",  Icon20x20));
+			Set("LandscapeEditor.CircleBrush_Spherical.Small",	new IMAGE_BRUSH("Icons/Landscape/Landscape_FalloffSpherical_x40",  Icon20x20));
+			Set("LandscapeEditor.CircleBrush_Tip",      		new IMAGE_BRUSH("Icons/Landscape/Landscape_FalloffTip_x40",        Icon20x20));
+			Set("LandscapeEditor.CircleBrush_Tip.Small",      	new IMAGE_BRUSH("Icons/Landscape/Landscape_FalloffTip_x40",        Icon20x20));
+
+		}
 
 		Set("LandscapeEditor.Brushes.Alpha.UseRChannel", new IMAGE_BRUSH("Icons/icon_Landscape_Brush_Alpha_UseRChannel_20x", Icon20x20));
 		Set("LandscapeEditor.Brushes.Alpha.UseGChannel", new IMAGE_BRUSH("Icons/icon_Landscape_Brush_Alpha_UseGChannel_20x", Icon20x20));
@@ -7046,7 +7129,18 @@ void FSlateEditorStyle::FStyle::SetupLandscapeEditorStyle()
 
 		Set("LandscapeEditor.TargetList.DropZone.Above", new BOX_BRUSH("Common/VerticalBoxDropZoneIndicator_Above", FMargin(10.0f / 16.0f, 10.0f / 16.0f, 0, 0), SelectionColor_Subdued));
 		Set("LandscapeEditor.TargetList.DropZone.Below", new BOX_BRUSH("Common/VerticalBoxDropZoneIndicator_Below", FMargin(10.0f / 16.0f, 0, 0, 10.0f / 16.0f), SelectionColor_Subdued));
-	}
+
+		Set("LandscapeEditor.SpinBox", FSpinBoxStyle(GetWidgetStyle<FSpinBoxStyle>("SpinBox"))
+			.SetTextPadding(FMargin(0))
+			.SetBackgroundBrush(FSlateNoResource())
+			.SetHoveredBackgroundBrush(FSlateNoResource())
+			.SetInactiveFillBrush(FSlateNoResource())
+			.SetActiveFillBrush(FSlateNoResource())
+			.SetForegroundColor(FSlateColor::UseForeground())
+			.SetArrowsImage(FSlateNoResource())
+		);
+}
+
 #endif
 }
 
diff --git a/Engine/Source/Editor/LandscapeEditor/Private/LandscapeEdMode.cpp b/Engine/Source/Editor/LandscapeEditor/Private/LandscapeEdMode.cpp
--- a/Engine/Source/Editor/LandscapeEditor/Private/LandscapeEdMode.cpp
+++ b/Engine/Source/Editor/LandscapeEditor/Private/LandscapeEdMode.cpp
@@ -2256,6 +2256,7 @@ void FEdModeLandscape::SetCurrentToolMode(FName ToolModeName, bool bRestoreCurre
 					}
 					SetCurrentTool(CurrentToolMode->CurrentToolName, CurrentToolMode->CurrentTargetLayerName);
 				}
+				GetModeManager()->InvokeToolPaletteTab(GetID(), ToolModeName);
 				break;
 			}
 		}
diff --git a/Engine/Source/Editor/LandscapeEditor/Private/LandscapeEdModeBrushes.cpp b/Engine/Source/Editor/LandscapeEditor/Private/LandscapeEdModeBrushes.cpp
--- a/Engine/Source/Editor/LandscapeEditor/Private/LandscapeEdModeBrushes.cpp
+++ b/Engine/Source/Editor/LandscapeEditor/Private/LandscapeEdModeBrushes.cpp
@@ -828,7 +828,7 @@ public:
 class FLandscapeBrushDummy : public FLandscapeBrush
 {
 public:
-	const TCHAR* GetBrushName() override { return TEXT("None"); }
+	const TCHAR* GetBrushName() override { return TEXT("Circle_Dummy"); }
 	virtual FText GetDisplayName() override { return NSLOCTEXT("UnrealEd", "LandscapeMode_Brush_None", "None"); };
 
 	FEdModeLandscape* EdMode;
diff --git a/Engine/Source/Editor/LandscapeEditor/Private/LandscapeEditorCommands.cpp b/Engine/Source/Editor/LandscapeEditor/Private/LandscapeEditorCommands.cpp
--- a/Engine/Source/Editor/LandscapeEditor/Private/LandscapeEditorCommands.cpp
+++ b/Engine/Source/Editor/LandscapeEditor/Private/LandscapeEditorCommands.cpp
@@ -27,76 +27,89 @@ void FLandscapeEditorCommands::RegisterCommands()
 	UI_COMMAND(PaintMode, "Mode - Paint", "", EUserInterfaceActionType::RadioButton, FInputChord());
 	NameToCommandMap.Add("ToolMode_Paint", PaintMode);
 
-	UI_COMMAND(NewLandscape, "Tool - New Landscape", "", EUserInterfaceActionType::RadioButton, FInputChord());
+	UI_COMMAND(NewLandscape, "New", "Create or import a new landscape", EUserInterfaceActionType::RadioButton, FInputChord());
 	NameToCommandMap.Add("Tool_NewLandscape", NewLandscape);
 
-	UI_COMMAND(ResizeLandscape, "Tool - Change Component Size", "", EUserInterfaceActionType::RadioButton, FInputChord());
+	UI_COMMAND(ResizeLandscape, "Resize", "Change Component Size", EUserInterfaceActionType::RadioButton, FInputChord());
 	NameToCommandMap.Add("Tool_ResizeLandscape", ResizeLandscape);
 
-	UI_COMMAND(SculptTool, "Tool - Sculpt", "", EUserInterfaceActionType::RadioButton, FInputChord());
+	UI_COMMAND(SculptTool, "Sculpt", "Sculpt height data.\n\nCtrl+Click to Raise, Ctrl+Shift_Click to lower.", EUserInterfaceActionType::RadioButton, FInputChord());
 	NameToCommandMap.Add("Tool_Sculpt", SculptTool);
-	UI_COMMAND(EraseTool, "Tool - Erase", "", EUserInterfaceActionType::RadioButton, FInputChord());
+
+	UI_COMMAND(EraseTool, "Erase", "Erase height data.", EUserInterfaceActionType::RadioButton, FInputChord());
 	NameToCommandMap.Add("Tool_Erase", EraseTool);
-	UI_COMMAND(PaintTool, "Tool - Paint", "", EUserInterfaceActionType::RadioButton, FInputChord());
+
+	UI_COMMAND(PaintTool, "Paint", "Paint weight data.\n\nCtrl+Click to paint, Ctrl+Shift+Click to erase.", EUserInterfaceActionType::RadioButton, FInputChord());
 	NameToCommandMap.Add("Tool_Paint", PaintTool);
-	UI_COMMAND(SmoothTool, "Tool - Smooth", "", EUserInterfaceActionType::RadioButton, FInputChord());
+
+	UI_COMMAND(SmoothTool, "Smooth", "Smooths heightmaps or blend layers.", EUserInterfaceActionType::RadioButton, FInputChord());
 	NameToCommandMap.Add("Tool_Smooth", SmoothTool);
-	UI_COMMAND(FlattenTool, "Tool - Flatten", "", EUserInterfaceActionType::RadioButton, FInputChord());
+
+	UI_COMMAND(FlattenTool, "Flatten", "Flattens an area of heightmap or blend layer.", EUserInterfaceActionType::RadioButton, FInputChord());
 	NameToCommandMap.Add("Tool_Flatten", FlattenTool);
-	UI_COMMAND(RampTool, "Tool - Ramp", "", EUserInterfaceActionType::RadioButton, FInputChord());
+
+	UI_COMMAND(RampTool, "Ramp", "Creates a ramp between two points.", EUserInterfaceActionType::RadioButton, FInputChord());
 	NameToCommandMap.Add("Tool_Ramp", RampTool);
-	UI_COMMAND(ErosionTool, "Tool - Erosion", "", EUserInterfaceActionType::RadioButton, FInputChord());
+
+	UI_COMMAND(ErosionTool, "Erosion", "Thermal Erosion - Simulates erosion caused by the movement of soil from higher areas to lower areas", EUserInterfaceActionType::RadioButton, FInputChord());
 	NameToCommandMap.Add("Tool_Erosion", ErosionTool);
-	UI_COMMAND(HydroErosionTool, "Tool - Hydraulic Erosion", "", EUserInterfaceActionType::RadioButton, FInputChord());
+
+	UI_COMMAND(HydroErosionTool, "Hydro", "Hydro Erosion - Simulates erosion caused by rainfall", EUserInterfaceActionType::RadioButton, FInputChord());
 	NameToCommandMap.Add("Tool_HydraErosion", HydroErosionTool);
-	UI_COMMAND(NoiseTool, "Tool - Noise", "", EUserInterfaceActionType::RadioButton, FInputChord());
+
+	UI_COMMAND(NoiseTool, "Noise", "Adds noise to the heightmap or blend layer", EUserInterfaceActionType::RadioButton, FInputChord());
 	NameToCommandMap.Add("Tool_Noise", NoiseTool);
-	UI_COMMAND(RetopologizeTool, "Tool - Retopologize", "", EUserInterfaceActionType::RadioButton, FInputChord());
+
+	UI_COMMAND(RetopologizeTool, "Retop", "Automatically adjusts landscape vertices with an X/Y offset map to improve vertex density on cliffs, reducing texture stretching.\nNote: An X/Y offset map makes the landscape slower to render and paint on with other tools, so only use if needed", EUserInterfaceActionType::RadioButton, FInputChord());
 	NameToCommandMap.Add("Tool_Retopologize", RetopologizeTool);
-	UI_COMMAND(VisibilityTool, "Tool - Visibility", "", EUserInterfaceActionType::RadioButton, FInputChord());
+
+	UI_COMMAND(VisibilityTool, "Visibility", "Mask out individual quads in the landscape, leaving a hole.", EUserInterfaceActionType::RadioButton, FInputChord());
 	NameToCommandMap.Add("Tool_Visibility", VisibilityTool);
 
-	UI_COMMAND(BlueprintBrushTool, "Tool - Blueprint Brush", "", EUserInterfaceActionType::RadioButton, FInputChord());
+	UI_COMMAND(BlueprintBrushTool, "Blueprint", "Custom painting tools created using Blueprint.", EUserInterfaceActionType::RadioButton, FInputChord());
 	NameToCommandMap.Add("Tool_BlueprintBrush", BlueprintBrushTool);
 
-	UI_COMMAND(SelectComponentTool, "Tool - Component Selection", "", EUserInterfaceActionType::RadioButton, FInputChord());
+	UI_COMMAND(SelectComponentTool, "Select", "Select components to use with other tools", EUserInterfaceActionType::RadioButton, FInputChord());
 	NameToCommandMap.Add("Tool_Select", SelectComponentTool);
-	UI_COMMAND(AddComponentTool, "Tool - Add Components", "", EUserInterfaceActionType::RadioButton, FInputChord());
+
+	UI_COMMAND(AddComponentTool, "Add", "Add components to the landscape", EUserInterfaceActionType::RadioButton, FInputChord());
 	NameToCommandMap.Add("Tool_AddComponent", AddComponentTool);
-	UI_COMMAND(DeleteComponentTool, "Tool - Delete Components", "", EUserInterfaceActionType::RadioButton, FInputChord());
+
+	UI_COMMAND(DeleteComponentTool, "Delete", "Delete components from the landscape, leaving a hole", EUserInterfaceActionType::RadioButton, FInputChord());
 	NameToCommandMap.Add("Tool_DeleteComponent", DeleteComponentTool);
-	UI_COMMAND(MoveToLevelTool, "Tool - Move to Level", "Moves the selected landscape components to the current streaming level", EUserInterfaceActionType::RadioButton, FInputChord());
+
+	UI_COMMAND(MoveToLevelTool, "Move", "Move landscape components to a landscape proxy in the currently active streaming level, so that they can be streamed in/out independently of the rest of the landscape", EUserInterfaceActionType::RadioButton, FInputChord());
 	NameToCommandMap.Add("Tool_MoveToLevel", MoveToLevelTool);
 
-	UI_COMMAND(RegionSelectTool, "Tool - Region Selection", "", EUserInterfaceActionType::RadioButton, FInputChord());
+	UI_COMMAND(RegionSelectTool, "Select", "Select a region of landscape to use as a mask for other tools", EUserInterfaceActionType::RadioButton, FInputChord());
 	NameToCommandMap.Add("Tool_Mask", RegionSelectTool);
-	UI_COMMAND(RegionCopyPasteTool, "Tool - Copy/Paste", "", EUserInterfaceActionType::RadioButton, FInputChord());
+	UI_COMMAND(RegionCopyPasteTool, "Copy", "Copy/Paste areas of the landscape, or import/export a copied area of landscape from disk", EUserInterfaceActionType::RadioButton, FInputChord());
 	NameToCommandMap.Add("Tool_CopyPaste", RegionCopyPasteTool);
-	UI_COMMAND(MirrorTool, "Tool - Mirror", "", EUserInterfaceActionType::RadioButton, FInputChord());
+	UI_COMMAND(MirrorTool, "Mirror", "Copies one side of a landscape to the other, to easily create a mirrored landscape.", EUserInterfaceActionType::RadioButton, FInputChord());
 	NameToCommandMap.Add("Tool_Mirror", MirrorTool);
 
-	UI_COMMAND(SplineTool, "Tool - Edit Splines", "", EUserInterfaceActionType::RadioButton, FInputChord());
+	UI_COMMAND(SplineTool, "Splines", "Ctrl+click to add control points\n\nHaving a control point selected when you ctrl+click will connect to the new control point with a segment\n\nSpline mesh settings can be found on the details panel when you have segments selected", EUserInterfaceActionType::RadioButton, FInputChord());
 	NameToCommandMap.Add("Tool_Splines", SplineTool);
 
-	UI_COMMAND(CircleBrush, "Brush - Circle", "", EUserInterfaceActionType::RadioButton, FInputChord());
+	UI_COMMAND(CircleBrush, "Circle", "Simple circular brush", EUserInterfaceActionType::RadioButton, FInputChord());
 	NameToCommandMap.Add("BrushSet_Circle", CircleBrush);
-	UI_COMMAND(AlphaBrush, "Brush - Alpha", "", EUserInterfaceActionType::RadioButton, FInputChord());
+	UI_COMMAND(AlphaBrush, "Alpha", "Alpha brush, orients a mask image with the brush stroke", EUserInterfaceActionType::RadioButton, FInputChord());
 	NameToCommandMap.Add("BrushSet_Alpha", AlphaBrush);
-	UI_COMMAND(AlphaBrush_Pattern, "Brush - Pattern", "", EUserInterfaceActionType::RadioButton, FInputChord());
+	UI_COMMAND(AlphaBrush_Pattern, "Pattern", "Pattern brush, tiles a mask image across the landscape", EUserInterfaceActionType::RadioButton, FInputChord());
 	NameToCommandMap.Add("BrushSet_Pattern", AlphaBrush_Pattern);
-	UI_COMMAND(ComponentBrush, "Brush - Component", "", EUserInterfaceActionType::RadioButton, FInputChord());
+	UI_COMMAND(ComponentBrush, "Component", "Work with entire landscape components", EUserInterfaceActionType::RadioButton, FInputChord());
 	NameToCommandMap.Add("BrushSet_Component", ComponentBrush);
-	UI_COMMAND(GizmoBrush, "Brush - Gizmo", "", EUserInterfaceActionType::RadioButton, FInputChord());
+	UI_COMMAND(GizmoBrush, "Gizmo", "Work with the landscape gizmo, used for copy/pasting landscape", EUserInterfaceActionType::RadioButton, FInputChord());
 	NameToCommandMap.Add("BrushSet_Gizmo", GizmoBrush);
 	NameToCommandMap.Add("BrushSet_Splines", SplineTool);
 
-	UI_COMMAND(CircleBrush_Smooth, "Circle Brush - Smooth Falloff", "", EUserInterfaceActionType::RadioButton, FInputChord());
+	UI_COMMAND(CircleBrush_Smooth, "Smooth", "Smooth falloff", EUserInterfaceActionType::RadioButton, FInputChord());
 	NameToCommandMap.Add("Circle_Smooth", CircleBrush_Smooth);
-	UI_COMMAND(CircleBrush_Linear, "Circle Brush - Linear Falloff", "", EUserInterfaceActionType::RadioButton, FInputChord());
+	UI_COMMAND(CircleBrush_Linear, "Linear", "Sharp, linear falloff", EUserInterfaceActionType::RadioButton, FInputChord());
 	NameToCommandMap.Add("Circle_Linear", CircleBrush_Linear);
-	UI_COMMAND(CircleBrush_Spherical, "Circle Brush - Spherical Falloff", "", EUserInterfaceActionType::RadioButton, FInputChord());
+	UI_COMMAND(CircleBrush_Spherical, "Spherical", "Spherical falloff, smooth at the center and sharp at the edge", EUserInterfaceActionType::RadioButton, FInputChord());
 	NameToCommandMap.Add("Circle_Spherical", CircleBrush_Spherical);
-	UI_COMMAND(CircleBrush_Tip, "Circle Brush - Tip Falloff", "", EUserInterfaceActionType::RadioButton, FInputChord());
+	UI_COMMAND(CircleBrush_Tip, "Tip", "Tip falloff, sharp at the center and smooth at the edge", EUserInterfaceActionType::RadioButton, FInputChord());
 	NameToCommandMap.Add("Circle_Tip", CircleBrush_Tip);
 
 	UI_COMMAND(ViewModeNormal, "Normal", "", EUserInterfaceActionType::RadioButton, FInputChord());
diff --git a/Engine/Source/Editor/LandscapeEditor/Private/LandscapeEditorDetails.cpp b/Engine/Source/Editor/LandscapeEditor/Private/LandscapeEditorDetails.cpp
--- a/Engine/Source/Editor/LandscapeEditor/Private/LandscapeEditorDetails.cpp
+++ b/Engine/Source/Editor/LandscapeEditor/Private/LandscapeEditorDetails.cpp
@@ -7,6 +7,7 @@
 #include "Framework/MultiBox/MultiBoxBuilder.h"
 #include "Widgets/Input/SComboButton.h"
 #include "Widgets/Input/SMultiLineEditableTextBox.h"
+#include "Widgets/Input/SSpinBox.h"
 #include "LandscapeEdMode.h"
 #include "LandscapeEditorDetailCustomization_NewLandscape.h"
 #include "LandscapeEditorDetailCustomization_ResizeLandscape.h"
@@ -17,15 +18,20 @@
 #include "LandscapeEditorDetailCustomization_TargetLayers.h"
 #include "DetailCategoryBuilder.h"
 #include "DetailLayoutBuilder.h"
+#include "Classes/EditorStyleSettings.h"
+#include "Editor/PropertyEditor/Public/VariablePrecisionNumericInterface.h"
 
+#include "Templates/SharedPointer.h"
+
+#include "SLandscapeEditor.h"
 #include "LandscapeEditorCommands.h"
 #include "LandscapeEditorDetailWidgets.h"
 #include "LandscapeEditorDetailCustomization_LayersBrushStack.h"
+#include "LandscapeEditorObject.h"
 #include "Landscape.h"
 
 #define LOCTEXT_NAMESPACE "LandscapeEditor"
 
-
 TSharedRef<IDetailCustomization> FLandscapeEditorDetails::MakeInstance()
 {
 	return MakeShareable(new FLandscapeEditorDetails);
@@ -149,6 +155,151 @@ void FLandscapeEditorDetails::CustomizeDetails(IDetailLayoutBuilder& DetailBuild
 	Customization_TargetLayers->CustomizeDetails(DetailBuilder);
 }
 
+void FLandscapeEditorDetails::CustomizeToolBarPalette(FToolBarBuilder& ToolBarBuilder, const TSharedRef<FLandscapeToolKit> LandscapeToolkit)
+{
+
+	FEdModeLandscape* LandscapeEdMode = GetEditorMode();
+	CommandList = LandscapeEdMode->GetUICommandList();
+
+	TSharedPtr<INumericTypeInterface<float>> NumericInterface = MakeShareable(new FVariablePrecisionNumericInterface());
+
+	// Tool Strength
+	{
+		UProperty* ToolStrengthProperty = LandscapeEdMode->UISettings->GetClass()->FindPropertyByName(GET_MEMBER_NAME_CHECKED(ULandscapeEditorObject, ToolStrength));
+		const FString& UIMinString = ToolStrengthProperty->GetMetaData("UIMin");
+		const FString& UIMaxString = ToolStrengthProperty->GetMetaData("UIMax");
+		const FString& SliderExponentString = ToolStrengthProperty->GetMetaData("SliderExponent");
+		float UIMin = TNumericLimits<float>::Lowest();
+		float UIMax = TNumericLimits<float>::Max();
+		TTypeFromString<float>::FromString(UIMin, *UIMinString);
+		TTypeFromString<float>::FromString(UIMax, *UIMaxString);
+		float SliderExponent = 1.0f;
+		if (SliderExponentString.Len())
+		{
+			TTypeFromString<float>::FromString(SliderExponent, *SliderExponentString);
+		}
+
+		TSharedRef<SWidget> StrengthControl = SNew(SSpinBox<float>)
+			.Style(&FEditorStyle::Get().GetWidgetStyle<FSpinBoxStyle>("LandscapeEditor.SpinBox"))
+			.PreventThrottling(true)
+			.MinValue(UIMin)
+			.MaxValue(UIMax)
+			.SliderExponent(SliderExponent)
+			.Font(FCoreStyle::GetDefaultFontStyle("Regular", 11))
+			.MinDesiredWidth(40.f)
+			.TypeInterface(NumericInterface)
+			.Justification(ETextJustify::Center)
+			.ContentPadding( FMargin(0.f, 2.f, 0.f, 0.f) )
+
+			.Visibility_Lambda( [ToolStrengthProperty, LandscapeToolkit]() -> EVisibility { return LandscapeToolkit->GetIsPropertyVisibleFromProperty(*ToolStrengthProperty) ? EVisibility::Visible : EVisibility::Collapsed;} )
+			.IsEnabled(this, &FLandscapeEditorDetails::IsBrushSetEnabled) 
+			.Value_Lambda([LandscapeEdMode]() -> float { return LandscapeEdMode->UISettings->ToolStrength; })
+			.OnValueChanged_Lambda([LandscapeEdMode](float NewValue) { LandscapeEdMode->UISettings->ToolStrength = NewValue; });
+
+		ToolBarBuilder.AddToolBarWidget( StrengthControl, LOCTEXT("BrushStrength", "Strength") );
+	}
+
+	//  Brush
+	FUIAction BrushSelectorUIAction;
+	BrushSelectorUIAction.IsActionVisibleDelegate.BindSP(this, &FLandscapeEditorDetails::GetBrushSelectorIsVisible);
+	BrushSelectorUIAction.CanExecuteAction.BindSP(this, &FLandscapeEditorDetails::IsBrushSetEnabled);
+
+	ToolBarBuilder.AddComboButton(
+		BrushSelectorUIAction,
+		FOnGetContent::CreateSP(this, &FLandscapeEditorDetails::GetBrushSelector),
+		TAttribute<FText>(this, &FLandscapeEditorDetails::GetCurrentBrushName),
+		LOCTEXT("BrushSelector.Tooltip", "Select Brush"),
+		TAttribute<FSlateIcon>(this, &FLandscapeEditorDetails::GetCurrentBrushIcon)
+		);
+
+	//  Brush Size 
+	{
+
+		UProperty* BrushRadiusProperty = LandscapeEdMode->UISettings->GetClass()->FindPropertyByName(GET_MEMBER_NAME_CHECKED(ULandscapeEditorObject, BrushRadius));
+		const FString& UIMinString = BrushRadiusProperty->GetMetaData("UIMin");
+		const FString& UIMaxString = BrushRadiusProperty->GetMetaData("UIMax");
+		const FString& SliderExponentString = BrushRadiusProperty->GetMetaData("SliderExponent");
+		float UIMin = TNumericLimits<float>::Lowest();
+		float UIMax = TNumericLimits<float>::Max();
+		TTypeFromString<float>::FromString(UIMin, *UIMinString);
+		TTypeFromString<float>::FromString(UIMax, *UIMaxString);
+		float SliderExponent = 1.0f;
+		if (SliderExponentString.Len())
+		{
+			TTypeFromString<float>::FromString(SliderExponent, *SliderExponentString);
+		}
+
+		TSharedRef<SWidget> SizeControl = 
+			SNew(SSpinBox<float>)
+			.Style(&FEditorStyle::Get().GetWidgetStyle<FSpinBoxStyle>("LandscapeEditor.SpinBox"))
+			.Visibility_Lambda( [BrushRadiusProperty, LandscapeToolkit]() -> EVisibility { return LandscapeToolkit->GetIsPropertyVisibleFromProperty(*BrushRadiusProperty) ? EVisibility::Visible : EVisibility::Collapsed;} )
+			.IsEnabled(this, &FLandscapeEditorDetails::IsBrushSetEnabled) 
+			.PreventThrottling(true)
+			.Value_Lambda([LandscapeEdMode]() -> float { return LandscapeEdMode->UISettings->BrushRadius; })
+			.OnValueChanged_Lambda([LandscapeEdMode](float NewValue) { LandscapeEdMode->UISettings->BrushRadius = NewValue; })
+
+			.MinValue(UIMin)
+			.MaxValue(UIMax)
+			.SliderExponent(SliderExponent)
+			.Font(FCoreStyle::GetDefaultFontStyle("Regular", 11))
+			.MinDesiredWidth(40.f)
+			.TypeInterface(NumericInterface)
+			.Justification(ETextJustify::Center);
+		ToolBarBuilder.AddToolBarWidget( SizeControl, LOCTEXT("BrushRadius", "Radius") );
+	}
+
+	//  Brush Falloff Curve
+	FUIAction BrushFalloffSelectorUIAction;
+	BrushFalloffSelectorUIAction.IsActionVisibleDelegate.BindSP(this, &FLandscapeEditorDetails::GetBrushFalloffSelectorIsVisible);
+	BrushFalloffSelectorUIAction.CanExecuteAction.BindSP(this, &FLandscapeEditorDetails::IsBrushSetEnabled);
+	ToolBarBuilder.AddComboButton(
+		BrushFalloffSelectorUIAction,
+		FOnGetContent::CreateSP(this, &FLandscapeEditorDetails::GetBrushFalloffSelector),
+		TAttribute<FText>(this, &FLandscapeEditorDetails::GetCurrentBrushFalloffName),
+		LOCTEXT("BrushFalloffSelector.Tooltip", "Select Brush Falloff Type"),
+		TAttribute<FSlateIcon>(this, &FLandscapeEditorDetails::GetCurrentBrushFalloffIcon)
+		);
+
+	//  Brush Falloff Percentage 
+	{
+		UProperty* BrushFalloffProperty = LandscapeEdMode->UISettings->GetClass()->FindPropertyByName(GET_MEMBER_NAME_CHECKED(ULandscapeEditorObject, BrushFalloff));
+
+		const FString& UIMinString = BrushFalloffProperty->GetMetaData("UIMin");
+		const FString& UIMaxString = BrushFalloffProperty->GetMetaData("UIMax");
+		const FString& SliderExponentString = BrushFalloffProperty->GetMetaData("SliderExponent");
+		float UIMin = TNumericLimits<float>::Lowest();
+		float UIMax = TNumericLimits<float>::Max();
+		TTypeFromString<float>::FromString(UIMin, *UIMinString);
+		TTypeFromString<float>::FromString(UIMax, *UIMaxString);
+		float SliderExponent = 1.0f;
+		if (SliderExponentString.Len())
+		{
+			TTypeFromString<float>::FromString(SliderExponent, *SliderExponentString);
+		}
+
+		TSharedRef<SWidget> FalloffControl = 
+			SNew(SSpinBox<float>)
+			.Style(&FEditorStyle::Get().GetWidgetStyle<FSpinBoxStyle>("LandscapeEditor.SpinBox"))
+			.Visibility_Lambda( [BrushFalloffProperty, LandscapeToolkit, this]() -> EVisibility { return GetBrushFalloffSelectorIsVisible() && LandscapeToolkit->GetIsPropertyVisibleFromProperty(*BrushFalloffProperty) ? EVisibility::Visible : EVisibility::Collapsed;} )
+			.IsEnabled(this, &FLandscapeEditorDetails::IsBrushSetEnabled) 
+			.PreventThrottling(true)
+			.Value_Lambda([LandscapeEdMode]() -> float { return LandscapeEdMode->UISettings->BrushFalloff; })
+			.OnValueChanged_Lambda([LandscapeEdMode](float NewValue) { LandscapeEdMode->UISettings->BrushFalloff = NewValue; })
+
+			.MinValue(UIMin)
+			.MaxValue(UIMax)
+			.SliderExponent(SliderExponent)
+			.Font(FCoreStyle::GetDefaultFontStyle("Regular", 11))
+			.MinDesiredWidth(40.f)
+			.TypeInterface(NumericInterface)
+			.Justification(ETextJustify::Center);
+
+		ToolBarBuilder.AddToolBarWidget( FalloffControl, LOCTEXT("BrushFalloff", "Falloff") );
+	}
+
+}
+
+
 FText FLandscapeEditorDetails::GetLocalizedName(FString Name)
 {
 	static bool bInitialized = false;
@@ -188,11 +339,13 @@ FText FLandscapeEditorDetails::GetLocalizedName(FString Name)
 		LOCTEXT("BrushSet_Pattern", "Pattern");
 		LOCTEXT("BrushSet_Component", "Component");
 		LOCTEXT("BrushSet_Gizmo", "Gizmo");
+		LOCTEXT("BrushSet_Dummy", "NoBrush");
 
 		LOCTEXT("Circle_Smooth", "Smooth");
 		LOCTEXT("Circle_Linear", "Linear");
 		LOCTEXT("Circle_Spherical", "Spherical");
 		LOCTEXT("Circle_Tip", "Tip");
+		LOCTEXT("Circle_Dummy", "NoBrush");
 	}
 
 	FText Result;
@@ -368,6 +521,11 @@ TSharedRef<SWidget> FLandscapeEditorDetails::GetToolSelector()
 
 bool FLandscapeEditorDetails::GetToolSelectorIsVisible() const
 {
+	if (!GetDefault<UEditorStyleSettings>()->bEnableLegacyEditorModeUI)
+	{
+		return false;
+	}
+
 	FEdModeLandscape* LandscapeEdMode = GetEditorMode();
 	if (LandscapeEdMode && LandscapeEdMode->CurrentTool)
 	{
@@ -543,4 +701,10 @@ bool FLandscapeEditorDetails::GetBrushFalloffSelectorIsVisible() const
 	return false;
 }
 
+bool FLandscapeEditorDetails::IsBrushSetEnabled() const
+{
+	FEdModeLandscape* LandscapeEdMode = GetEditorMode();
+	return (LandscapeEdMode != nullptr && LandscapeEdMode->GetLandscapeList().Num() > 0);
+}
+
 #undef LOCTEXT_NAMESPACE
diff --git a/Engine/Source/Editor/LandscapeEditor/Private/LandscapeEditorDetails.h b/Engine/Source/Editor/LandscapeEditor/Private/LandscapeEditorDetails.h
--- a/Engine/Source/Editor/LandscapeEditor/Private/LandscapeEditorDetails.h
+++ b/Engine/Source/Editor/LandscapeEditor/Private/LandscapeEditorDetails.h
@@ -21,6 +21,18 @@ class FLandscapeEditorDetailCustomization_TargetLayers;
 class FUICommandList;
 class IDetailLayoutBuilder;
 class ULandscapeInfo;
+class FLandscapeToolKit;
+
+
+// 
+// FLandscapeEditorDetails
+// 
+// NOTE: If and when the legacy LandscapeEditor Mode (pre the ToolBar implementation) is removed, 
+// this class can cease to inheriting from DetailsCustomization as it will no longer be used as such.
+//
+// The toolbar implementation directly creates an instances and uses the LandscapeEditorDetails
+// to generate and manage the brush and falloff ui.
+// 
 
 class FLandscapeEditorDetails : public FLandscapeEditorDetailCustomization_Base
 {
@@ -31,6 +43,8 @@ public:
 	/** IDetailCustomization interface */
 	virtual void CustomizeDetails(IDetailLayoutBuilder& DetailBuilder) override;
 
+	void CustomizeToolBarPalette(FToolBarBuilder& ToolBarBuilder, const TSharedRef<FLandscapeToolKit> LandscapeToolkit);
+
 protected:
 	static FText GetLocalizedName(FString Name);
 
@@ -55,6 +69,8 @@ protected:
 	TSharedRef<SWidget> GetBrushFalloffSelector();
 	bool GetBrushFalloffSelectorIsVisible() const;
 
+	bool IsBrushSetEnabled() const;
+
 	TSharedPtr<FUICommandList> CommandList;
 
 	TSharedPtr<FLandscapeEditorDetailCustomization_NewLandscape> Customization_NewLandscape;
diff --git a/Engine/Source/Editor/LandscapeEditor/Private/SLandscapeEditor.cpp b/Engine/Source/Editor/LandscapeEditor/Private/SLandscapeEditor.cpp
--- a/Engine/Source/Editor/LandscapeEditor/Private/SLandscapeEditor.cpp
+++ b/Engine/Source/Editor/LandscapeEditor/Private/SLandscapeEditor.cpp
@@ -9,6 +9,7 @@
 #include "SlateOptMacros.h"
 #include "Widgets/Layout/SBox.h"
 #include "Widgets/Notifications/SErrorText.h"
+#include "Classes/EditorStyleSettings.h"
 #include "EditorStyleSet.h"
 #include "EditorModeManager.h"
 #include "EditorModes.h"
@@ -135,6 +136,11 @@ void FLandscapeToolKit::Init(const TSharedPtr<IToolkitHost>& InitToolkitHost)
 
 	LandscapeEditorWidgets = SNew(SLandscapeEditor, SharedThis(this));
 
+	if (!GetDefault<UEditorStyleSettings>()->bEnableLegacyEditorModeUI)
+	{
+		BrushesWidgets = StaticCastSharedRef<FLandscapeEditorDetails>(FLandscapeEditorDetails::MakeInstance());
+	}
+
 	FModeToolkit::Init(InitToolkitHost);
 }
 
@@ -158,6 +164,132 @@ TSharedPtr<SWidget> FLandscapeToolKit::GetInlineContent() const
 	return LandscapeEditorWidgets;
 }
 
+const TArray<FName> FLandscapeToolKit::PaletteNames = { LandscapeEditorNames::Manage, LandscapeEditorNames::Sculpt, LandscapeEditorNames::Paint };
+
+void FLandscapeToolKit::GetToolPaletteNames(TArray<FName>& InPaletteName) const
+{
+	if (!GetDefault<UEditorStyleSettings>()->bEnableLegacyEditorModeUI)
+	{
+		InPaletteName = PaletteNames;
+	}
+}
+
+FText FLandscapeToolKit::GetToolPaletteDisplayName(FName PaletteName) const
+{
+
+	if (PaletteName == LandscapeEditorNames::Manage)
+	{
+		return LOCTEXT("Mode.Manage", "Manage");
+	}
+	else if (PaletteName == LandscapeEditorNames::Sculpt)
+	{
+		return LOCTEXT("Mode.Sculpt", "Sculpt");
+	}
+	else if (PaletteName == LandscapeEditorNames::Paint)
+	{
+		return LOCTEXT("Mode.Paint",  "Paint");
+	}
+
+	return FText();
+}
+
+void FLandscapeToolKit::BuildToolPalette(FName PaletteName, class FToolBarBuilder& ToolBarBuilder)
+{
+	auto Commands = FLandscapeEditorCommands::Get();
+	FEdModeLandscape* LandscapeEdMode = GetEditorMode();
+	if (PaletteName == LandscapeEditorNames::Manage)
+	{
+		ToolBarBuilder.BeginSection("Manage");
+
+		ToolBarBuilder.AddToolBarButton(Commands.NewLandscape);
+		ToolBarBuilder.AddSeparator();	
+
+		ToolBarBuilder.AddToolBarButton(Commands.SelectComponentTool);
+		ToolBarBuilder.AddToolBarButton(Commands.AddComponentTool);
+		ToolBarBuilder.AddToolBarButton(Commands.DeleteComponentTool);
+		ToolBarBuilder.AddToolBarButton(Commands.MoveToLevelTool);
+
+		ToolBarBuilder.AddSeparator();	
+		ToolBarBuilder.AddToolBarButton(Commands.SplineTool);
+		ToolBarBuilder.EndSection();
+	}
+
+	else if (PaletteName == LandscapeEditorNames::Sculpt)
+	{
+
+		ToolBarBuilder.AddToolBarButton(Commands.SculptTool);
+		if (LandscapeEdMode->CanHaveLandscapeLayersContent())
+		{
+			ToolBarBuilder.AddToolBarButton(Commands.EraseTool);
+		}
+
+		ToolBarBuilder.AddToolBarButton(Commands.SmoothTool);
+		ToolBarBuilder.AddToolBarButton(Commands.FlattenTool);
+		ToolBarBuilder.AddToolBarButton(Commands.RampTool);
+		ToolBarBuilder.AddToolBarButton(Commands.ErosionTool);
+		ToolBarBuilder.AddToolBarButton(Commands.HydroErosionTool);
+		ToolBarBuilder.AddToolBarButton(Commands.NoiseTool);
+		ToolBarBuilder.AddToolBarButton(Commands.RetopologizeTool);
+		ToolBarBuilder.AddToolBarButton(Commands.VisibilityTool);
+		if (LandscapeEdMode->CanHaveLandscapeLayersContent())
+		{
+			ToolBarBuilder.AddToolBarButton(Commands.BlueprintBrushTool);
+		}
+		ToolBarBuilder.AddToolBarButton(Commands.MirrorTool);
+
+		ToolBarBuilder.AddSeparator();
+
+		ToolBarBuilder.AddToolBarButton(Commands.RegionSelectTool);
+		ToolBarBuilder.AddToolBarButton(Commands.RegionCopyPasteTool);
+
+		ToolBarBuilder.AddSeparator();
+
+		if (BrushesWidgets)
+		{
+			BrushesWidgets->CustomizeToolBarPalette(ToolBarBuilder, SharedThis(this));
+		}
+
+	}
+
+	else if (PaletteName == LandscapeEditorNames::Paint)
+	{
+
+		ToolBarBuilder.AddToolBarButton(Commands.PaintTool);
+		ToolBarBuilder.AddToolBarButton(Commands.SmoothTool);
+		ToolBarBuilder.AddToolBarButton(Commands.FlattenTool);
+		ToolBarBuilder.AddToolBarButton(Commands.NoiseTool);
+		if (LandscapeEdMode->CanHaveLandscapeLayersContent())
+		{
+			ToolBarBuilder.AddToolBarButton(Commands.BlueprintBrushTool);
+		}
+
+		ToolBarBuilder.AddSeparator();	
+
+		if (BrushesWidgets)
+		{
+			BrushesWidgets->CustomizeToolBarPalette(ToolBarBuilder, SharedThis(this));
+		}
+	}
+
+}
+
+void FLandscapeToolKit::OnToolPaletteChanged(FName PaletteName)
+{
+	if (PaletteName == LandscapeEditorNames::Manage && !IsModeActive(LandscapeEditorNames::Manage ))
+	{
+		OnChangeMode(LandscapeEditorNames::Manage);
+	}
+	else if (PaletteName == LandscapeEditorNames::Sculpt && !IsModeActive(LandscapeEditorNames::Sculpt ))
+	{
+		OnChangeMode(LandscapeEditorNames::Sculpt); 
+	}
+	else if (PaletteName == LandscapeEditorNames::Paint  && !IsModeActive(LandscapeEditorNames::Paint ))
+	{
+		OnChangeMode(LandscapeEditorNames::Paint); 
+	}
+}
+
+
 void FLandscapeToolKit::NotifyToolChanged()
 {
 	LandscapeEditorWidgets->NotifyToolChanged();
@@ -188,7 +320,7 @@ bool FLandscapeToolKit::IsModeEnabled(FName ModeName) const
 	if (LandscapeEdMode)
 	{
 		// Manage is the only mode enabled if we have no landscape
-		if (ModeName == "ToolMode_Manage" || (LandscapeEdMode->GetLandscapeList().Num() > 0 && LandscapeEdMode->CanEditCurrentTarget()))
+		if (ModeName == LandscapeEditorNames::Manage || (LandscapeEdMode->GetLandscapeList().Num() > 0 && LandscapeEdMode->CanEditCurrentTarget()))
 		{
 			return true;
 		}
@@ -255,7 +387,9 @@ void FLandscapeToolKit::OnChangeBrushSet(FName BrushSetName)
 bool FLandscapeToolKit::IsBrushSetEnabled(FName BrushSetName) const
 {
 	FEdModeLandscape* LandscapeEdMode = GetEditorMode();
-	if (LandscapeEdMode && LandscapeEdMode->CurrentTool)
+	if (LandscapeEdMode 
+		&& LandscapeEdMode->IsEditingEnabled()
+		&& LandscapeEdMode->CurrentTool )
 	{
 		return LandscapeEdMode->CurrentTool->ValidBrushes.Contains(BrushSetName);
 	}
@@ -301,6 +435,7 @@ bool FLandscapeToolKit::IsBrushActive(FName BrushName) const
 BEGIN_SLATE_FUNCTION_BUILD_OPTIMIZATION
 void SLandscapeEditor::Construct(const FArguments& InArgs, TSharedRef<FLandscapeToolKit> InParentToolkit)
 {
+	ParentToolkit = InParentToolkit;
 	TSharedRef<FUICommandList> CommandList = InParentToolkit->GetToolkitCommands();
 
 	// Modes:
@@ -347,6 +482,7 @@ void SLandscapeEditor::Construct(const FArguments& InArgs, TSharedRef<FLandscape
 				+ SOverlay::Slot()
 				[
 					SNew(SBorder)
+					.Visibility_Lambda( [] () -> EVisibility { return GetDefault<UEditorStyleSettings>()->bEnableLegacyEditorModeUI ? EVisibility::Visible : EVisibility::Collapsed; } )
 					.BorderImage(FEditorStyle::GetBrush("ToolPanel.GroupBorder"))
 					.HAlign(HAlign_Center)
 					[
@@ -450,8 +586,11 @@ bool SLandscapeEditor::GetLandscapeEditorIsEnabled() const
 
 bool SLandscapeEditor::GetIsPropertyVisible(const FPropertyAndParent& PropertyAndParent) const
 {
-	const UProperty& Property = PropertyAndParent.Property;
+	return ParentToolkit.Pin()->GetIsPropertyVisibleFromProperty(PropertyAndParent.Property);
+}
 
+bool FLandscapeToolKit::GetIsPropertyVisibleFromProperty(const UProperty& Property) const
+{
 	FEdModeLandscape* LandscapeEdMode = GetEditorMode();
 	
 	if (LandscapeEdMode != nullptr && LandscapeEdMode->CurrentTool != nullptr)
diff --git a/Engine/Source/Editor/LandscapeEditor/Private/SLandscapeEditor.h b/Engine/Source/Editor/LandscapeEditor/Private/SLandscapeEditor.h
--- a/Engine/Source/Editor/LandscapeEditor/Private/SLandscapeEditor.h
+++ b/Engine/Source/Editor/LandscapeEditor/Private/SLandscapeEditor.h
@@ -6,6 +6,7 @@
 #include "Widgets/DeclarativeSyntaxSupport.h"
 #include "Widgets/SCompoundWidget.h"
 #include "LandscapeEdMode.h"
+#include "LandscapeEditorDetails.h"
 #include "Toolkits/IToolkitHost.h"
 #include "AssetThumbnail.h"
 #include "Toolkits/BaseToolkit.h"
@@ -38,6 +39,13 @@ private:
 	TSharedPtr<FAssetThumbnail> AssetThumbnail;
 };
 
+namespace LandscapeEditorNames
+{
+	static const FName Manage(TEXT("ToolMode_Manage")); 
+	static const FName Sculpt(TEXT("ToolMode_Sculpt")); 
+	static const FName Paint(TEXT("ToolMode_Paint"));
+}
+
 /**
  * Mode Toolkit for the Landscape Editor Mode
  */
@@ -60,6 +68,14 @@ public:
 	void NotifyBrushChanged();
 	void RefreshDetailPanel();
 
+	/** Mode Toolbar Palettes **/
+	virtual void GetToolPaletteNames(TArray<FName>& InPaletteName) const;
+	virtual FText GetToolPaletteDisplayName(FName PaletteName) const; 
+	virtual void BuildToolPalette(FName PaletteName, class FToolBarBuilder& ToolbarBuilder);
+	virtual void OnToolPaletteChanged(FName PaletteName) override;
+
+	bool GetIsPropertyVisibleFromProperty(const UProperty& Property) const;
+
 protected:
 	void OnChangeMode(FName ModeName);
 	bool IsModeEnabled(FName ModeName) const;
@@ -77,8 +93,10 @@ protected:
 	bool IsBrushActive(FName BrushName) const;
 
 private:
-	/** Geometry tools widget */
 	TSharedPtr<SLandscapeEditor> LandscapeEditorWidgets;
+	TSharedPtr<FLandscapeEditorDetails> BrushesWidgets;
+
+	const static TArray<FName> PaletteNames;
 };
 
 /**
@@ -107,6 +125,6 @@ protected:
 
 protected:
 	TSharedPtr<SErrorText> Error;
-
 	TSharedPtr<IDetailsView> DetailsPanel;
+	TWeakPtr<FLandscapeToolKit> ParentToolkit;
 };
diff --git a/Engine/Source/Editor/PropertyEditor/Public/VariablePrecisionNumericInterface.h b/Engine/Source/Editor/PropertyEditor/Public/VariablePrecisionNumericInterface.h
new file mode 100644
--- /dev/null
+++ b/Engine/Source/Editor/PropertyEditor/Public/VariablePrecisionNumericInterface.h
@@ -0,0 +1,38 @@
+
+// Copyright 1998-2019 Epic Games, Inc. All Rights Reserved.
+#pragma once
+
+#include "Widgets/Input/NumericTypeInterface.h"
+
+//
+// FVariablePrecisionNumericInterface
+// 
+// allow more precision as the numbers get closer to zero
+// 
+struct FVariablePrecisionNumericInterface : public TDefaultNumericTypeInterface<float>
+{
+	FVariablePrecisionNumericInterface() {}
+
+	/** Convert the type to/from a string */
+	virtual FString ToString(const float& Value) const override
+	{
+		// 1000
+		// 100.1
+		// 10.12
+		// 1.123
+
+		int32 MaxDigits = 3;
+		if ((Value / 1000.f) >= 1.f)
+			MaxDigits = 0;
+		else if ((Value / 100.f) >= 1.f)
+			MaxDigits = 1;
+		else if ((Value / 10.f) >= 1.f)
+			MaxDigits = 2;
+
+		const FNumberFormattingOptions NumberFormattingOptions = FNumberFormattingOptions()
+			.SetUseGrouping(false)
+			.SetMinimumFractionalDigits( MaxDigits )
+			.SetMaximumFractionalDigits( MaxDigits );
+		return FastDecimalFormat::NumberToString(Value, ExpressionParser::GetLocalizedNumberFormattingRules(), NumberFormattingOptions);
+	}
+};
\ No newline at end of file
