Editor - Adding some ClassIcon.*s icons for Blueprint CallFunction nodes and variable nodes.

diff --git a/Engine/Source/Editor/EditorStyle/Private/SlateEditorStyle.cpp b/Engine/Source/Editor/EditorStyle/Private/SlateEditorStyle.cpp
--- a/Engine/Source/Editor/EditorStyle/Private/SlateEditorStyle.cpp
+++ b/Engine/Source/Editor/EditorStyle/Private/SlateEditorStyle.cpp
@@ -4672,6 +4672,12 @@ void FSlateEditorStyle::FStyle::SetupGraphEditorStyles()
 			);
 		}
 
+		Set("ClassIcon.K2Node_CallFunction", new IMAGE_BRUSH("Icons/icon_Blueprint_NewFunction_16x", Icon16x16));
+		Set("ClassIcon.K2Node_Variable", new IMAGE_BRUSH("Graph/Icons/Node", Icon22x22, FLinearColor::White));
+		Set("ClassIcon.K2Node_VariableGet", new IMAGE_BRUSH("Graph/Icons/FIB_VarGet", Icon22x22, FLinearColor::White));
+		Set("ClassIcon.K2Node_VariableSet", new IMAGE_BRUSH("Graph/Icons/FIB_VarSet", Icon22x22, FLinearColor::White));
+
+
 		Set( "GraphEditor.Default_16x", new IMAGE_BRUSH("Icons/icon_Blueprint_Node_16x", Icon16x16));
 		Set( "GraphEditor.EventGraph_16x", new IMAGE_BRUSH( "Icons/icon_Blueprint_EventGraph_16x", Icon16x16 ) );
 		Set( "GraphEditor.InterfaceFunction_16x", new IMAGE_BRUSH( "Icons/icon_Blueprint_Interfacefunction_16x", Icon16x16 ) );
