Fix Resize Landscape Tool not showing up in manage landscape toolbar

diff --git a/Engine/Source/Editor/EditorStyle/Private/SlateEditorStyle.cpp b/Engine/Source/Editor/EditorStyle/Private/SlateEditorStyle.cpp
--- a/Engine/Source/Editor/EditorStyle/Private/SlateEditorStyle.cpp
+++ b/Engine/Source/Editor/EditorStyle/Private/SlateEditorStyle.cpp
@@ -7159,10 +7159,11 @@ void FSlateEditorStyle::FStyle::SetupLandscapeEditorStyle()
 		else
 		{
 			// Tools
-			Set("LandscapeEditor.NewLandscape",       			new IMAGE_BRUSH("Icons/Landscape/Landscape_NewLandscape_x40", 		Icon20x20));
+			Set("LandscapeEditor.NewLandscape",					new IMAGE_BRUSH("Icons/Landscape/Landscape_NewLandscape_x40",		Icon20x20));
 			Set("LandscapeEditor.NewLandscape.Small",       	new IMAGE_BRUSH("Icons/Landscape/Landscape_NewLandscape_x40", 		Icon20x20));
-			// Set("LandscapeEditor.ResizeLandscape		",      new IMAGE_BRUSH("Icons/Landscape/Landscape_Copy_x40", 			Icon40x40));
-			// Set("LandscapeEditor.ResizeLandscape.Small",     new IMAGE_BRUSH("Icons/Landscape/Landscape_Copy_x40", 			Icon40x40));
+			Set("LandscapeEditor.ResizeLandscape",				new IMAGE_BRUSH("Icons/Landscape/Landscape_Resize_x40",				Icon20x20));
+			Set("LandscapeEditor.ResizeLandscape.Small",		new IMAGE_BRUSH("Icons/Landscape/Landscape_Resize_x40",				Icon20x20));
+
 
 			Set("LandscapeEditor.SculptTool",       			new IMAGE_BRUSH("Icons/Landscape/Landscape_Sculpt_x40",            Icon20x20));
 			Set("LandscapeEditor.SculptTool.Small",       		new IMAGE_BRUSH("Icons/Landscape/Landscape_Sculpt_x40",            Icon20x20));
diff --git a/Engine/Source/Editor/LandscapeEditor/Private/SLandscapeEditor.cpp b/Engine/Source/Editor/LandscapeEditor/Private/SLandscapeEditor.cpp
--- a/Engine/Source/Editor/LandscapeEditor/Private/SLandscapeEditor.cpp
+++ b/Engine/Source/Editor/LandscapeEditor/Private/SLandscapeEditor.cpp
@@ -208,6 +208,7 @@ void FLandscapeToolKit::BuildToolPalette(FName PaletteName, class FToolBarBuilde
 		ToolBarBuilder.AddToolBarButton(Commands.AddComponentTool);
 		ToolBarBuilder.AddToolBarButton(Commands.DeleteComponentTool);
 		ToolBarBuilder.AddToolBarButton(Commands.MoveToLevelTool);
+		ToolBarBuilder.AddToolBarButton(Commands.ResizeLandscape);
 
 		ToolBarBuilder.AddSeparator();	
 		ToolBarBuilder.AddToolBarButton(Commands.SplineTool);
