Added a toolbar style to avoid so much duplication of styles related to toolbars in the editor
Fixed up all obvious use cases

diff --git a/Engine/Plugins/Developer/RenderDocPlugin/Source/RenderDocPlugin/Private/SRenderDocPluginEditorExtension.cpp b/Engine/Plugins/Developer/RenderDocPlugin/Source/RenderDocPlugin/Private/SRenderDocPluginEditorExtension.cpp
--- a/Engine/Plugins/Developer/RenderDocPlugin/Source/RenderDocPlugin/Private/SRenderDocPluginEditorExtension.cpp
+++ b/Engine/Plugins/Developer/RenderDocPlugin/Source/RenderDocPlugin/Private/SRenderDocPluginEditorExtension.cpp
@@ -32,13 +32,12 @@ public:
 	void Construct(const FArguments& Args)
 	{
 		FSlateIcon IconBrush = FSlateIcon(FRenderDocPluginStyle::Get()->GetStyleSetName(), "RenderDocPlugin.CaptureFrameIcon");
-		
+	
 		ChildSlot
 		[
 			SNew(SButton)
 			.HAlign(HAlign_Center)
 			.VAlign(VAlign_Bottom)
-			.ButtonStyle(FEditorStyle::Get(), "ViewportMenu.Button")
 			.ContentPadding(FMargin(1.0f))
 			.ToolTipText(FRenderDocPluginCommands::Get().CaptureFrameCommand->GetDescription())
 			.OnClicked_Lambda([this]() 
diff --git a/Engine/Plugins/Editor/MacGraphicsSwitching/Source/MacGraphicsSwitching/Private/MacGraphicsSwitchingModule.cpp b/Engine/Plugins/Editor/MacGraphicsSwitching/Source/MacGraphicsSwitching/Private/MacGraphicsSwitchingModule.cpp
--- a/Engine/Plugins/Editor/MacGraphicsSwitching/Source/MacGraphicsSwitching/Private/MacGraphicsSwitchingModule.cpp
+++ b/Engine/Plugins/Editor/MacGraphicsSwitching/Source/MacGraphicsSwitching/Private/MacGraphicsSwitchingModule.cpp
@@ -68,7 +68,8 @@ void FMacGraphicsSwitchingModule::Initialize( TSharedPtr<SWindow> InRootWindow,
 		bool MacUseAutomaticGraphicsSwitching = false;
 		if ( GConfig->GetBool(TEXT("/Script/MacGraphicsSwitching.MacGraphicsSwitchingSettings"), TEXT("bShowGraphicsSwitching"), MacUseAutomaticGraphicsSwitching, GEditorSettingsIni) && MacUseAutomaticGraphicsSwitching )
 		{
-			NotificationBarExtender = MakeShareable( new FExtender() );
+			// Switch to status bar
+			/*NotificationBarExtender = MakeShareable( new FExtender() );
 			NotificationBarExtender->AddToolBarExtension("Start",
 														 EExtensionHook::After,
 														 nullptr,
@@ -76,7 +77,7 @@ void FMacGraphicsSwitchingModule::Initialize( TSharedPtr<SWindow> InRootWindow,
 			
 			FLevelEditorModule& LevelEditorModule = FModuleManager::LoadModuleChecked<FLevelEditorModule>("LevelEditor");
 			LevelEditorModule.GetNotificationBarExtensibilityManager()->AddExtender( NotificationBarExtender );
-			LevelEditorModule.BroadcastTitleBarMessagesChanged();
+			LevelEditorModule.BroadcastTitleBarMessagesChanged();*/
 		}
 	}
 }
@@ -98,7 +99,8 @@ void FMacGraphicsSwitchingModule::ShutdownModule()
 	if ( FModuleManager::Get().IsModuleLoaded("LevelEditor") && NotificationBarExtender.IsValid() )
 	{
 		FLevelEditorModule& LevelEditorModule = FModuleManager::GetModuleChecked<FLevelEditorModule>("LevelEditor");
-		LevelEditorModule.GetNotificationBarExtensibilityManager()->RemoveExtender( NotificationBarExtender );
+		// Switch to status bar
+		//LevelEditorModule.GetNotificationBarExtensibilityManager()->RemoveExtender( NotificationBarExtender );
 	}
 	
 	if ( FModuleManager::Get().IsModuleLoaded("MainFrame") )
diff --git a/Engine/Plugins/Enterprise/DataprepEditor/Source/DataprepEditor/Private/DataprepEditorStyle.cpp b/Engine/Plugins/Enterprise/DataprepEditor/Source/DataprepEditor/Private/DataprepEditorStyle.cpp
--- a/Engine/Plugins/Enterprise/DataprepEditor/Source/DataprepEditor/Private/DataprepEditorStyle.cpp
+++ b/Engine/Plugins/Enterprise/DataprepEditor/Source/DataprepEditor/Private/DataprepEditorStyle.cpp
@@ -90,7 +90,7 @@ void FDataprepEditorStyle::Initialize()
 			TilteTextStyle.SetFont(FCoreStyle::GetDefaultFontStyle("Bold", 11));
 			StyleSet->Set("DataprepAction.TitleTextStyle", TilteTextStyle);
 
-			FEditableTextBoxStyle TitleEditableText = FEditorStyle::GetWidgetStyle< FEditableTextBoxStyle >("ViewportMenu.EditableText");
+			FEditableTextBoxStyle TitleEditableText = FEditorStyle::Get().GetWidgetStyle<FToolBarStyle>("ViewportMenu").EditableTextStyle;
 			TitleEditableText.SetFont(FCoreStyle::GetDefaultFontStyle("Bold", 11));
 
 			StyleSet->Set( "DataprepAction.TitleInlineEditableText", FInlineEditableTextBlockStyle()
diff --git a/Engine/Source/Editor/EditorStyle/Private/StarshipStyle.cpp b/Engine/Source/Editor/EditorStyle/Private/StarshipStyle.cpp
--- a/Engine/Source/Editor/EditorStyle/Private/StarshipStyle.cpp
+++ b/Engine/Source/Editor/EditorStyle/Private/StarshipStyle.cpp
@@ -12,6 +12,7 @@
 #if (WITH_EDITOR || (IS_PROGRAM && PLATFORM_DESKTOP))
 	#include "PlatformInfo.h"
 #endif
+#include "Styling/ToolBarStyle.h"
 
 
 #define DEFAULT_FONT(...) FStarshipCoreStyle::GetDefaultFontStyle(__VA_ARGS__)
@@ -214,7 +215,6 @@ void FStarshipEditorStyle::FStyle::Initialize()
 	SetupWorldBrowserStyles();
 	SetupSequencerStyles();
 	SetupViewportStyles();
-	SetupNotificationBarStyles();
 	SetupMenuBarStyles();
 	SetupGeneralIcons();
 	SetupWindowStyles();
@@ -1430,7 +1430,14 @@ void FStarshipEditorStyle::FStyle::SetupGeneralStyles()
 		FLinearColor DimBackgroundHover = FLinearColor(FColor(50, 50, 50));
 		FLinearColor DarkBackground = FLinearColor(FColor(42, 42, 42));
 
-		Set("FoliageEditToolBar.ToggleButton", FCheckBoxStyle()
+		FToolBarStyle FoliageEditToolBar = FStarshipCoreStyle::GetCoreStyle().GetWidgetStyle<FToolBarStyle>("ToolBar");
+
+		FoliageEditToolBar.SetButtonPadding(FMargin(0.f));
+		FoliageEditToolBar.SetComboButtonPadding(FMargin(4.0f));
+		FoliageEditToolBar.SetCheckBoxPadding(FMargin(10.0f, 6.f));
+		FoliageEditToolBar.SetSeparatorPadding(1.0f);
+		FoliageEditToolBar.SetToggleButtonStyle(
+			FCheckBoxStyle()
 			.SetCheckBoxType(ESlateCheckBoxType::ToggleButton)
 			.SetUncheckedImage(BOX_BRUSH("Common/Selection", 8.0f / 32.0f, DimBackground))
 			.SetUncheckedPressedImage(BOX_BRUSH("PlacementMode/TabActive", 8.0f / 32.0f))
@@ -1440,29 +1447,7 @@ void FStarshipEditorStyle::FStyle::SetupGeneralStyles()
 			.SetCheckedPressedImage(BOX_BRUSH("PlacementMode/TabActive", 8.0f / 32.0f))
 			.SetPadding(0));
 
-		Set("FoliageEditToolBar.Background", new FSlateColorBrush(FStyleColors::Background));
-		Set("FoliageEditToolBar.Icon", new IMAGE_BRUSH("Icons/icon_tab_Toolbars_16x", Icon16x16));
-		Set("FoliageEditToolBar.Expand", new IMAGE_BRUSH("Icons/toolbar_expand_16x", Icon16x16));
-		Set("FoliageEditToolBar.SubMenuIndicator", new IMAGE_BRUSH("Common/SubmenuArrow", Icon8x8));
-		Set("FoliageEditToolBar.SToolBarComboButtonBlock.Padding", FMargin(4.0f));
-		Set("FoliageEditToolBar.SToolBarButtonBlock.Padding", FMargin(0.f));
-		Set("FoliageEditToolBar.SToolBarCheckComboButtonBlock.Padding", FMargin(4.0f));
-		Set("FoliageEditToolBar.SToolBarButtonBlock.CheckBox.Padding", FMargin(10.0f, 6.f));
-		Set("FoliageEditToolBar.SToolBarComboButtonBlock.ComboButton.Color", DefaultForeground);
-
-		Set("FoliageEditToolBar.Block.IndentedPadding", FMargin(18.0f, 2.0f, 4.0f, 4.0f));
-		Set("FoliageEditToolBar.Block.Padding", FMargin(2.0f, 2.0f, 4.0f, 4.0f));
-
-		Set("FoliageEditToolBar.Separator", new BOX_BRUSH("Old/Button", 4.0f / 32.0f));
-		Set("FoliageEditToolBar.Separator.Padding", FMargin(0.5f));
-
-		Set("FoliageEditToolBar.Label", FTextBlockStyle(NormalText).SetFont(DEFAULT_FONT("Regular", 9)));
-		Set("FoliageEditToolBar.EditableText", FEditableTextBoxStyle(NormalEditableTextBoxStyle).SetFont(DEFAULT_FONT("Regular", 9)));
-		Set("FoliageEditToolBar.Keybinding", FTextBlockStyle(NormalText).SetFont(DEFAULT_FONT("Regular", 8)));
-
-		Set("FoliageEditToolBar.Heading", FTextBlockStyle(NormalText)
-			.SetFont(DEFAULT_FONT("Regular", 8))
-			.SetColorAndOpacity(FLinearColor(0.4f, 0.4, 0.4f, 1.0f)));
+		Set("FoliageEditToolBar", FoliageEditToolBar);
 
 		Set("FoliageEditMode.SetSelect",                 new IMAGE_BRUSH("Icons/GeneralTools/Select_40x", Icon20x20));
 		Set("FoliageEditMode.SetSelect.Small",           new IMAGE_BRUSH("Icons/GeneralTools/Select_40x", Icon20x20));
@@ -1730,43 +1715,15 @@ void FStarshipEditorStyle::FStyle::SetupGeneralStyles()
 #if WITH_EDITOR || (IS_PROGRAM && WITH_UNREAL_DEVELOPER_TOOLS)
 	// Mode ToolPalette 
 	{
-		Set( "PaletteToolBar.Background", 				new FSlateColorBrush(FStyleColors::Background));
-		Set( "PaletteToolBar.Icon", 					new IMAGE_BRUSH( "Icons/icon_tab_Toolbars_16x", Icon16x16 ) );
-		Set( "PaletteToolBar.Expand", 					new IMAGE_BRUSH( "Icons/toolbar_expand_16x", Icon16x16) );
-		Set( "PaletteToolBar.SubMenuIndicator", 		new IMAGE_BRUSH( "Common/SubmenuArrow", Icon8x8 ) );
-
-		// Set( "PaletteToolBar.Content.Padding", FMargin(2.0) );
-		Set( "PaletteToolBar.Label.Padding", FMargin(0.0f, 4.0f, 0.0f, 0.0f) );
-
-		// "SToolBarButtonBlock.Padding" appears as space between the buttons (Outside the Orange selection/hover region)
-		// Outside Margin for Buttons
-		Set( "PaletteToolBar.SToolBarButtonBlock.Padding", 				FMargin(2.f, 0.0f) );
-
-		// SToolBarButtonBlock.CheckBox.Padding is the space inside the toggle button around the icon
-		// Inside Margin for Buttons
-		Set( "PaletteToolBar.SToolBarButtonBlock.CheckBox.Padding", 	FMargin(4.0f, 2.0f, 4.0f, 2.0f) );
-
-		// Outside Margin for Combo Buttons.   
-		// Combo Buttons with Arrows already have an annoying 2.0f built in padding, so we subtract to compensate for that.
-		Set( "PaletteToolBar.SToolBarComboButtonBlock.Padding", 		FMargin(2.0f, 0.0f) );
-
-		// Outside Margin for Generic Widget Blocks.  
-		Set( "PaletteToolBar.Block.IndentedPadding", 					FMargin(4.0f, 2.0f, 4.0f, 2.0f) );
-		Set( "PaletteToolBar.Block.Padding", 							FMargin(4.0f, 2.0f, 4.0f, 2.0f) );
-
-		Set( "PaletteToolBar.SToolBarComboButtonBlock.ComboButton.Color", DefaultForeground );
-
-		Set( "PaletteToolBar.Separator", new FSlateColorBrush( FLinearColor(FColor(47, 47, 47)) ) );
-		Set( "PaletteToolBar.Separator.Padding", FMargin( 6.f, 0.f, 6.f, 0.f) );
 
-		Set( "PaletteToolBar.Label", FTextBlockStyle(NormalText).SetFont( DEFAULT_FONT( "Regular", 9 ) ) );
-
-		Set( "PaletteToolBar.EditableText",  			GetWidgetStyle<FEditableTextBoxStyle>("ToolBar.EditableText"));
-		Set( "PaletteToolBar.Keybinding", 				GetWidgetStyle<FTextBlockStyle>("ToolBar.Keybinding"));
-		Set( "PaletteToolBar.Heading", 					GetWidgetStyle<FTextBlockStyle>("ToolBar.Heading"));
-		Set( "PaletteToolBar.CheckBox", 				GetWidgetStyle<FCheckBoxStyle>("ToolBar.CheckBox"));
-		Set( "PaletteToolbar.Check", 					GetWidgetStyle<FCheckBoxStyle>("ToolBar.Check"));
-		Set( "PaletteToolBar.RadioButton", 				GetWidgetStyle<FCheckBoxStyle>("ToolBar.RadioButton"));
+		FToolBarStyle PaletteToolBarStyle = FStarshipCoreStyle::GetCoreStyle().GetWidgetStyle<FToolBarStyle>("SlimToolBar");
+		
+		PaletteToolBarStyle.SetBackground(FSlateColorBrush(FStyleColors::Background));
+		PaletteToolBarStyle.SetLabelPadding(FMargin(0.0f, 4.0f, 0.0f, 0.0f));
+		PaletteToolBarStyle.SetCheckBoxPadding(FMargin(4.0f, 2.0f, 4.0f, 2.0f));
+		PaletteToolBarStyle.SetComboButtonPadding(FMargin(2.0f, 0.0f));
+		PaletteToolBarStyle.SetIndentedBlockPadding(FMargin(4.0f, 2.0f, 4.0f, 2.0f) );
+		PaletteToolBarStyle.SetBlockPadding(FMargin(4.0f, 2.0f, 4.0f, 2.0f));
 
 		Set("PaletteToolBar.MaxUniformToolbarSize", 50.f);
 		Set("PaletteToolbar.MinUniformToolbarSize", 22.f);
@@ -1782,26 +1739,18 @@ void FStarshipEditorStyle::FStyle::SetupGeneralStyles()
 				.SetCheckedImage( 		   	BOX_BRUSH( "Common/RoundedSelection_16x", 4.0f/16.0f, SelectionColor_Pressed ) )
 				.SetCheckedHoveredImage( 	BOX_BRUSH( "Common/RoundedSelection_16x", 4.0f/16.0f, SelectionColor_Pressed ) )
 				.SetCheckedPressedImage( 	BOX_BRUSH( "Common/RoundedSelection_16x", 4.0f/16.0f, SelectionColor_Pressed ) );
+		
+		PaletteToolBarStyle.SetToggleButtonStyle(ToolBarToggleButtonCheckBoxStyle);
+		PaletteToolBarStyle.SetButtonStyle(FButtonStyle(Button)
+			.SetNormal(FSlateNoResource())
+			.SetNormalPadding(FMargin(2.0f, 0.0f, 2.0f, 0.0f))
+			.SetPressed(BOX_BRUSH("Common/RoundedSelection_16x", 4.0f / 16.0f, SelectionColor_Pressed))
+			.SetHovered(BOX_BRUSH("Common/RoundedSelection_16x", 4.0f / 16.0f, ButtonHoverColor)));
 
-		/* ... and add new style */
-		Set( "PaletteToolBar.ToggleButton", ToolBarToggleButtonCheckBoxStyle );
-
-		Set( "PaletteToolBar.Button", FButtonStyle(Button)
-			.SetNormal ( FSlateNoResource() )
-			.SetNormalPadding( FMargin(2.0f, 0.0f, 2.0f, 0.0f) )
-			.SetPressed( BOX_BRUSH( "Common/RoundedSelection_16x", 4.0f/16.0f, SelectionColor_Pressed ) )
-			.SetHovered( BOX_BRUSH( "Common/RoundedSelection_16x", 4.0f/16.0f, ButtonHoverColor) )
-		);
-
-		Set( "PaletteToolBar.Button.Normal", 			new FSlateNoResource() );
-		Set( "PaletteToolBar.Button.Pressed", 			new BOX_BRUSH( "Common/RoundedSelection_16x", 4.0f/16.0f, SelectionColor_Pressed) );
-		Set( "PaletteToolBar.Button.Hovered", 			new BOX_BRUSH( "Common/RoundedSelection_16x", 4.0f/16.0f, ButtonHoverColor) );
-		Set( "PaletteToolBar.Button.Checked", 			new BOX_BRUSH( "Common/RoundedSelection_16x",  4.0f/16.0f, SelectionColor_Pressed) );
-		Set( "PaletteToolBar.Button.Checked_Hovered", 	new BOX_BRUSH( "Common/RoundedSelection_16x",  4.0f/16.0f, SelectionColor_Pressed) );
-		Set( "PaletteToolBar.Button.Checked_Pressed", 	new BOX_BRUSH( "Common/RoundedSelection_16x",  4.0f/16.0f, SelectionColor_Pressed) );
+		Set("PaletteToolBar", PaletteToolBarStyle);
 
-		Set( "EditorModesPanel.CategoryFontStyle", DEFAULT_FONT( "Bold", 10 ) );
-		Set( "EditorModesPanel.ToolDescriptionFont", DEFAULT_FONT("Italic", 10) );
+		Set("EditorModesPanel.CategoryFontStyle", DEFAULT_FONT( "Bold", 10 ));
+		Set("EditorModesPanel.ToolDescriptionFont", DEFAULT_FONT("Italic", 10));
 
 	}
 	
@@ -1836,29 +1785,11 @@ void FStarshipEditorStyle::FStyle::SetupGeneralStyles()
 
 	// ViewportLayoutToolbar
 	{
-		const FLinearColor LayoutSelectionColor_Hovered = FLinearColor(0.5f, 0.5f, 0.5f);
-
-		Set( "ViewportLayoutToolbar.Background", new FSlateNoResource() );
-		Set( "ViewportLayoutToolbar.Label", FTextBlockStyle() );
-		Set( "ViewportLayoutToolbar.Button", FButtonStyle(NoBorder) );
-		Set( "ViewportLayoutToolbar.Expand", new IMAGE_BRUSH("Icons/toolbar_expand_16x", Icon16x16) );
+		FToolBarStyle ViewportLayoutToolbar = FStarshipCoreStyle::GetCoreStyle().GetWidgetStyle<FToolBarStyle>("SlimToolBar");
+		ViewportLayoutToolbar.SetIconSize(Icon40x40);
 
-		/* Create style for "ViewportLayoutToolbar.ToggleButton" ... */
-		const FCheckBoxStyle ViewportLayoutToolbarToggleButtonStyle = FCheckBoxStyle()
-			.SetCheckBoxType( ESlateCheckBoxType::ToggleButton )
-			.SetUncheckedImage( FSlateNoResource() )
-			.SetUncheckedPressedImage( BOX_BRUSH( "Common/RoundedSelection_16x", 4.0f/16.0f, LayoutSelectionColor_Hovered ) )
-			.SetUncheckedHoveredImage( BOX_BRUSH( "Common/RoundedSelection_16x", 4.0f/16.0f, LayoutSelectionColor_Hovered ) )
-			.SetCheckedHoveredImage( BOX_BRUSH( "Common/RoundedSelection_16x",  4.0f/16.0f, SelectionColor ) )
-			.SetCheckedPressedImage( BOX_BRUSH( "Common/RoundedSelection_16x",  4.0f/16.0f, SelectionColor ) )
-			.SetCheckedImage( BOX_BRUSH( "Common/RoundedSelection_16x",  4.0f/16.0f, SelectionColor ) );
-		/* ... and add new style */
-		Set( "ViewportLayoutToolbar.ToggleButton", ViewportLayoutToolbarToggleButtonStyle );
-
-		Set( "ViewportLayoutToolbar.SToolBarButtonBlock.Padding", FMargin(4.0f) );
-		Set( "ViewportLayoutToolbar.SToolBarButtonBlock.CheckBox.Padding", FMargin(0.0f) );
-		Set( "ViewportLayoutToolbar.SToolBarComboButtonBlock.ComboButton.Color", DefaultForeground );
-		}
+		Set("ViewportLayoutToolbar", ViewportLayoutToolbar);
+	}
 
 	
 	// Highres Screenshot
@@ -2321,34 +2252,16 @@ void FStarshipEditorStyle::FStyle::SetupSequencerStyles()
 		Set("EMovieSceneBlendType::Relative", new IMAGE_BRUSH("Sequencer/EMovieSceneBlendType_Relative", FVector2D(32, 16)));
 		Set("EMovieSceneBlendType::Additive", new IMAGE_BRUSH("Sequencer/EMovieSceneBlendType_Additive", FVector2D(32, 16)));
 
+		FToolBarStyle SequencerToolbarStyle = FToolBarStyle(FStarshipCoreStyle::GetCoreStyle().GetWidgetStyle<FToolBarStyle>("ToolBar"));
+
 		// Sequencer & Curve Editor Toolbar Style ( Grabs core default, copies brushes, changes separator and block spacing )
-		Set( "Sequencer.ToolBar.Background", 		new FSlateColorBrush(FStyleColors::Background));
-		Set( "Sequencer.ToolBar.Icon", 				new CORE_IMAGE_BRUSH( "Icons/icon_tab_toolbar_16px", Icon16x16 ) );
-		Set( "Sequencer.ToolBar.Expand", 			new CORE_IMAGE_BRUSH( "Icons/toolbar_expand_16x", Icon16x16) );
-		Set( "Sequencer.ToolBar.SubMenuIndicator", 	new CORE_IMAGE_BRUSH( "Common/SubmenuArrow", Icon8x8 ) );
-		Set( "Sequencer.ToolBar.SToolBarComboButtonBlock.Padding", 		 FStarshipCoreStyle::GetCoreStyle().GetMargin("ToolBar.SToolBarComboButtonBlock.Padding"));
-		Set( "Sequencer.ToolBar.SToolBarButtonBlock.Padding", 			 FStarshipCoreStyle::GetCoreStyle().GetMargin("ToolBar.SToolBarButtonBlock.Padding"));
-		Set( "Sequencer.ToolBar.SToolBarCheckComboButtonBlock.Padding",  FStarshipCoreStyle::GetCoreStyle().GetMargin("ToolBar.SToolBarCheckComboButtonBlock.Padding"));
-		Set( "Sequencer.ToolBar.SToolBarButtonBlock.CheckBox.Padding", 	 FStarshipCoreStyle::GetCoreStyle().GetMargin("ToolBar.SToolBarButtonBlock.CheckBox.Padding"));
-		Set("Sequencer.ToolBar.SimpleComboButton",						 FStarshipCoreStyle::GetCoreStyle().GetVector("ToolBar.Label.DropShadowSize"));
-
-		// Used only for WidgetBlocks
-		Set( "Sequencer.ToolBar.Block.IndentedPadding", 	FStarshipCoreStyle::GetCoreStyle().GetMargin("ToolBar.Block.IndentedPadding") );
-		Set( "Sequencer.ToolBar.Block.Padding", 			FMargin(0.0f, 4.0f));
-
-		Set( "Sequencer.ToolBar.Separator", 				new FSlateColorBrush( FLinearColor(FColor(48, 48, 48)) ) );
-		Set( "Sequencer.ToolBar.Separator.Padding", 		FMargin( 1.f, 1.f ) );
-
-		Set( "Sequencer.ToolBar.Label", 					FStarshipCoreStyle::GetCoreStyle().GetWidgetStyle<FTextBlockStyle>("ToolBar.Label"));
-		Set( "Sequencer.ToolBar.EditableText",  			FStarshipCoreStyle::GetCoreStyle().GetWidgetStyle<FEditableTextBoxStyle>("ToolBar.EditableText"));
-		Set( "Sequencer.ToolBar.Keybinding", 				FStarshipCoreStyle::GetCoreStyle().GetWidgetStyle<FTextBlockStyle>("ToolBar.Keybinding"));
-		Set( "Sequencer.ToolBar.Heading", 					FStarshipCoreStyle::GetCoreStyle().GetWidgetStyle<FTextBlockStyle>("ToolBar.Heading"));
-		Set( "Sequencer.ToolBar.CheckBox", 					FStarshipCoreStyle::GetCoreStyle().GetWidgetStyle<FCheckBoxStyle>("ToolBar.CheckBox"));
-		Set( "Sequencer.Toolbar.Check", 					FStarshipCoreStyle::GetCoreStyle().GetWidgetStyle<FCheckBoxStyle>("ToolBar.Check"));
-		Set( "Sequencer.ToolBar.RadioButton", 				FStarshipCoreStyle::GetCoreStyle().GetWidgetStyle<FCheckBoxStyle>("ToolBar.RadioButton"));
-		Set( "Sequencer.ToolBar.ToggleButton", 				FStarshipCoreStyle::GetCoreStyle().GetWidgetStyle<FCheckBoxStyle>("ToolBar.ToggleButton"));
-		Set( "Sequencer.ToolBar.Button", 					FStarshipCoreStyle::GetCoreStyle().GetWidgetStyle<FButtonStyle>("ToolBar.Button"));
-		Set("Sequencer.ToolBar.SimpleComboButton",			FStarshipCoreStyle::GetCoreStyle().GetWidgetStyle<FButtonStyle>("ToolBar.SimpleComboButton"));
+		SequencerToolbarStyle.SetBackground(FSlateColorBrush(FStyleColors::Background));
+		SequencerToolbarStyle.SetBlockPadding(FMargin(0.0f, 4.0f));
+
+		SequencerToolbarStyle.SetSeparatorBrush(FSlateColorBrush(FLinearColor(FColor(48, 48, 48))));
+		SequencerToolbarStyle.SetSeparatorPadding(FMargin( 1.f, 1.f ));
+
+		Set("Sequencer.ToolBar", SequencerToolbarStyle);
 	}
 
 
@@ -2376,60 +2289,29 @@ void FStarshipEditorStyle::FStyle::SetupSequencerStyles()
 
 void FStarshipEditorStyle::FStyle::SetupViewportStyles()
 {
-// Viewport ToolbarBar
-	{
-		Set("ViewportMenu.Background", new BOX_BRUSH("Old/Menu_Background", FMargin(8.0f / 64.0f), FLinearColor::Transparent));
-		Set("ViewportMenu.Icon", new IMAGE_BRUSH("Icons/icon_tab_toolbar_16px", Icon16x16));
-		Set("ViewportMenu.Expand", new IMAGE_BRUSH("Icons/toolbar_expand_16x", Icon8x8));
-		Set("ViewportMenu.SubMenuIndicator", new IMAGE_BRUSH("Common/SubmenuArrow", Icon8x8));
-		Set("ViewportMenu.SToolBarComboButtonBlock.Padding", FMargin(0));
-		Set("ViewportMenu.SToolBarButtonBlock.Padding", FMargin(0));
-		Set("ViewportMenu.SToolBarButtonBlock.Button.Padding", FMargin(0));
-		Set("ViewportMenu.SToolBarCheckComboButtonBlock.Padding", FMargin(0));
-		Set("ViewportMenu.SToolBarButtonBlock.CheckBox.Padding", FMargin(4.0f));
-		Set("ViewportMenu.SToolBarComboButtonBlock.ComboButton.Color", FLinearColor(0.f, 0.f, 0.f, 0.75f));
-
-		Set("ViewportMenu.Separator", new BOX_BRUSH("Old/Button", 8.0f / 32.0f, FLinearColor::Transparent));
-		Set("ViewportMenu.Separator.Padding", FMargin( 1.f, 0.f, 0.f, 0.f) );
-
-		Set("ViewportMenu.IconSize", Icon16x16);
-
-		Set("ViewportMenu.Label", FTextBlockStyle(NormalText)
-			.SetFont(DEFAULT_FONT("Bold", 9))
-			.SetColorAndOpacity(FLinearColor(0.0f, 0.0f, 0.0f, 1.0f)));
-
-		Set("ViewportMenu.Label.Padding", FMargin(0.0f, 0.0f, 3.0f, 0.0f));
-		Set("ViewportMenu.Label.ContentPadding", FMargin(5.0f, 2.0f));
-		Set("ViewportMenu.Label.DropShadowSize", FVector2D::ZeroVector);
-
-		Set("ViewportMenu.EditableText", FEditableTextBoxStyle(NormalEditableTextBoxStyle).SetFont(DEFAULT_FONT("Regular", 9)));
-		Set("ViewportMenu.Keybinding", FTextBlockStyle(NormalText).SetFont(DEFAULT_FONT("Regular", 8)));
-
-		Set("ViewportMenu.Block.IndentedPadding", FMargin(0));
-		Set("ViewportMenu.Block.Padding", FMargin(0));
-
-		Set("ViewportMenu.Heading.Font", DEFAULT_FONT("Regular", 8));
-		Set("ViewportMenu.Heading.ColorAndOpacity", FLinearColor(0.4f, 0.4, 0.4f, 1.0f));
-
-		const FCheckBoxStyle ViewportMenuCheckBoxCheckBoxStyle = FCheckBoxStyle()
-			.SetUncheckedImage(IMAGE_BRUSH("Common/SmallCheckBox", Icon14x14))
-			.SetUncheckedPressedImage(IMAGE_BRUSH("Common/SmallCheckBox_Hovered", Icon14x14, FLinearColor(0.5f, 0.5f, 0.5f)))
-			.SetUncheckedHoveredImage(IMAGE_BRUSH("Common/SmallCheckBox_Hovered", Icon14x14))
-			.SetCheckedHoveredImage(IMAGE_BRUSH("Common/SmallCheckBox_Checked_Hovered", Icon14x14))
-			.SetCheckedPressedImage(IMAGE_BRUSH("Common/SmallCheckBox_Checked_Hovered", Icon14x14, FLinearColor(0.5f, 0.5f, 0.5f)))
-			.SetCheckedImage(IMAGE_BRUSH("Common/SmallCheckBox_Checked", Icon14x14));
-		Set("ViewportMenu.CheckBox", ViewportMenuCheckBoxCheckBoxStyle);
-
-		// Read-only checkbox that appears next to a menu item
-		const FCheckBoxStyle ViewportMenuCheckCheckBoxStyle = FCheckBoxStyle()
-			.SetUncheckedImage(IMAGE_BRUSH("Icons/Empty_14x", Icon14x14))
-			.SetUncheckedPressedImage(FSlateNoResource())
-			.SetUncheckedHoveredImage(FSlateNoResource())
-			.SetCheckedHoveredImage(IMAGE_BRUSH("Common/SmallCheck", Icon14x14))
-			.SetCheckedPressedImage(IMAGE_BRUSH("Common/SmallCheck", Icon14x14))
-			.SetCheckedImage(IMAGE_BRUSH("Common/SmallCheck", Icon14x14));
-		Set("ViewportMenu.Check", ViewportMenuCheckCheckBoxStyle);
-
+	// Viewport ToolbarBar
+	{
+		FToolBarStyle ViewportToolbarStyle =
+			FToolBarStyle()
+			.SetBackground(BOX_BRUSH("Old/Menu_Background", FMargin(8.0f / 64.0f), FLinearColor::Transparent))
+			.SetExpandBrush(IMAGE_BRUSH("Icons/toolbar_expand_16x", Icon8x8))
+			.SetSubMenuIndicator(IMAGE_BRUSH("Common/SubmenuArrow", Icon8x8))
+			.SetComboButtonPadding(FMargin(0))
+			.SetButtonPadding(FMargin(0))
+			.SetCheckBoxPadding(FMargin(0))
+			.SetSeparatorBrush(BOX_BRUSH("Old/Button", 8.0f / 32.0f, FLinearColor::Transparent))
+			.SetSeparatorPadding(FMargin(1.f, 0.f, 0.f, 0.f))
+			.SetIconSize(Icon16x16)
+			.SetLabelPadding(FMargin(0.0f, 0.0f, 3.0f, 0.0f))
+			.SetEditableTextStyle(FEditableTextBoxStyle(NormalEditableTextBoxStyle).SetFont(DEFAULT_FONT("Regular", 9)))
+			.SetIndentedBlockPadding(FMargin(0))
+			.SetBlockPadding(FMargin(0))
+			.SetLabelStyle(
+				FTextBlockStyle(NormalText)
+				.SetFont(DEFAULT_FONT("Bold", 9))
+				.SetColorAndOpacity(FLinearColor(0.0f, 0.0f, 0.0f, 1.0f))
+			);
+		
 		const FString SmallRoundedButton(TEXT("Common/SmallRoundedButton"));
 		const FString SmallRoundedButtonStart(TEXT("Common/SmallRoundedButtonLeft"));
 		const FString SmallRoundedButtonMiddle(TEXT("Common/SmallRoundedButtonCentre"));
@@ -2438,60 +2320,51 @@ void FStarshipEditorStyle::FStyle::SetupViewportStyles()
 		const FLinearColor NormalColor(1, 1, 1, 0.75f);
 		const FLinearColor PressedColor(1, 1, 1, 1.f);
 
-		const FCheckBoxStyle ViewportMenuRadioButtonCheckBoxStyle = FCheckBoxStyle()
-		.SetUncheckedImage(IMAGE_BRUSH("Common/MenuItemRadioButton_Off", Icon14x14))
-		.SetUncheckedPressedImage(IMAGE_BRUSH("Common/MenuItemRadioButton_Off", Icon14x14, FLinearColor(0.5f, 0.5f, 0.5f)))
-		.SetUncheckedHoveredImage(IMAGE_BRUSH("Common/MenuItemRadioButton_Off", Icon14x14))
-		.SetCheckedHoveredImage(IMAGE_BRUSH("Common/MenuItemRadioButton_On", Icon14x14))
-		.SetCheckedPressedImage(IMAGE_BRUSH("Common/MenuItemRadioButton_On_Pressed", Icon14x14))
-		.SetCheckedImage(IMAGE_BRUSH("Common/MenuItemRadioButton_On", Icon14x14));
-		Set("ViewportMenu.RadioButton", ViewportMenuRadioButtonCheckBoxStyle);
-
 		/* Create style for "ViewportMenu.ToggleButton" ... */
 		const FCheckBoxStyle ViewportMenuToggleButtonStyle = FCheckBoxStyle()
-		.SetCheckBoxType(ESlateCheckBoxType::ToggleButton)
-		.SetUncheckedImage(BOX_BRUSH(*SmallRoundedButton, FMargin(7.f / 16.f), NormalColor))
-		.SetUncheckedPressedImage(BOX_BRUSH(*SmallRoundedButton, FMargin(7.f / 16.f), PressedColor))
-		.SetUncheckedHoveredImage(BOX_BRUSH(*SmallRoundedButton, FMargin(7.f / 16.f), PressedColor))
-		.SetCheckedHoveredImage(BOX_BRUSH(*SmallRoundedButton, FMargin(7.f / 16.f), SelectionColor_Pressed))
-		.SetCheckedPressedImage(BOX_BRUSH(*SmallRoundedButton, FMargin(7.f / 16.f), SelectionColor_Pressed))
-		.SetCheckedImage(BOX_BRUSH(*SmallRoundedButton, FMargin(7.f / 16.f), SelectionColor_Pressed));
+			.SetCheckBoxType(ESlateCheckBoxType::ToggleButton)
+			.SetUncheckedImage(BOX_BRUSH(*SmallRoundedButton, FMargin(7.f / 16.f), NormalColor))
+			.SetUncheckedPressedImage(BOX_BRUSH(*SmallRoundedButton, FMargin(7.f / 16.f), PressedColor))
+			.SetUncheckedHoveredImage(BOX_BRUSH(*SmallRoundedButton, FMargin(7.f / 16.f), PressedColor))
+			.SetCheckedHoveredImage(BOX_BRUSH(*SmallRoundedButton, FMargin(7.f / 16.f), SelectionColor_Pressed))
+			.SetCheckedPressedImage(BOX_BRUSH(*SmallRoundedButton, FMargin(7.f / 16.f), SelectionColor_Pressed))
+			.SetCheckedImage(BOX_BRUSH(*SmallRoundedButton, FMargin(7.f / 16.f), SelectionColor_Pressed));
 		/* ... and add new style */
-		Set("ViewportMenu.ToggleButton", ViewportMenuToggleButtonStyle);
+		ViewportToolbarStyle.SetToggleButtonStyle(ViewportMenuToggleButtonStyle);
 
 		/* Create style for "ViewportMenu.ToggleButton.Start" ... */
 		const FCheckBoxStyle ViewportMenuToggleStartButtonStyle = FCheckBoxStyle()
-		.SetCheckBoxType(ESlateCheckBoxType::ToggleButton)
-		.SetUncheckedImage(BOX_BRUSH(*SmallRoundedButtonStart, FMargin(7.f / 16.f), NormalColor))
-		.SetUncheckedPressedImage(BOX_BRUSH(*SmallRoundedButtonStart, FMargin(7.f / 16.f), PressedColor))
-		.SetUncheckedHoveredImage(BOX_BRUSH(*SmallRoundedButtonStart, FMargin(7.f / 16.f), PressedColor))
-		.SetCheckedHoveredImage(BOX_BRUSH(*SmallRoundedButtonStart, FMargin(7.f / 16.f), SelectionColor_Pressed))
-		.SetCheckedPressedImage(BOX_BRUSH(*SmallRoundedButtonStart, FMargin(7.f / 16.f), SelectionColor_Pressed))
-		.SetCheckedImage(BOX_BRUSH(*SmallRoundedButtonStart, FMargin(7.f / 16.f), SelectionColor_Pressed));
+			.SetCheckBoxType(ESlateCheckBoxType::ToggleButton)
+			.SetUncheckedImage(BOX_BRUSH(*SmallRoundedButtonStart, FMargin(7.f / 16.f), NormalColor))
+			.SetUncheckedPressedImage(BOX_BRUSH(*SmallRoundedButtonStart, FMargin(7.f / 16.f), PressedColor))
+			.SetUncheckedHoveredImage(BOX_BRUSH(*SmallRoundedButtonStart, FMargin(7.f / 16.f), PressedColor))
+			.SetCheckedHoveredImage(BOX_BRUSH(*SmallRoundedButtonStart, FMargin(7.f / 16.f), SelectionColor_Pressed))
+			.SetCheckedPressedImage(BOX_BRUSH(*SmallRoundedButtonStart, FMargin(7.f / 16.f), SelectionColor_Pressed))
+			.SetCheckedImage(BOX_BRUSH(*SmallRoundedButtonStart, FMargin(7.f / 16.f), SelectionColor_Pressed));
 		/* ... and add new style */
 		Set("ViewportMenu.ToggleButton.Start", ViewportMenuToggleStartButtonStyle);
 
 		/* Create style for "ViewportMenu.ToggleButton.Middle" ... */
 		const FCheckBoxStyle ViewportMenuToggleMiddleButtonStyle = FCheckBoxStyle()
-		.SetCheckBoxType(ESlateCheckBoxType::ToggleButton)
-		.SetUncheckedImage(BOX_BRUSH(*SmallRoundedButtonMiddle, FMargin(7.f / 16.f), NormalColor))
-		.SetUncheckedPressedImage(BOX_BRUSH(*SmallRoundedButtonMiddle, FMargin(7.f / 16.f), PressedColor))
-		.SetUncheckedHoveredImage(BOX_BRUSH(*SmallRoundedButtonMiddle, FMargin(7.f / 16.f), PressedColor))
-		.SetCheckedHoveredImage(BOX_BRUSH(*SmallRoundedButtonMiddle, FMargin(7.f / 16.f), SelectionColor_Pressed))
-		.SetCheckedPressedImage(BOX_BRUSH(*SmallRoundedButtonMiddle, FMargin(7.f / 16.f), SelectionColor_Pressed))
-		.SetCheckedImage(BOX_BRUSH(*SmallRoundedButtonMiddle, FMargin(7.f / 16.f), SelectionColor_Pressed));
+			.SetCheckBoxType(ESlateCheckBoxType::ToggleButton)
+			.SetUncheckedImage(BOX_BRUSH(*SmallRoundedButtonMiddle, FMargin(7.f / 16.f), NormalColor))
+			.SetUncheckedPressedImage(BOX_BRUSH(*SmallRoundedButtonMiddle, FMargin(7.f / 16.f), PressedColor))
+			.SetUncheckedHoveredImage(BOX_BRUSH(*SmallRoundedButtonMiddle, FMargin(7.f / 16.f), PressedColor))
+			.SetCheckedHoveredImage(BOX_BRUSH(*SmallRoundedButtonMiddle, FMargin(7.f / 16.f), SelectionColor_Pressed))
+			.SetCheckedPressedImage(BOX_BRUSH(*SmallRoundedButtonMiddle, FMargin(7.f / 16.f), SelectionColor_Pressed))
+			.SetCheckedImage(BOX_BRUSH(*SmallRoundedButtonMiddle, FMargin(7.f / 16.f), SelectionColor_Pressed));
 		/* ... and add new style */
 		Set("ViewportMenu.ToggleButton.Middle", ViewportMenuToggleMiddleButtonStyle);
 
 		/* Create style for "ViewportMenu.ToggleButton.End" ... */
 		const FCheckBoxStyle ViewportMenuToggleEndButtonStyle = FCheckBoxStyle()
-		.SetCheckBoxType(ESlateCheckBoxType::ToggleButton)
-		.SetUncheckedImage(BOX_BRUSH(*SmallRoundedButtonEnd, FMargin(7.f / 16.f), NormalColor))
-		.SetUncheckedPressedImage(BOX_BRUSH(*SmallRoundedButtonEnd, FMargin(7.f / 16.f), PressedColor))
-		.SetUncheckedHoveredImage(BOX_BRUSH(*SmallRoundedButtonEnd, FMargin(7.f / 16.f), PressedColor))
-		.SetCheckedHoveredImage(BOX_BRUSH(*SmallRoundedButtonEnd, FMargin(7.f / 16.f), SelectionColor_Pressed))
-		.SetCheckedPressedImage(BOX_BRUSH(*SmallRoundedButtonEnd, FMargin(7.f / 16.f), SelectionColor_Pressed))
-		.SetCheckedImage(BOX_BRUSH(*SmallRoundedButtonEnd, FMargin(7.f / 16.f), SelectionColor_Pressed));
+			.SetCheckBoxType(ESlateCheckBoxType::ToggleButton)
+			.SetUncheckedImage(BOX_BRUSH(*SmallRoundedButtonEnd, FMargin(7.f / 16.f), NormalColor))
+			.SetUncheckedPressedImage(BOX_BRUSH(*SmallRoundedButtonEnd, FMargin(7.f / 16.f), PressedColor))
+			.SetUncheckedHoveredImage(BOX_BRUSH(*SmallRoundedButtonEnd, FMargin(7.f / 16.f), PressedColor))
+			.SetCheckedHoveredImage(BOX_BRUSH(*SmallRoundedButtonEnd, FMargin(7.f / 16.f), SelectionColor_Pressed))
+			.SetCheckedPressedImage(BOX_BRUSH(*SmallRoundedButtonEnd, FMargin(7.f / 16.f), SelectionColor_Pressed))
+			.SetCheckedImage(BOX_BRUSH(*SmallRoundedButtonEnd, FMargin(7.f / 16.f), SelectionColor_Pressed));
 		/* ... and add new style */
 		Set("ViewportMenu.ToggleButton.End", ViewportMenuToggleEndButtonStyle);
 
@@ -2505,7 +2378,7 @@ void FStarshipEditorStyle::FStyle::SetupViewportStyles()
 		.SetPressedPadding(PressedPadding)
 		.SetNormalPadding(NormalPadding);
 
-		Set("ViewportMenu.Button", ViewportMenuButton);
+		ViewportToolbarStyle.SetButtonStyle(ViewportMenuButton);
 
 		Set("ViewportMenu.Button.Start", FButtonStyle(ViewportMenuButton)
 			.SetNormal(BOX_BRUSH(*SmallRoundedButtonStart, 7.0f / 16.0f, NormalColor))
@@ -2524,102 +2397,18 @@ void FStarshipEditorStyle::FStyle::SetupViewportStyles()
 			.SetPressed(BOX_BRUSH(*SmallRoundedButtonEnd, 7.0f / 16.0f, PressedColor))
 			.SetHovered(BOX_BRUSH(*SmallRoundedButtonEnd, 7.0f / 16.0f, PressedColor))
 		);
+
+		Set("ViewportMenu", ViewportToolbarStyle);
 	}
 
 	// Viewport actor preview's pin/unpin and attach/detach buttons
 	{
-	Set("ViewportActorPreview.Pinned", new IMAGE_BRUSH("Common/PushPin_Down", Icon16x16));
-	Set("ViewportActorPreview.Unpinned", new IMAGE_BRUSH("Common/PushPin_Up", Icon16x16));
-	Set("VRViewportActorPreview.Pinned", new IMAGE_BRUSH("Common/PushPin_Down_VR", Icon64x64));
-	Set("VRViewportActorPreview.Unpinned", new IMAGE_BRUSH("Common/PushPin_Up_VR", Icon64x64));
-	Set("VRViewportActorPreview.Attached", new IMAGE_BRUSH("Common/ScreenAttach_VR", Icon64x64));
-	Set("VRViewportActorPreview.Detached", new IMAGE_BRUSH("Common/ScreenDetach_VR", Icon64x64));
-	}
-}
-
-void FStarshipEditorStyle::FStyle::SetupNotificationBarStyles()
-{
-// NotificationBar
-	{
-		Set("NotificationBar.Background", new FSlateNoResource());
-		Set("NotificationBar.Icon", new FSlateNoResource());
-		Set("NotificationBar.Expand", new IMAGE_BRUSH("Icons/toolbar_expand_16x", Icon16x16));
-		Set("NotificationBar.SubMenuIndicator", new IMAGE_BRUSH("Common/SubmenuArrow", Icon8x8));
-
-		Set("NotificationBar.Block.IndentedPadding", FMargin(0));
-		Set("NotificationBar.Block.Padding", FMargin(0));
-
-		Set("NotificationBar.Separator", new BOX_BRUSH("Old/Button", 4.0f / 32.0f));
-		Set("NotificationBar.Separator.Padding", FMargin(0.5f));
-
-		Set("NotificationBar.Label", FTextBlockStyle(NormalText).SetFont(DEFAULT_FONT("Regular", 9)));
-		Set("NotificationBar.EditableText", FEditableTextBoxStyle(NormalEditableTextBoxStyle).SetFont(DEFAULT_FONT("Regular", 9)));
-		Set("NotificationBar.Keybinding", FTextBlockStyle(NormalText).SetFont(DEFAULT_FONT("Regular", 8)));
-
-		Set("NotificationBar.Heading", FTextBlockStyle(NormalText)
-			.SetFont(DEFAULT_FONT("Regular", 8))
-			.SetColorAndOpacity(FLinearColor(0.4f, 0.4, 0.4f, 1.0f)));
-
-		const FCheckBoxStyle NotificationBarCheckBoxCheckBoxStyle = FCheckBoxStyle()
-			.SetUncheckedImage(IMAGE_BRUSH("Common/SmallCheckBox", Icon14x14))
-			.SetUncheckedPressedImage(IMAGE_BRUSH("Common/SmallCheckBox_Hovered", Icon14x14, FLinearColor(0.5f, 0.5f, 0.5f)))
-			.SetUncheckedHoveredImage(IMAGE_BRUSH("Common/SmallCheckBox_Hovered", Icon14x14))
-			.SetCheckedHoveredImage(IMAGE_BRUSH("Common/SmallCheckBox_Checked_Hovered", Icon14x14))
-			.SetCheckedPressedImage(IMAGE_BRUSH("Common/SmallCheckBox_Checked_Hovered", Icon14x14, FLinearColor(0.5f, 0.5f, 0.5f)))
-			.SetCheckedImage(IMAGE_BRUSH("Common/SmallCheckBox_Checked", Icon14x14))
-			.SetUndeterminedImage(IMAGE_BRUSH("Common/CheckBox_Undetermined", Icon14x14))
-			.SetUndeterminedHoveredImage(IMAGE_BRUSH("Common/CheckBox_Undetermined_Hovered", Icon14x14))
-			.SetUndeterminedPressedImage(IMAGE_BRUSH("Common/CheckBox_Undetermined_Hovered", Icon14x14, FLinearColor(0.5f, 0.5f, 0.5f)));
-		Set("NotificationBar.CheckBox", NotificationBarCheckBoxCheckBoxStyle);
-
-		// Read-only checkbox that appears next to a menu item
-		const FCheckBoxStyle NotificationBarCheckCheckBoxStyle = FCheckBoxStyle()
-			.SetUncheckedImage(IMAGE_BRUSH("Icons/Empty_14x", Icon14x14))
-			.SetUncheckedPressedImage(FSlateNoResource())
-			.SetUncheckedHoveredImage(FSlateNoResource())
-			.SetCheckedHoveredImage(IMAGE_BRUSH("Common/SmallCheck", Icon14x14))
-			.SetCheckedPressedImage(IMAGE_BRUSH("Common/SmallCheck", Icon14x14))
-			.SetCheckedImage(IMAGE_BRUSH("Common/SmallCheck", Icon14x14))
-			.SetUndeterminedImage(IMAGE_BRUSH("Icons/Empty_14x", Icon14x14))
-			.SetUndeterminedPressedImage(FSlateNoResource())
-			.SetUndeterminedHoveredImage(FSlateNoResource());
-		Set("NotificationBar.Check", NotificationBarCheckCheckBoxStyle);
-
-		// This radio button is actually just a check box with different images
-		const FCheckBoxStyle NotificationBarRadioButtonCheckBoxStyle = FCheckBoxStyle()
-			.SetUncheckedImage(IMAGE_BRUSH("Common/RadioButton_Unselected_16x", Icon16x16))
-			.SetUncheckedPressedImage(IMAGE_BRUSH("Common/RadioButton_Unselected_16x", Icon16x16, SelectionColor_Pressed))
-			.SetUncheckedHoveredImage(IMAGE_BRUSH("Common/RadioButton_Unselected_16x", Icon16x16, SelectionColor))
-			.SetCheckedHoveredImage(IMAGE_BRUSH("Common/RadioButton_Selected_16x", Icon16x16, SelectionColor))
-			.SetCheckedPressedImage(IMAGE_BRUSH("Common/RadioButton_Selected_16x", Icon16x16, SelectionColor_Pressed))
-			.SetCheckedImage(IMAGE_BRUSH("Common/RadioButton_Selected_16x", Icon16x16));
-		Set("NotificationBar.RadioButton", NotificationBarRadioButtonCheckBoxStyle);
-
-		const FCheckBoxStyle NotificationBarToggleButtonCheckBoxStyle = FCheckBoxStyle()
-			.SetCheckBoxType(ESlateCheckBoxType::ToggleButton)
-			.SetUncheckedImage(FSlateNoResource())
-			.SetUncheckedPressedImage(BOX_BRUSH("Common/RoundedSelection_16x", 4.0f / 16.0f, SelectionColor_Pressed))
-			.SetUncheckedHoveredImage(BOX_BRUSH("Common/RoundedSelection_16x", 4.0f / 16.0f, SelectionColor))
-			.SetCheckedHoveredImage(BOX_BRUSH("Common/RoundedSelection_16x", 4.0f / 16.0f, SelectionColor_Pressed))
-			.SetCheckedPressedImage(BOX_BRUSH("Common/RoundedSelection_16x", 4.0f / 16.0f, SelectionColor))
-			.SetCheckedImage(BOX_BRUSH("Common/RoundedSelection_16x", 4.0f / 16.0f, SelectionColor_Pressed));
-		Set("NotificationBar.ToggleButton", NotificationBarToggleButtonCheckBoxStyle);
-
-		Set("NotificationBar.Button", FButtonStyle(NoBorder)
-			.SetNormal(FSlateNoResource())
-			.SetPressed(BOX_BRUSH("Common/RoundedSelection_16x", 4.0f / 16.0f, SelectionColor_Pressed))
-			.SetHovered(BOX_BRUSH("Common/RoundedSelection_16x", 4.0f / 16.0f, SelectionColor))
-			.SetNormalPadding(FMargin(0, 1))
-			.SetPressedPadding(FMargin(0, 2, 0, 0))
-		);
-
-		Set("NotificationBar.Button.Checked", new BOX_BRUSH("Common/RoundedSelection_16x", 4.0f / 16.0f, SelectionColor_Pressed));
-		Set("NotificationBar.Button.Checked_Hovered", new BOX_BRUSH("Common/RoundedSelection_16x", 4.0f / 16.0f, SelectionColor_Pressed));
-		Set("NotificationBar.Button.Checked_Pressed", new BOX_BRUSH("Common/RoundedSelection_16x", 4.0f / 16.0f, SelectionColor));
-
-		Set("NotificationBar.SToolBarButtonBlock.CheckBox.Padding", FMargin(4.0f));
-		Set("NotificationBar.SToolBarButtonBlock.Button.Padding", FMargin(0.0f));
-		Set("NotificationBar.SToolBarComboButtonBlock.ComboButton.Color", DefaultForeground);
+		Set("ViewportActorPreview.Pinned", new IMAGE_BRUSH("Common/PushPin_Down", Icon16x16));
+		Set("ViewportActorPreview.Unpinned", new IMAGE_BRUSH("Common/PushPin_Up", Icon16x16));
+		Set("VRViewportActorPreview.Pinned", new IMAGE_BRUSH("Common/PushPin_Down_VR", Icon64x64));
+		Set("VRViewportActorPreview.Unpinned", new IMAGE_BRUSH("Common/PushPin_Up_VR", Icon64x64));
+		Set("VRViewportActorPreview.Attached", new IMAGE_BRUSH("Common/ScreenAttach_VR", Icon64x64));
+		Set("VRViewportActorPreview.Detached", new IMAGE_BRUSH("Common/ScreenDetach_VR", Icon64x64));
 	}
 }
 
@@ -4798,94 +4587,6 @@ void FStarshipEditorStyle::FStyle::SetupLevelEditorStyle()
 		Set("MeshPaint.SwitchForeAndBackgroundColor.Small", new IMAGE_BRUSH("Icons/Paint/Paint_SwapColors_40x", Icon20x20));
 	}
 
-	// EditorModesToolbar
-	{
-		Set("EditorModesToolbar.Background", new FSlateNoResource());
-		Set("EditorModesToolbar.Icon", new IMAGE_BRUSH("Icons/icon_tab_toolbar_16px", Icon16x16));
-		Set("EditorModesToolbar.Expand", new IMAGE_BRUSH("Icons/toolbar_expand_16x", Icon16x16));
-		Set("EditorModesToolbar.SubMenuIndicator", new IMAGE_BRUSH("Common/SubmenuArrow", Icon8x8));
-		Set("EditorModesToolbar.SToolBarComboButtonBlock.Padding", FMargin(0));
-		Set("EditorModesToolbar.SToolBarComboButtonBlock.ComboButton.Color", DefaultForeground);
-		Set("EditorModesToolbar.SToolBarButtonBlock.Padding", FMargin(1.0f, 0.0f, 0.0f, 0));
-		Set("EditorModesToolbar.SToolBarButtonBlock.CheckBox.Padding", FMargin(6.0f, 4.0f, 6.0f, 6.0f));
-		Set("EditorModesToolbar.SToolBarCheckComboButtonBlock.Padding", FMargin(0));
-
-		Set("EditorModesToolbar.Block.IndentedPadding", FMargin(0));
-		Set("EditorModesToolbar.Block.Padding", FMargin(0));
-
-		Set("EditorModesToolbar.Separator", new BOX_BRUSH("Old/Button", 4.0f / 32.0f));
-		Set("EditorModesToolbar.Separator.Padding", FMargin(8.f, 0.f, 8.f, 0.f));
-
-		Set("EditorModesToolbar.Label", FTextBlockStyle(NormalText).SetFont(DEFAULT_FONT("Regular", 7)));
-		Set("EditorModesToolbar.EditableText", FEditableTextBoxStyle(NormalEditableTextBoxStyle).SetFont(DEFAULT_FONT("Regular", 9)));
-		Set("EditorModesToolbar.Keybinding", FTextBlockStyle(NormalText).SetFont(DEFAULT_FONT("Regular", 8)));
-
-		Set("EditorModesToolbar.Heading.Font", DEFAULT_FONT("Regular", 8));
-		Set("EditorModesToolbar.Heading.ColorAndOpacity", FLinearColor(0.4f, 0.4, 0.4f, 1.0f));
-
-		/* Create style for "EditorModesToolbar.CheckBox" ... */
-		const FCheckBoxStyle EditorModesToolbarCheckBoxStyle = FCheckBoxStyle()
-			.SetUncheckedImage(IMAGE_BRUSH("Common/SmallCheckBox", Icon14x14))
-			.SetCheckedImage(IMAGE_BRUSH("Common/SmallCheckBox_Checked", Icon14x14))
-			.SetUncheckedHoveredImage(IMAGE_BRUSH("Common/SmallCheckBox_Hovered", Icon14x14))
-			.SetCheckedHoveredImage(IMAGE_BRUSH("Common/SmallCheckBox_Checked_Hovered", Icon14x14))
-			.SetCheckedPressedImage(IMAGE_BRUSH("Common/SmallCheckBox_Hovered", Icon14x14, FLinearColor(0.5f, 0.5f, 0.5f)));
-		/* ... and set new style */
-		Set("EditorModesToolbar.CheckBox", EditorModesToolbarCheckBoxStyle);
-
-		// Read-only checkbox that appears next to a menu item
-		/* Create style for "EditorModesToolbar.Check" ... */
-		const FCheckBoxStyle EditorModesToolbarCheckStyle = FCheckBoxStyle()
-			.SetUncheckedImage(IMAGE_BRUSH("Icons/Empty_14x", Icon14x14))
-			.SetCheckedImage(IMAGE_BRUSH("Common/SmallCheck", Icon14x14))
-			.SetUncheckedHoveredImage(FSlateNoResource())
-			.SetUncheckedPressedImage(FSlateNoResource())
-			.SetCheckedPressedImage(IMAGE_BRUSH("Common/SmallCheck", Icon14x14));
-		/*  ... and set new style  */
-		Set("EditorModesToolbar.Check", EditorModesToolbarCheckStyle);
-
-		// This radio button is actually just a check box with different images
-		/* Create style for "EditorModesToolbar.RadioButton" ... */
-		const FCheckBoxStyle EditorModesToolbarRadioButtonStyle = FCheckBoxStyle()
-			.SetUncheckedImage(IMAGE_BRUSH("Common/RadioButton_Unselected_16x", Icon16x16))
-			.SetCheckedImage(IMAGE_BRUSH("Common/RadioButton_Selected_16x", Icon16x16))
-			.SetUncheckedHoveredImage(IMAGE_BRUSH("Common/RadioButton_Unselected_16x", Icon16x16, SelectionColor))
-			.SetCheckedHoveredImage(IMAGE_BRUSH("Common/RadioButton_Selected_16x", Icon16x16, SelectionColor))
-			.SetUncheckedPressedImage(IMAGE_BRUSH("Common/RadioButton_Unselected_16x", Icon16x16, SelectionColor_Pressed))
-			.SetCheckedPressedImage(IMAGE_BRUSH("Common/RadioButton_Selected_16x", Icon16x16, SelectionColor_Pressed));
-		/* ... and add new style */
-		Set("EditorModesToolbar.RadioButton", EditorModesToolbarRadioButtonStyle);
-
-		/* Create style for "EditorModesToolbar.ToggleButton" ... */
-		const FCheckBoxStyle EditorModesToolbarToggleButtonStyle = FCheckBoxStyle()
-			.SetCheckBoxType(ESlateCheckBoxType::ToggleButton)
-			.SetUncheckedImage(BOX_BRUSH("/EditorModes/Tab_Inactive", 4 / 16.0f))
-			.SetUncheckedPressedImage(BOX_BRUSH("/EditorModes/Tab_Active", 4 / 16.0f))
-			.SetUncheckedHoveredImage(BOX_BRUSH("/EditorModes/Tab_Active", 4 / 16.0f))
-			.SetCheckedHoveredImage(BOX_BRUSH("/EditorModes/Tab_Active", 4 / 16.0f))
-			.SetCheckedPressedImage(BOX_BRUSH("/EditorModes/Tab_Active", 4 / 16.0f))
-			.SetCheckedImage(BOX_BRUSH("/EditorModes/Tab_Active", 4 / 16.0f));
-
-		/* ... and add new style */
-		Set("EditorModesToolbar.ToggleButton", EditorModesToolbarToggleButtonStyle);
-
-		Set("EditorModesToolbar.Button", FButtonStyle(Button)
-			.SetNormal(FSlateNoResource())
-			.SetPressed(BOX_BRUSH("Common/RoundedSelection_16x", 4.0f / 16.0f, SelectionColor_Pressed))
-			.SetHovered(BOX_BRUSH("Common/RoundedSelection_16x", 4.0f / 16.0f, SelectionColor))
-		);
-
-		Set("EditorModesToolbar.Button.Normal", new FSlateNoResource());
-		Set("EditorModesToolbar.Button.Pressed", new BOX_BRUSH("Common/RoundedSelection_16x", 4.0f / 16.0f, SelectionColor_Pressed));
-		Set("EditorModesToolbar.Button.Hovered", new BOX_BRUSH("Common/RoundedSelection_16x", 4.0f / 16.0f, SelectionColor));
-
-		Set("EditorModesToolbar.Button.Checked", new BOX_BRUSH("Common/RoundedSelection_16x", 4.0f / 16.0f, SelectionColor_Pressed));
-		Set("EditorModesToolbar.Button.Checked_Hovered", new BOX_BRUSH("Common/RoundedSelection_16x", 4.0f / 16.0f, SelectionColor_Pressed));
-		Set("EditorModesToolbar.Button.Checked_Pressed", new BOX_BRUSH("Common/RoundedSelection_16x", 4.0f / 16.0f, SelectionColor));
-	}
-
-
-
 	// Scalability (Performance Warning)
 	{
 		Set( "Scalability.ScalabilitySettings", new IMAGE_BRUSH("Scalability/ScalabilitySettings", FVector2D(473.0f, 266.0f) ) );
@@ -5372,29 +5073,24 @@ void FStarshipEditorStyle::FStyle::SetupPersonaStyle()
 		FLinearColor GreenHover = FLinearColor(GreenHSV.R, GreenHSV.G * .5, GreenHSV.B, GreenHSV.A).HSVToLinearRGB();
 		FLinearColor GreenPress = FLinearColor(GreenHSV.R, GreenHSV.G, GreenHSV.B*.5, GreenHSV.A).HSVToLinearRGB(); 
 
-		const FButtonStyle PlayToolbarButton = FButtonStyle(FStarshipCoreStyle::GetCoreStyle().GetWidgetStyle<FButtonStyle>("SlimToolBar.Button"))
+		FToolBarStyle PlayToolbar = FStarshipCoreStyle::GetCoreStyle().GetWidgetStyle<FToolBarStyle>("SlimToolBar");
+
+		const FButtonStyle PlayToolbarButton = FButtonStyle(PlayToolbar.ButtonStyle)
 			.SetPressed(FSlateNoResource())
 			.SetHovered(FSlateNoResource())
 			.SetNormalForeground(FStyleColors::AccentGreen)
 			.SetPressedForeground(GreenPress)
 			.SetHoveredForeground(GreenHover);
 
-		Set("PlayToolbar.Button", PlayToolbarButton);
-
-
 		FComboButtonStyle PlayToolbarComboButton = FComboButtonStyle(FStarshipCoreStyle::GetCoreStyle().GetWidgetStyle<FComboButtonStyle>("ComboButton"))
 			.SetContentPadding(0)
 			.SetDownArrowImage(CORE_IMAGE_BRUSH_SVG("Starship/Common/ellipsis-vertical-narrow", FVector2D(6, 24)));
 		PlayToolbarComboButton.ButtonStyle = PlayToolbarButton;
 
-		Set("PlayToolbar.SimpleComboButton", PlayToolbarComboButton);
-
-
-		Set("PlayToolbar.Label", FTextBlockStyle(NormalText).SetFont(DEFAULT_FONT("Regular", 10)));
-		Set("PlayToolbar.Label.Padding", FMargin(5, 0, 0, 0));
-		Set("PlayToolbar.Label.DropShadowSize", FVector2D::ZeroVector);
-		Set("PlayToolbar.SToolBarButtonBlock.Padding", FMargin(5.0f, 0.f));
-		Set("PlayToolbar.IconSize", Icon20x20);
+	
+		PlayToolbar.SetButtonStyle(PlayToolbarButton);
+		PlayToolbar.SetSettingsComboButtonStyle(PlayToolbarComboButton);
+		Set("PlayToolBar", PlayToolbar);
 
 		Set( "PlayWorld.Simulate", new IMAGE_BRUSH( "Icons/icon_simulate_40x", Icon40x40 ) );
 		Set( "PlayWorld.Simulate.Small", new IMAGE_BRUSH( "Icons/icon_simulate_40x", Icon20x20 ) );
@@ -7916,46 +7612,12 @@ void FStarshipEditorStyle::FStyle::SetupStatusBarStyle()
 	Set("StatusBar.Background", new FSlateColorBrush(FStyleColors::Background));
 	Set("StatusBar.HelpIcon", new CORE_IMAGE_BRUSH_SVG("Starship/Common/help", Icon16x16, FStyleColors::Foreground));
 
-	const FButtonStyle SlimToolBarButton = FButtonStyle(FStarshipCoreStyle::GetCoreStyle().GetWidgetStyle<FButtonStyle>("SlimToolBar.Button"));
-
-	Set("StatusBarToolBar.Background", new FSlateNoResource());
-	Set("StatusBarToolBar.Expand", new IMAGE_BRUSH("Icons/toolbar_expand_16x", Icon16x16));
-	Set("StatusBarToolBar.SubMenuIndicator", new IMAGE_BRUSH("Common/SubmenuArrow", Icon8x8));
-	Set("StatusBarToolBar.SToolBarComboButtonBlock.Padding", FMargin(6.0f, 0.0f));
-	Set("StatusBarToolBar.SToolBarButtonBlock.Padding", FMargin(4.0f, 0.0f));
-	Set("StatusBarToolBar.SToolBarButtonBlock.CheckBox.Padding", FMargin(10.0f, 0.0f));
-
-	Set("StatusBarToolBar.Separator", new FSlateColorBrush(FStyleColors::Input));
-	Set("StatusBarToolBar.Separator.Padding", FMargin(8.f, 0));
-
-	Set("StatusBarToolBar.Label", FTextBlockStyle(NormalText).SetFont(DEFAULT_FONT("Regular", 10)));
-	Set("StatusBarToolBar.Label.Padding", FMargin(5, 5, 0, 5));
-	Set("StatusBarToolBar.Label.DropShadowSize", FVector2D::ZeroVector);
-
-	Set("StatusBarToolBar.EditableText", FEditableTextBoxStyle(NormalEditableTextBoxStyle).SetFont(DEFAULT_FONT("Regular", 9)));
-
-	//	Style->Set("SlimToolBar.CheckBox", ToolBarCheckBoxStyle);
-
-	const FCheckBoxStyle SlimToolBarToggleButtonCheckBoxStyle = FCheckBoxStyle()
-		.SetCheckBoxType(ESlateCheckBoxType::ToggleButton)
-		.SetUncheckedImage(FSlateNoResource())
-		.SetUncheckedPressedImage(FSlateNoResource())
-		.SetUncheckedHoveredImage(FSlateNoResource())
-		.SetCheckedImage(FSlateNoResource())
-		.SetCheckedHoveredImage(FSlateNoResource())
-		.SetCheckedPressedImage(FSlateNoResource())
-		.SetForegroundColor(FStyleColors::Foreground)
-		.SetPressedForegroundColor(FStyleColors::ForegroundHover)
-		.SetHoveredForegroundColor(FStyleColors::ForegroundHover)
-		.SetCheckedForegroundColor(FStyleColors::Primary)
-		.SetCheckedPressedForegroundColor(FStyleColors::PrimaryPress)
-		.SetCheckedHoveredForegroundColor(FStyleColors::PrimaryHover);
-
-	Set("StatusBarToolBar.ToggleButton", SlimToolBarToggleButtonCheckBoxStyle);
-	Set("StatusBarToolBar.Button", SlimToolBarButton);
-
-	Set("StatusBarToolBar.SimpleComboButton", FStarshipCoreStyle::GetCoreStyle().GetWidgetStyle<FComboButtonStyle>("StatusBarToolBar.SimpleComboButton"));
-	Set("StatusBarToolBar.IconSize", Icon16x16);
+	FToolBarStyle StatusBarToolBarStyle = FStarshipCoreStyle::GetCoreStyle().GetWidgetStyle<FToolBarStyle>("SlimToolBar");
+
+	StatusBarToolBarStyle.SetBackground(FSlateNoResource());
+	StatusBarToolBarStyle.SetLabelPadding(FMargin(5, 5, 0, 5));
+	StatusBarToolBarStyle.SetIconSize(Icon16x16);
+	Set("StatusBarToolBar", StatusBarToolBarStyle);
 
 	Set("StatusBar.Message.MessageText",
 		FTextBlockStyle(FAppStyle::Get().GetWidgetStyle<FTextBlockStyle>("NormalText"))
diff --git a/Engine/Source/Editor/EditorStyle/Private/StarshipStyle.h b/Engine/Source/Editor/EditorStyle/Private/StarshipStyle.h
--- a/Engine/Source/Editor/EditorStyle/Private/StarshipStyle.h
+++ b/Engine/Source/Editor/EditorStyle/Private/StarshipStyle.h
@@ -73,7 +73,6 @@ public:
 		void SetupWorldBrowserStyles();
 		void SetupSequencerStyles();
 		void SetupViewportStyles();
-		void SetupNotificationBarStyles();
 		void SetupMenuBarStyles();
 		void SetupGeneralIcons();
 		void SetupWindowStyles();
diff --git a/Engine/Source/Editor/LevelEditor/Private/LevelEditorMenu.cpp b/Engine/Source/Editor/LevelEditor/Private/LevelEditorMenu.cpp
--- a/Engine/Source/Editor/LevelEditor/Private/LevelEditorMenu.cpp
+++ b/Engine/Source/Editor/LevelEditor/Private/LevelEditorMenu.cpp
@@ -274,21 +274,6 @@ TSharedRef< SWidget > FLevelEditorMenu::MakeLevelEditorMenu( const TSharedPtr<FU
 	return MenuBarWidget;
 }
 
-TSharedRef< SWidget > FLevelEditorMenu::MakeNotificationBar( const TSharedPtr<FUICommandList>& CommandList, TSharedPtr<class SLevelEditor> LevelEditor )
-{
-	FLevelEditorModule& LevelEditorModule = FModuleManager::GetModuleChecked<FLevelEditorModule>( "LevelEditor");
-	const TSharedPtr<FExtender> NotificationBarExtenders = LevelEditorModule.GetNotificationBarExtensibilityManager()->GetAllExtenders();
-
-	FToolBarBuilder NotificationBarBuilder( CommandList, FMultiBoxCustomization::None, NotificationBarExtenders );
-	NotificationBarBuilder.SetStyle(&FEditorStyle::Get(), "NotificationBar");
-	{
-		NotificationBarBuilder.BeginSection("Start");
-		NotificationBarBuilder.EndSection();
-	}
-
-	return NotificationBarBuilder.MakeWidget();
-}
-
 void FLevelEditorMenu::RegisterBuildMenu()
 {
 	static const FName BaseMenuName = "LevelEditor.MainMenu.Build";
diff --git a/Engine/Source/Editor/LevelEditor/Private/LevelEditorMenu.h b/Engine/Source/Editor/LevelEditor/Private/LevelEditorMenu.h
--- a/Engine/Source/Editor/LevelEditor/Private/LevelEditorMenu.h
+++ b/Engine/Source/Editor/LevelEditor/Private/LevelEditorMenu.h
@@ -24,6 +24,6 @@ public:
 	 * @return	New widget
 	 */
 	static TSharedRef< SWidget > MakeLevelEditorMenu( const TSharedPtr<FUICommandList>& CommandList, TSharedPtr<class SLevelEditor> LevelEditor );
-
-	static TSharedRef< SWidget > MakeNotificationBar( const TSharedPtr<FUICommandList>& CommandList, TSharedPtr<class SLevelEditor> LevelEditor );
+private:
+	static void RegisterBuildMenu();
 };
diff --git a/Engine/Source/Editor/LevelEditor/Private/LevelEditorToolBar.cpp b/Engine/Source/Editor/LevelEditor/Private/LevelEditorToolBar.cpp
--- a/Engine/Source/Editor/LevelEditor/Private/LevelEditorToolBar.cpp
+++ b/Engine/Source/Editor/LevelEditor/Private/LevelEditorToolBar.cpp
@@ -61,6 +61,7 @@
 #include "Subsystems/AssetEditorSubsystem.h"
 #include "LevelEditorModesActions.h"
 #include "ISourceControlModule.h"
+#include "Styling/ToolBarStyle.h"
 
 static TAutoConsoleVariable<int32> CVarAllowMatineeActors(
 	TEXT("Matinee.AllowMatineeActors"),
@@ -1479,7 +1480,7 @@ TSharedRef< SWidget > FLevelEditorToolBar::MakeLevelEditorToolBar( const TShared
 		+SOverlay::Slot()
 		[
 			SNew(SImage)
-			.Image(FCoreStyle::IsStarshipStyle() ? FCoreStyle::Get().GetBrush("SlimToolBar.Background") : FCoreStyle::Get().GetBrush("ToolPanel.GroupBorder"))
+			.Image(&FAppStyle::Get().GetWidgetStyle<FToolBarStyle>("SlimToolBar").BackgroundBrush)
 		]
 		+ SOverlay::Slot()
 		[
diff --git a/Engine/Source/Editor/LevelEditor/Private/SLevelEditor.cpp b/Engine/Source/Editor/LevelEditor/Private/SLevelEditor.cpp
--- a/Engine/Source/Editor/LevelEditor/Private/SLevelEditor.cpp
+++ b/Engine/Source/Editor/LevelEditor/Private/SLevelEditor.cpp
@@ -241,13 +241,6 @@ void SLevelEditor::ConstructTitleBarMessages()
 {
 	TtileBarMessageBox->ClearChildren();
 
-	TtileBarMessageBox->AddSlot()
-		.AutoWidth()
-		.Padding(5.0f, 0.0f, 0.0f, 0.0f)
-		[
-			FLevelEditorMenu::MakeNotificationBar( LevelEditorCommands, SharedThis(this ) )
-		];
-
 	const IMainFrameModule& MainFrameModule = FModuleManager::GetModuleChecked<IMainFrameModule>(MainFrameModuleName);
 	const FLevelEditorModule& LevelEditorModule = FModuleManager::GetModuleChecked< FLevelEditorModule >(LevelEditorModuleName);
 	
diff --git a/Engine/Source/Editor/LevelEditor/Private/SLevelEditorToolBox.cpp b/Engine/Source/Editor/LevelEditor/Private/SLevelEditorToolBox.cpp
--- a/Engine/Source/Editor/LevelEditor/Private/SLevelEditorToolBox.cpp
+++ b/Engine/Source/Editor/LevelEditor/Private/SLevelEditorToolBox.cpp
@@ -103,7 +103,6 @@ void SLevelEditorToolBox::UpdateModeLegacyToolBar()
 
 	FToolBarBuilder EditorModeTools( CommandList, FMultiBoxCustomization::None, ModeBarExtenders );
 	{
-		EditorModeTools.SetStyle(&FEditorStyle::Get(), "EditorModesToolbar");
 		EditorModeTools.SetLabelVisibility( EVisibility::Collapsed );
 
 		const FLevelEditorModesCommands& Commands = LevelEditorModule.GetLevelEditorModesCommands();
diff --git a/Engine/Source/Editor/LevelEditor/Private/SLevelViewportToolBar.cpp b/Engine/Source/Editor/LevelEditor/Private/SLevelViewportToolBar.cpp
--- a/Engine/Source/Editor/LevelEditor/Private/SLevelViewportToolBar.cpp
+++ b/Engine/Source/Editor/LevelEditor/Private/SLevelViewportToolBar.cpp
@@ -1006,7 +1006,7 @@ void SLevelViewportToolBar::GenerateViewportConfigsMenu(UToolMenu* Menu) const
 	{
 		FToolMenuSection& Section = Menu->AddSection("LevelViewportOnePaneConfigs", LOCTEXT("OnePaneConfigHeader", "One Pane"));
 
-		FToolBarBuilder OnePaneButton(CommandList, FMultiBoxCustomization::None);
+		FSlimHorizontalToolBarBuilder OnePaneButton(CommandList, FMultiBoxCustomization::None);
 		OnePaneButton.SetLabelVisibility(EVisibility::Collapsed);
 		OnePaneButton.SetStyle(&FEditorStyle::Get(), "ViewportLayoutToolbar");
 
@@ -1031,7 +1031,7 @@ void SLevelViewportToolBar::GenerateViewportConfigsMenu(UToolMenu* Menu) const
 
 	{
 		FToolMenuSection& Section = Menu->AddSection("LevelViewportTwoPaneConfigs", LOCTEXT("TwoPaneConfigHeader", "Two Panes"));
-		FToolBarBuilder TwoPaneButtons(CommandList, FMultiBoxCustomization::None);
+		FSlimHorizontalToolBarBuilder TwoPaneButtons(CommandList, FMultiBoxCustomization::None);
 		TwoPaneButtons.SetLabelVisibility(EVisibility::Collapsed);
 		TwoPaneButtons.SetStyle(&FEditorStyle::Get(), "ViewportLayoutToolbar");
 
@@ -1057,7 +1057,7 @@ void SLevelViewportToolBar::GenerateViewportConfigsMenu(UToolMenu* Menu) const
 
 	{
 		FToolMenuSection& Section = Menu->AddSection("LevelViewportThreePaneConfigs", LOCTEXT("ThreePaneConfigHeader", "Three Panes"));
-		FToolBarBuilder ThreePaneButtons(CommandList, FMultiBoxCustomization::None);
+		FSlimHorizontalToolBarBuilder ThreePaneButtons(CommandList, FMultiBoxCustomization::None);
 		ThreePaneButtons.SetLabelVisibility(EVisibility::Collapsed);
 		ThreePaneButtons.SetStyle(&FEditorStyle::Get(), "ViewportLayoutToolbar");
 
@@ -1085,7 +1085,7 @@ void SLevelViewportToolBar::GenerateViewportConfigsMenu(UToolMenu* Menu) const
 
 	{
 		FToolMenuSection& Section = Menu->AddSection("LevelViewportFourPaneConfigs", LOCTEXT("FourPaneConfigHeader", "Four Panes"));
-		FToolBarBuilder FourPaneButtons(CommandList, FMultiBoxCustomization::None);
+		FSlimHorizontalToolBarBuilder FourPaneButtons(CommandList, FMultiBoxCustomization::None);
 		FourPaneButtons.SetLabelVisibility(EVisibility::Collapsed);
 		FourPaneButtons.SetStyle(&FEditorStyle::Get(), "ViewportLayoutToolbar");
 
diff --git a/Engine/Source/Editor/LevelEditor/Public/LevelEditor.h b/Engine/Source/Editor/LevelEditor/Public/LevelEditor.h
--- a/Engine/Source/Editor/LevelEditor/Public/LevelEditor.h
+++ b/Engine/Source/Editor/LevelEditor/Public/LevelEditor.h
@@ -269,7 +269,6 @@ public:
 	virtual TSharedPtr<FExtensibilityManager> GetMenuExtensibilityManager() override {return MenuExtensibilityManager;}
 	virtual TSharedPtr<FExtensibilityManager> GetToolBarExtensibilityManager() override {return ToolBarExtensibilityManager;}
 	virtual TSharedPtr<FExtensibilityManager> GetModeBarExtensibilityManager() {return ModeBarExtensibilityManager;}
-	virtual TSharedPtr<FExtensibilityManager> GetNotificationBarExtensibilityManager() {return NotificationBarExtensibilityManager;}
 
 	virtual TSharedPtr<FExtender> AssembleExtenders(TSharedRef<FUICommandList>& InCommandList, TArray<FLevelEditorMenuExtender>& MenuExtenderDelegates) const;
 
diff --git a/Engine/Source/Editor/MergeActors/Private/SMergeActorsToolbar.cpp b/Engine/Source/Editor/MergeActors/Private/SMergeActorsToolbar.cpp
--- a/Engine/Source/Editor/MergeActors/Private/SMergeActorsToolbar.cpp
+++ b/Engine/Source/Editor/MergeActors/Private/SMergeActorsToolbar.cpp
@@ -19,6 +19,7 @@
 #include "IDocumentation.h"
 #include "IContentBrowserSingleton.h"
 #include "ContentBrowserModule.h"
+#include "Styling/ToolBarStyle.h"
 
 #define LOCTEXT_NAMESPACE "SMergeActorsToolbar"
 
@@ -196,19 +197,21 @@ void SMergeActorsToolbar::UpdateToolbar()
 	TSharedRef<SHorizontalBox> HorizontalBox =
 		SNew(SHorizontalBox);
 
+	const FToolBarStyle& BaseToolBarStyle = StyleSet.GetWidgetStyle<FToolBarStyle>("ToolBar");
+
 	for (int32 ToolIndex = 0; ToolIndex < RegisteredTools.Num(); ToolIndex++)
 	{
 		const IMergeActorsTool* Tool = RegisteredTools[ToolIndex];
 
 		HorizontalBox->AddSlot()
-		.Padding(StyleSet.GetMargin("EditorModesToolbar.SToolBarButtonBlock.Padding"))
+		.Padding(BaseToolBarStyle.ButtonPadding)
 		.AutoWidth()
 		[
 			SNew(SCheckBox)
-			.Style(&StyleSet, "EditorModesToolbar.ToggleButton")
+			.Style(&BaseToolBarStyle.ToggleButton)
 			.OnCheckStateChanged(this, &SMergeActorsToolbar::OnToolSelectionChanged, ToolIndex)
 			.IsChecked(this, &SMergeActorsToolbar::OnIsToolSelected, ToolIndex)
-			.Padding(StyleSet.GetMargin("EditorModesToolbar.SToolBarButtonBlock.CheckBox.Padding"))
+			.Padding(BaseToolBarStyle.CheckBoxPadding)
 			.ToolTip(IDocumentation::Get()->CreateToolTip(Tool->GetTooltipText(), nullptr, FString(), FString()))
 			[
 				SNew(SImage)
diff --git a/Engine/Source/Editor/Persona/Private/SAssetFamilyShortcutBar.cpp b/Engine/Source/Editor/Persona/Private/SAssetFamilyShortcutBar.cpp
--- a/Engine/Source/Editor/Persona/Private/SAssetFamilyShortcutBar.cpp
+++ b/Engine/Source/Editor/Persona/Private/SAssetFamilyShortcutBar.cpp
@@ -24,6 +24,7 @@
 #include "Widgets/Input/SComboButton.h"
 #include "Framework/Application/SlateApplication.h"
 #include "Subsystems/AssetEditorSubsystem.h"
+#include "Styling/ToolBarStyle.h"
 
 #define LOCTEXT_NAMESPACE "SAssetFamilyShortcutBar"
 
@@ -66,13 +67,15 @@ public:
 		bMultipleAssetsExist = Assets.Num() > 1;
 		AssetDirtyBrush = FEditorStyle::GetBrush("ContentBrowser.ContentDirty");
 
+		const FToolBarStyle& ToolBarStyle = FEditorStyle::Get().GetWidgetStyle<FToolBarStyle>("ToolBar");
+
 		ChildSlot
 		[
 			SNew(SHorizontalBox)
 			+SHorizontalBox::Slot()
 			[
 				SAssignNew(CheckBox, SCheckBox)
-				.Style(FEditorStyle::Get(), "ToolBar.ToggleButton")
+				.Style(&ToolBarStyle.ToggleButton)
 				.ForegroundColor(FSlateColor::UseForeground())
 				.Padding(0.0f)
 				.OnCheckStateChanged(this, &SAssetShortcut::HandleOpenAssetShortcut)
@@ -85,10 +88,10 @@ public:
 					.VAlign(VAlign_Center)
 					.AutoWidth()
 					[
-						SNew(SBorder)
-						.Padding(4.0f)
-						.BorderImage(FEditorStyle::GetBrush("PropertyEditor.AssetThumbnailShadow"))
-						[
+						//SNew(SBorder)
+						//.Padding(4.0f)
+						//.BorderImage(FEditorStyle::GetBrush("PropertyEditor.AssetThumbnailShadow"))
+						//[
 							SNew(SHorizontalBox)
 							+SHorizontalBox::Slot()
 							[
@@ -132,7 +135,7 @@ public:
 									]
 								]
 							]
-						]
+						//]
 					]
 					+SHorizontalBox::Slot()
 					.VAlign(VAlign_Center)
@@ -146,8 +149,7 @@ public:
 						[
 							SNew(STextBlock)
 							.Text(this, &SAssetShortcut::GetAssetText)
-							.TextStyle(FEditorStyle::Get(), "Toolbar.Label")
-							.ShadowOffset(FVector2D::UnitVector)
+							.TextStyle(&ToolBarStyle.LabelStyle)
 						]
 					]
 				]
@@ -161,7 +163,7 @@ public:
 				.Visibility(this, &SAssetShortcut::GetComboVisibility)
 				.ContentPadding(0)
 				.ForegroundColor(FSlateColor::UseForeground())
-				.ButtonStyle(FEditorStyle::Get(), "Toolbar.Button")
+				.ButtonStyle(&ToolBarStyle.ButtonStyle)
 				.OnGetMenuContent(this, &SAssetShortcut::HandleGetMenuContent)
 				.ToolTipText(LOCTEXT("AssetComboTooltip", "Find other assets of this type and perform asset operations./nShift-Click to open in new window."))
 			]
diff --git a/Engine/Source/Editor/Sequencer/Private/SSequencerPlayRateCombo.cpp b/Engine/Source/Editor/Sequencer/Private/SSequencerPlayRateCombo.cpp
--- a/Engine/Source/Editor/Sequencer/Private/SSequencerPlayRateCombo.cpp
+++ b/Engine/Source/Editor/Sequencer/Private/SSequencerPlayRateCombo.cpp
@@ -21,6 +21,7 @@
 #include "IContentBrowserSingleton.h"
 #include "ContentBrowserModule.h"
 #include "Misc/Timecode.h"
+#include "Styling/ToolBarStyle.h"
 
 #define LOCTEXT_NAMESPACE "SSequencerPlayRateCombo"
 
@@ -29,8 +30,7 @@ void SSequencerPlayRateCombo::Construct(const FArguments& InArgs, TWeakPtr<FSequ
 	WeakSequencer = InWeakSequencer;
 	WeakSequencerWidget = InWeakSequencerWidget;
 
-	FName BlockStyle = EMultiBlockLocation::ToName(ISlateStyle::Join( InArgs._StyleName, ".Button" ), InArgs._BlockLocation);
-	FName ColorStyle = ISlateStyle::Join( InArgs._StyleName, ".SToolBarComboButtonBlock.ComboButton.Color" );
+	const FToolBarStyle& SequencerToolBarStyle = FAppStyle::Get().GetWidgetStyle<FToolBarStyle>(InArgs._StyleName);
 
 	SetToolTipText(MakeAttributeSP(this, &SSequencerPlayRateCombo::GetToolTipText));
 
@@ -40,8 +40,7 @@ void SSequencerPlayRateCombo::Construct(const FArguments& InArgs, TWeakPtr<FSequ
 		SNew(SComboButton)
 		.ContentPadding(FMargin(2.f, 1.0f))
 		.VAlign(VAlign_Fill)
-		.ButtonStyle( InArgs._StyleSet, BlockStyle )
-		.ForegroundColor( InArgs._StyleSet->GetSlateColor( ColorStyle ) )
+		.ButtonStyle(&SequencerToolBarStyle.ButtonStyle)
 		.OnGetMenuContent(this, &SSequencerPlayRateCombo::OnCreateMenu)
 		.ButtonContent()
 		[
@@ -54,7 +53,7 @@ void SSequencerPlayRateCombo::Construct(const FArguments& InArgs, TWeakPtr<FSequ
 			[
 				SNew(STextBlock)
 				.Text(this, &SSequencerPlayRateCombo::GetFrameRateText)
-				.TextStyle( InArgs._StyleSet, ISlateStyle::Join( InArgs._StyleName, ".Label" ) )
+				.TextStyle(&SequencerToolBarStyle.LabelStyle)
 			]
 
 			+ SHorizontalBox::Slot()
@@ -64,7 +63,7 @@ void SSequencerPlayRateCombo::Construct(const FArguments& InArgs, TWeakPtr<FSequ
 			[
 				SNew(STextBlock)
 				.Visibility(this, &SSequencerPlayRateCombo::GetFrameLockedVisibility)
-				.TextStyle( InArgs._StyleSet, ISlateStyle::Join( InArgs._StyleName, ".Label" ) )
+				.TextStyle(&SequencerToolBarStyle.LabelStyle)
 				.Font(FEditorStyle::Get().GetFontStyle("FontAwesome.11"))
 				.Text(FEditorFontGlyphs::Lock)
 			]
@@ -77,7 +76,7 @@ void SSequencerPlayRateCombo::Construct(const FArguments& InArgs, TWeakPtr<FSequ
 				SNew(STextBlock)
 				.ToolTipText(this, &SSequencerPlayRateCombo::GetFrameRateIsMultipleOfErrorDescription)
 				.Visibility(this, &SSequencerPlayRateCombo::GetFrameRateIsMultipleOfErrorVisibility)
-				.TextStyle( InArgs._StyleSet, ISlateStyle::Join( InArgs._StyleName, ".Label" ) )
+				.TextStyle(&SequencerToolBarStyle.LabelStyle)
 				.Font(FEditorStyle::Get().GetFontStyle("FontAwesome.11"))
 				.Text(FEditorFontGlyphs::Exclamation_Triangle)
 			]
@@ -90,7 +89,7 @@ void SSequencerPlayRateCombo::Construct(const FArguments& InArgs, TWeakPtr<FSequ
 				SNew(STextBlock)
 				.ToolTipText(this, &SSequencerPlayRateCombo::GetFrameRateMismatchErrorDescription)
 				.Visibility(this, &SSequencerPlayRateCombo::GetFrameRateMismatchErrorVisibility)
-				.TextStyle(InArgs._StyleSet, ISlateStyle::Join(InArgs._StyleName, ".Label"))
+				.TextStyle(&SequencerToolBarStyle.LabelStyle)
 				.Font(FEditorStyle::Get().GetFontStyle("FontAwesome.11"))
 				.Text(FEditorFontGlyphs::Exclamation_Triangle)
 			]
diff --git a/Engine/Source/Editor/UMGEditor/Private/Designer/SDesignerView.cpp b/Engine/Source/Editor/UMGEditor/Private/Designer/SDesignerView.cpp
--- a/Engine/Source/Editor/UMGEditor/Private/Designer/SDesignerView.cpp
+++ b/Engine/Source/Editor/UMGEditor/Private/Designer/SDesignerView.cpp
@@ -71,6 +71,7 @@
 #include "DeviceProfiles/DeviceProfileManager.h"
 #include "Engine/DPICustomScalingRule.h"
 #include "UMGEditorModule.h"
+#include "Styling/ToolBarStyle.h"
 
 #define LOCTEXT_NAMESPACE "UMG"
 
@@ -537,6 +538,8 @@ EActiveTimerReturnType SDesignerView::EnsureTick(double InCurrentTime, float InD
 
 TSharedRef<SWidget> SDesignerView::CreateOverlayUI()
 {
+	const FToolBarStyle& ToolBarStyle = FEditorStyle::Get().GetWidgetStyle<FToolBarStyle>("ViewportMenu");
+
 	return SNew(SOverlay)
 
 	// Outline and text for important state.
@@ -623,7 +626,7 @@ TSharedRef<SWidget> SDesignerView::CreateOverlayUI()
 			.VAlign(VAlign_Center)
 			[
 				SNew(SButton)
-				.ButtonStyle(FEditorStyle::Get(), "ViewportMenu.Button")
+				.ButtonStyle(&ToolBarStyle.ButtonStyle)
 				.ToolTipText(LOCTEXT("ZoomToFit_ToolTip", "Zoom To Fit"))
 				.OnClicked(this, &SDesignerView::HandleZoomToFitClicked)
 				.ContentPadding(FEditorStyle::Get().GetMargin("ViewportMenu.SToolBarButtonBlock.Button.Padding"))
@@ -638,7 +641,7 @@ TSharedRef<SWidget> SDesignerView::CreateOverlayUI()
 			.VAlign(VAlign_Center)
 			[
 				SNew(SButton)
-				.ButtonStyle(FEditorStyle::Get(), "ViewportMenu.Button")
+				.ButtonStyle(&ToolBarStyle.ButtonStyle)
 				.ToolTipText(LOCTEXT("SwapAspectRatio_ToolTip", "Switch between Landscape and Portrait"))
 				.OnClicked(this, &SDesignerView::HandleSwapAspectRatioClicked)
 				.ContentPadding(FEditorStyle::Get().GetMargin("ViewportMenu.SToolBarButtonBlock.Button.Padding"))
@@ -653,7 +656,7 @@ TSharedRef<SWidget> SDesignerView::CreateOverlayUI()
 			.VAlign(VAlign_Center)
 			[
 				SNew(SButton)
-				.ButtonStyle(FEditorStyle::Get(), "ViewportMenu.Button")
+				.ButtonStyle(&ToolBarStyle.ButtonStyle)
 				.ToolTipText(LOCTEXT("Mirror_ToolTip", "Flip the current safe zones"))
 				.OnClicked(this, &SDesignerView::HandleFlipSafeZonesClicked)
 				.ContentPadding(FEditorStyle::Get().GetMargin("ViewportMenu.SToolBarButtonBlock.Button.Padding"))	
@@ -670,7 +673,7 @@ TSharedRef<SWidget> SDesignerView::CreateOverlayUI()
 		.VAlign(VAlign_Center)
 		[
 			SNew(SComboButton)
-			.ButtonStyle(FEditorStyle::Get(), "ViewportMenu.Button")
+			.ButtonStyle(&ToolBarStyle.ButtonStyle)
 			.ForegroundColor(FLinearColor::Black)
 			.OnGetMenuContent(this, &SDesignerView::GetResolutionsMenu)
 			.ContentPadding(FEditorStyle::Get().GetMargin("ViewportMenu.SToolBarButtonBlock.Button.Padding"))
@@ -688,7 +691,7 @@ TSharedRef<SWidget> SDesignerView::CreateOverlayUI()
 		.VAlign(VAlign_Center)
 		[
 			SNew(SComboButton)
-			.ButtonStyle(FEditorStyle::Get(), "ViewportMenu.Button")
+			.ButtonStyle(&ToolBarStyle.ButtonStyle)
 			.ForegroundColor(FLinearColor::Black)
 			.OnGetMenuContent(this, &SDesignerView::GetScreenSizingFillMenu)
 			.ContentPadding(FEditorStyle::Get().GetMargin("ViewportMenu.SToolBarButtonBlock.Button.Padding"))
diff --git a/Engine/Source/Editor/UnrealEd/Private/EditorModeManager.cpp b/Engine/Source/Editor/UnrealEd/Private/EditorModeManager.cpp
--- a/Engine/Source/Editor/UnrealEd/Private/EditorModeManager.cpp
+++ b/Engine/Source/Editor/UnrealEd/Private/EditorModeManager.cpp
@@ -954,7 +954,7 @@ TSharedRef<SDockTab> FEditorModeTools::MakeModeToolbarTab()
 		.Label(NSLOCTEXT("EditorModes", "EditorModesToolbarTitle", "Mode Toolbar"))
 		.ShouldAutosize(true)
 		.ContentPadding(0.0f)
-		.Icon(FEditorStyle::GetBrush("ToolBar.Icon"))
+		.Icon(FStyleDefaults::GetNoBrush())
 		[
 			SAssignNew(ModeToolbarBox, SVerticalBox)	
 		];
@@ -1100,7 +1100,7 @@ void FEditorModeTools::ActivateMode(FEditorModeID InID, bool bToggle)
 				for(auto Palette : PaletteNames)
 				{
 					const bool bUniform = true;
-					FToolBarBuilder ModeToolbarBuilder(CommandList, FMultiBoxCustomization(Mode->GetModeInfo().ToolbarCustomizationName), TSharedPtr<FExtender>(), Orient_Horizontal, false, bUniform);
+					FUniformToolBarBuilder ModeToolbarBuilder(CommandList, FMultiBoxCustomization(Mode->GetModeInfo().ToolbarCustomizationName));
 					ModeToolbarBuilder.SetStyle(&FEditorStyle::Get(), "PaletteToolBar");
 					Toolkit->BuildToolPalette(Palette, ModeToolbarBuilder);
 
@@ -1175,8 +1175,7 @@ void FEditorModeTools::ActivateMode(FEditorModeID InID, bool bToggle)
 				Toolkit->GetToolPaletteNames(PaletteNames);
 				for(auto Palette : PaletteNames)
 				{
-					const bool bUniform = true;
-					FToolBarBuilder ModeToolbarBuilder(CommandList, FMultiBoxCustomization(ScriptableMode->GetModeInfo().ToolbarCustomizationName), TSharedPtr<FExtender>(), Orient_Horizontal, false, bUniform);
+					FToolBarBuilder ModeToolbarBuilder(CommandList, FMultiBoxCustomization(ScriptableMode->GetModeInfo().ToolbarCustomizationName));
 					ModeToolbarBuilder.SetStyle(&FEditorStyle::Get(), "PaletteToolBar");
 					Toolkit->BuildToolPalette(Palette, ModeToolbarBuilder);
 
diff --git a/Engine/Source/Editor/UnrealEd/Private/SAssetEditorViewport.cpp b/Engine/Source/Editor/UnrealEd/Private/SAssetEditorViewport.cpp
--- a/Engine/Source/Editor/UnrealEd/Private/SAssetEditorViewport.cpp
+++ b/Engine/Source/Editor/UnrealEd/Private/SAssetEditorViewport.cpp
@@ -128,7 +128,7 @@ void SAssetEditorViewport::GenerateLayoutMenu(FMenuBuilder& MenuBuilder) const
 {
 	MenuBuilder.BeginSection("EditorViewportOnePaneConfigs", LOCTEXT("OnePaneConfigHeader", "One Pane"));
 	{
-		FToolBarBuilder OnePaneButton(CommandList, FMultiBoxCustomization::None);
+		FSlimHorizontalToolBarBuilder OnePaneButton(CommandList, FMultiBoxCustomization::None);
 		OnePaneButton.SetLabelVisibility(EVisibility::Collapsed);
 		OnePaneButton.SetStyle(&FEditorStyle::Get(), "ViewportLayoutToolbar");
 
@@ -153,7 +153,7 @@ void SAssetEditorViewport::GenerateLayoutMenu(FMenuBuilder& MenuBuilder) const
 
 	MenuBuilder.BeginSection("EditorViewportTwoPaneConfigs", LOCTEXT("TwoPaneConfigHeader", "Two Panes"));
 	{
-		FToolBarBuilder TwoPaneButtons(CommandList, FMultiBoxCustomization::None);
+		FSlimHorizontalToolBarBuilder TwoPaneButtons(CommandList, FMultiBoxCustomization::None);
 		TwoPaneButtons.SetLabelVisibility(EVisibility::Collapsed);
 		TwoPaneButtons.SetStyle(&FEditorStyle::Get(), "ViewportLayoutToolbar");
 
@@ -167,7 +167,7 @@ void SAssetEditorViewport::GenerateLayoutMenu(FMenuBuilder& MenuBuilder) const
 			[
 				TwoPaneButtons.MakeWidget()
 			]
-		+ SHorizontalBox::Slot()
+			+ SHorizontalBox::Slot()
 			.FillWidth(1)
 			[
 				SNullWidget::NullWidget
@@ -179,7 +179,7 @@ void SAssetEditorViewport::GenerateLayoutMenu(FMenuBuilder& MenuBuilder) const
 
 	MenuBuilder.BeginSection("EditorViewportThreePaneConfigs", LOCTEXT("ThreePaneConfigHeader", "Three Panes"));
 	{
-		FToolBarBuilder ThreePaneButtons(CommandList, FMultiBoxCustomization::None);
+		FSlimHorizontalToolBarBuilder ThreePaneButtons(CommandList, FMultiBoxCustomization::None);
 		ThreePaneButtons.SetLabelVisibility(EVisibility::Collapsed);
 		ThreePaneButtons.SetStyle(&FEditorStyle::Get(), "ViewportLayoutToolbar");
 
@@ -207,7 +207,7 @@ void SAssetEditorViewport::GenerateLayoutMenu(FMenuBuilder& MenuBuilder) const
 
 	MenuBuilder.BeginSection("EditorViewportFourPaneConfigs", LOCTEXT("FourPaneConfigHeader", "Four Panes"));
 	{
-		FToolBarBuilder FourPaneButtons(CommandList, FMultiBoxCustomization::None);
+		FSlimHorizontalToolBarBuilder FourPaneButtons(CommandList, FMultiBoxCustomization::None);
 		FourPaneButtons.SetLabelVisibility(EVisibility::Collapsed);
 		FourPaneButtons.SetStyle(&FEditorStyle::Get(), "ViewportLayoutToolbar");
 
diff --git a/Engine/Source/Editor/UnrealEd/Private/SViewportToolBarComboMenu.cpp b/Engine/Source/Editor/UnrealEd/Private/SViewportToolBarComboMenu.cpp
--- a/Engine/Source/Editor/UnrealEd/Private/SViewportToolBarComboMenu.cpp
+++ b/Engine/Source/Editor/UnrealEd/Private/SViewportToolBarComboMenu.cpp
@@ -9,11 +9,21 @@
 #include "Widgets/Input/SMenuAnchor.h"
 #include "Widgets/Input/SButton.h"
 #include "EditorStyleSet.h"
+#include "Styling/ToolBarStyle.h"
 
 void SViewportToolBarComboMenu::Construct( const FArguments& InArgs )
 {
-	FName ButtonStyle = FEditorStyle::Join(InArgs._Style.Get(), ".Button");
-	FName CheckboxStyle = FEditorStyle::Join(InArgs._Style.Get(), ".ToggleButton");
+	const FToolBarStyle& ToolBarStyle = FAppStyle::Get().GetWidgetStyle<FToolBarStyle>(InArgs._Style.Get());
+
+	EMultiBlockLocation::Type BlockLocation = InArgs._BlockLocation;
+
+	FName ButtonBlockStyle = EMultiBlockLocation::ToName(FEditorStyle::Join(InArgs._Style.Get(), ".Button"), BlockLocation);
+	const FButtonStyle& ButtonStyle = BlockLocation == EMultiBlockLocation::None ? ToolBarStyle.ButtonStyle : FAppStyle::Get().GetWidgetStyle<FButtonStyle>(ButtonBlockStyle);
+
+	FName CheckBoxBlockStyle = EMultiBlockLocation::ToName(FEditorStyle::Join(InArgs._Style.Get(), ".ToggleButton"), BlockLocation);
+	const FCheckBoxStyle& CheckBoxStyle = BlockLocation == EMultiBlockLocation::None ? ToolBarStyle.ToggleButton : FAppStyle::Get().GetWidgetStyle<FCheckBoxStyle>(CheckBoxBlockStyle);
+
+	const FTextBlockStyle& LabelStyle = ToolBarStyle.LabelStyle;
 
 	const FSlateIcon& Icon = InArgs._Icon.Get();
 
@@ -24,7 +34,7 @@ void SViewportToolBarComboMenu::Construct( const FArguments& InArgs )
 		ToggleControl = SNew(SCheckBox)
 		.Cursor( EMouseCursor::Default )
 		.Padding(FMargin( 4.0f ))
-		.Style(FEditorStyle::Get(), EMultiBlockLocation::ToName(CheckboxStyle, InArgs._BlockLocation))
+		.Style(&CheckBoxStyle)
 		.OnCheckStateChanged(InArgs._OnCheckStateChanged)
 		.ToolTipText(InArgs._ToggleButtonToolTip)
 		.IsChecked(InArgs._IsChecked)
@@ -45,7 +55,7 @@ void SViewportToolBarComboMenu::Construct( const FArguments& InArgs )
 	{
 		TSharedRef<SWidget> ButtonContents =
 			SNew(SButton)
-			.ButtonStyle( FEditorStyle::Get(), EMultiBlockLocation::ToName(ButtonStyle,EMultiBlockLocation::End) )
+			.ButtonStyle(&ButtonStyle)
 			.ContentPadding( FMargin( 5.0f, 0.0f ) )
 			.ToolTipText(InArgs._MenuButtonToolTip)
 			.OnClicked(this, &SViewportToolBarComboMenu::OnMenuClicked)
@@ -57,7 +67,7 @@ void SViewportToolBarComboMenu::Construct( const FArguments& InArgs )
 				.VAlign(VAlign_Top)
 				[
 					SNew(STextBlock)
-					.TextStyle( FEditorStyle::Get(), FEditorStyle::Join( InArgs._Style.Get(), ".Label" ) )
+					.TextStyle(&LabelStyle)
 					.Text(InArgs._Label)
 				]
 				+SVerticalBox::Slot()
diff --git a/Engine/Source/Editor/UnrealEd/Private/SViewportToolBarIconMenu.cpp b/Engine/Source/Editor/UnrealEd/Private/SViewportToolBarIconMenu.cpp
--- a/Engine/Source/Editor/UnrealEd/Private/SViewportToolBarIconMenu.cpp
+++ b/Engine/Source/Editor/UnrealEd/Private/SViewportToolBarIconMenu.cpp
@@ -8,12 +8,13 @@
 #include "Widgets/Input/SMenuAnchor.h"
 #include "Widgets/Input/SButton.h"
 #include "EditorStyleSet.h"
+#include "Styling/ToolBarStyle.h"
 
 void SViewportToolBarIconMenu::Construct( const FArguments& InArgs )
 {
 	ParentToolBar = InArgs._ParentToolBar;
 
-	const FName ButtonStyle = FEditorStyle::Join(InArgs._Style.Get(), ".Button");
+	const FToolBarStyle& ViewportToolbarStyle = FAppStyle::Get().GetWidgetStyle<FToolBarStyle>(InArgs._Style.Get());
 
 	const FSlateIcon& Icon = InArgs._Icon.Get();
 
@@ -22,7 +23,7 @@ void SViewportToolBarIconMenu::Construct( const FArguments& InArgs )
 	.OnGetMenuContent( InArgs._OnGetMenuContent )
 	[
 		SNew(SButton)
-		.ButtonStyle( FEditorStyle::Get(), ButtonStyle )
+		.ButtonStyle(&ViewportToolbarStyle.ButtonStyle)
 		.ContentPadding( FMargin( 5.0f, 0.0f ) )
 		.OnClicked(this, &SViewportToolBarIconMenu::OnMenuClicked)
 		[
@@ -59,7 +60,7 @@ void SViewportToolBarIconMenu::Construct( const FArguments& InArgs )
 				.VAlign(VAlign_Top)
 				[
 					SNew(STextBlock)
-					.TextStyle( FEditorStyle::Get(), FEditorStyle::Join( InArgs._Style.Get(), ".Label" ) )
+					.TextStyle(&ViewportToolbarStyle.LabelStyle)
 					.Text(InArgs._Label)
 				]
 				+SVerticalBox::Slot()
diff --git a/Engine/Source/Editor/UnrealEd/Private/Toolkits/AssetEditorToolkit.cpp b/Engine/Source/Editor/UnrealEd/Private/Toolkits/AssetEditorToolkit.cpp
--- a/Engine/Source/Editor/UnrealEd/Private/Toolkits/AssetEditorToolkit.cpp
+++ b/Engine/Source/Editor/UnrealEd/Private/Toolkits/AssetEditorToolkit.cpp
@@ -22,7 +22,6 @@
 #include "CollectionManagerModule.h"
 #include "Widgets/SToolTip.h"
 #include "IDocumentation.h"
-#include "IIntroTutorials.h"
 #include "Widgets/Docking/SDockTab.h"
 #include "IAssetTools.h"
 #include "IAssetTypeActions.h"
@@ -183,23 +182,9 @@ void FAssetEditorToolkit::InitAssetEditor( const EToolkitMode::Type Mode, const
 			}
 		}
 
-		IIntroTutorials& IntroTutorials = FModuleManager::LoadModuleChecked<IIntroTutorials>(TEXT("IntroTutorials"));
-		TSharedRef<SWidget> TutorialWidget = IntroTutorials.CreateTutorialsWidget(GetToolkitContextFName(), NewMajorTab->GetParentWindow());
-
-		NewMajorTab->SetRightContent(
-				SNew(SHorizontalBox)
-
-				+SHorizontalBox::Slot()
-				.AutoWidth()
-				.Padding(8.0f, 0.0f, 8.0f, 0.0f)
-				.VAlign(VAlign_Center)
-				[
-					TutorialWidget
-				]	
-			);
-
 		const TSharedRef<FTabManager> NewTabManager = FGlobalTabmanager::Get()->NewTabManager( NewMajorTab.ToSharedRef() );		
 		NewTabManager->SetOnPersistLayout(FTabManager::FOnPersistLayout::CreateRaw(this, &FAssetEditorToolkit::HandleTabManagerPersistLayout));
+		NewTabManager->SetAllowWindowMenuBar(true);
 		this->TabManager = NewTabManager;
 
 		NewMajorTab->SetContent
@@ -1066,17 +1051,24 @@ void FAssetEditorToolkit::GenerateToolbar()
 	GeneratedToolbar->bToolBarForceSmallIcons = bIsToolbarUsingSmallIcons;
 	TSharedRef< class SWidget > ToolBarWidget = ToolMenus->GenerateWidget(GeneratedToolbar);
 
-	TSharedRef<SHorizontalBox> MiscWidgets = SNew(SHorizontalBox);
+	TSharedRef<SWidget> MiscWidgets = SNullWidget::NullWidget;
 
-	for (int32 WidgetIdx = 0; WidgetIdx < ToolbarWidgets.Num(); ++WidgetIdx)
+	if(ToolbarWidgets.Num() > 0)
 	{
-		MiscWidgets->AddSlot()
-		.AutoWidth()
-		.VAlign(VAlign_Center)
-		.Padding(0.0f, 2.0f, 0.0f, 2.0f)
-		[
-			ToolbarWidgets[WidgetIdx]
-		];
+		TSharedRef<SHorizontalBox> MiscWidgetsHBox = SNew(SHorizontalBox);
+
+		for (int32 WidgetIdx = 0; WidgetIdx < ToolbarWidgets.Num(); ++WidgetIdx)
+		{
+			MiscWidgetsHBox->AddSlot()
+				.AutoWidth()
+				.VAlign(VAlign_Center)
+				.Padding(0.0f, 2.0f, 0.0f, 2.0f)
+				[
+					ToolbarWidgets[WidgetIdx]
+				];
+		}
+
+		MiscWidgets = MiscWidgetsHBox;
 	}
 	
 	Toolbar = 
@@ -1103,12 +1095,7 @@ void FAssetEditorToolkit::GenerateToolbar()
 			.AutoHeight()
 			.VAlign(VAlign_Bottom)
 			[
-				SNew(SBorder)
-				.BorderImage(FEditorStyle::GetBrush(TEXT("Toolbar.Background")))
-				.Visibility(ToolbarWidgets.Num() > 0 ? EVisibility::Visible : EVisibility::Collapsed)
-				[
-					MiscWidgets
-				]
+				MiscWidgets
 			]
 		];
 
diff --git a/Engine/Source/Runtime/Engine/Private/Slate/SlateGameResources.cpp b/Engine/Source/Runtime/Engine/Private/Slate/SlateGameResources.cpp
--- a/Engine/Source/Runtime/Engine/Private/Slate/SlateGameResources.cpp
+++ b/Engine/Source/Runtime/Engine/Private/Slate/SlateGameResources.cpp
@@ -47,7 +47,7 @@ void FSlateGameResources::SetContentRoot( const FString& InContentRootDir )
 	ContentRootDir = InContentRootDir;
 }
 
-const FSlateBrush* FSlateGameResources::GetBrush( const FName PropertyName, const ANSICHAR* Specifier ) const
+const FSlateBrush* FSlateGameResources::GetBrush( const FName PropertyName, const ANSICHAR* Specifier, const ISlateStyle* RequestingStyle ) const
 {
 	ensureMsgf(Specifier == NULL, TEXT("Attempting to look up resource (%s, %s). \n Specifiers not supported by Slate Resource Sets loaded from content browser."), *PropertyName.ToString(), Specifier);
 	FName CleanName = GetCleanName(PropertyName);
@@ -58,7 +58,7 @@ const FSlateBrush* FSlateGameResources::GetBrush( const FName PropertyName, cons
 		ensureMsgf(BrushAsset, TEXT("Could not find resource '%s'"), *CleanName.ToString());
 		return BrushAsset ? &BrushAsset->Brush : GetDefaultBrush();
 	}
-	return FSlateStyleSet::GetBrush(PropertyName, Specifier);
+	return FSlateStyleSet::GetBrush(PropertyName, Specifier, RequestingStyle);
 }
 
 const FSlateBrush* FSlateGameResources::GetOptionalBrush(const FName PropertyName, const ANSICHAR* Specifier, const FSlateBrush* const InDefaultBrush) const
@@ -125,14 +125,14 @@ void FSlateGameResources::GetResources( TArray< const FSlateBrush* >& OutResourc
 	}
 }
 
-const FSlateWidgetStyle* FSlateGameResources::GetWidgetStyleInternal( const FName DesiredTypeName, const FName StyleName ) const
+const FSlateWidgetStyle* FSlateGameResources::GetWidgetStyleInternal(const FName DesiredTypeName, const FName StyleName, const FSlateWidgetStyle* DefaultStyle, bool bWarnIfNotFound) const
 {
 	UObject* const* UIResourcePtr = UIResources.Find( StyleName );
 	USlateWidgetStyleAsset* StyleAsset = UIResourcePtr ? Cast<USlateWidgetStyleAsset>(*UIResourcePtr) : NULL;
 
 	if ( StyleAsset == NULL )
 	{
-		return FSlateStyleSet::GetWidgetStyleInternal( DesiredTypeName, StyleName );
+		return FSlateStyleSet::GetWidgetStyleInternal( DesiredTypeName, StyleName, DefaultStyle, bWarnIfNotFound);
 	}
 
 	const FSlateWidgetStyle* Style = StyleAsset->GetStyleChecked( DesiredTypeName );
@@ -141,7 +141,8 @@ const FSlateWidgetStyle* FSlateGameResources::GetWidgetStyleInternal( const FNam
 	{
 		TSharedRef< FTokenizedMessage > Message = FTokenizedMessage::Create( EMessageSeverity::Error, FText::Format( NSLOCTEXT("SlateStyleSet", "WrongWidgetStyleType", "The Slate Widget Style '{0}' is not of the desired type. Desired: '{1}', Actual: '{2}'"), FText::FromName( StyleName ), FText::FromName( DesiredTypeName ), FText::FromName( StyleAsset->CustomStyle->GetStyle()->GetTypeName() ) ) );
 		Message->AddToken( FAssetNameToken::Create( StyleAsset->GetPathName(), FText::FromString( StyleAsset->GetName() ) ) );
-		Log( Message );
+	
+		Log(Message);
 	}
 
 	return Style;
diff --git a/Engine/Source/Runtime/Engine/Public/Slate/SlateGameResources.h b/Engine/Source/Runtime/Engine/Public/Slate/SlateGameResources.h
--- a/Engine/Source/Runtime/Engine/Public/Slate/SlateGameResources.h
+++ b/Engine/Source/Runtime/Engine/Public/Slate/SlateGameResources.h
@@ -45,7 +45,7 @@ public:
 
 	virtual void SetContentRoot( const FString& InContentRootDir ) override;
 
-	virtual const FSlateBrush* GetBrush( const FName PropertyName, const ANSICHAR* Specifier = NULL ) const override;
+	virtual const FSlateBrush* GetBrush( const FName PropertyName, const ANSICHAR* Specifier = NULL, const ISlateStyle* RequestingStyle = nullptr ) const override;
 	virtual const FSlateBrush* GetOptionalBrush(const FName PropertyName, const ANSICHAR* Specifier = NULL, const FSlateBrush* const DefaultBrush = FStyleDefaults::GetNoBrush()) const override;
 
 	UCurveFloat* GetCurveFloat( const FName AssetName ) const;
@@ -54,7 +54,7 @@ public:
 
 protected:
 
-	virtual const FSlateWidgetStyle* GetWidgetStyleInternal( const FName DesiredTypeName, const FName StyleName ) const override;
+	virtual const FSlateWidgetStyle* GetWidgetStyleInternal( const FName DesiredTypeName, const FName StyleName, const FSlateWidgetStyle* DefaultStyle, bool bWarnIfNotFound) const override;
 
 	virtual void Log( ISlateStyle::EStyleMessageSeverity Severity, const FText& Message ) const override;
 
diff --git a/Engine/Source/Runtime/Slate/Private/Framework/MultiBox/MultiBox.cpp b/Engine/Source/Runtime/Slate/Private/Framework/MultiBox/MultiBox.cpp
--- a/Engine/Source/Runtime/Slate/Private/Framework/MultiBox/MultiBox.cpp
+++ b/Engine/Source/Runtime/Slate/Private/Framework/MultiBox/MultiBox.cpp
@@ -23,6 +23,7 @@
 #include "Framework/MultiBox/SClippingHorizontalBox.h"
 #include "Framework/Commands/UICommandDragDropOp.h"
 #include "SUniformToolbarPanel.h"
+#include "Styling/ToolBarStyle.h"
 
 
 TAttribute<bool> FMultiBoxSettings::UseSmallToolBarIcons;
@@ -150,7 +151,7 @@ TSharedRef< IMultiBlockBaseWidget > FMultiBlock::MakeWidget( TSharedRef< SMultiB
 
 	// Work out what style the widget should be using
 	const ISlateStyle* const StyleSet = InOwnerMultiBoxWidget->GetStyleSet();
-	const FName& StyleName = InOwnerMultiBoxWidget->GetStyleName();
+	FName StyleName = StyleNameOverride != NAME_None ? StyleNameOverride : InOwnerMultiBoxWidget->GetStyleName();
 
 	// Build up the widget
 	NewMultiBlockWidget->BuildMultiBlockWidget(StyleSet, StyleName);
@@ -181,6 +182,11 @@ FMultiBox::FMultiBox( const EMultiBoxType InType, FMultiBoxCustomization InCusto
 	, Type( InType )
 	, bShouldCloseWindowAfterMenuSelection( bInShouldCloseWindowAfterMenuSelection )
 {
+
+	if (InType == EMultiBoxType::SlimHorizontalToolBar && FCoreStyle::IsStarshipStyle())
+	{
+		StyleName = "SlimToolBar";
+	}
 }
 
 FMultiBox::~FMultiBox()
@@ -414,6 +420,7 @@ TSharedPtr<FMultiBlock> FMultiBox::MakeMultiBlockFromCommand( TSharedPtr<const F
 		switch ( Type )
 		{
 		case EMultiBoxType::ToolBar:
+		case EMultiBoxType::UniformToolBar:
 			{
 				NewBlock = MakeShareable( new FToolBarButtonBlock( CommandInfo, CommandList ) );
 			}
@@ -598,30 +605,39 @@ void SMultiBoxWidget::AddBlockWidget( const FMultiBlock& Block, TSharedPtr<SHori
 	{
 		FinalWidget = BlockWidget;
 	}
+	
+	TSharedRef<SWidget> FinalWidgetWithHook = SNullWidget::NullWidget;
 
-	TSharedRef<SWidget> FinalWidgetWithHook = 
-		SNew(SVerticalBox)
-		+ SVerticalBox::Slot()
-		.HAlign(HAlign_Center)
-		.AutoHeight()
-		[
-			SNew(STextBlock)
-			.Visibility(bDisplayExtensionHooks ? EVisibility::Visible : EVisibility::Collapsed)
-			.ColorAndOpacity(StyleSet->GetColor("MultiboxHookColor"))
-			.Text(FText::FromName(Block.GetExtensionHook()))
-		]
-		+ SVerticalBox::Slot()
-		[
-			FinalWidget.ToSharedRef()
-		];
+	if (bDisplayExtensionHooks)
+	{
+		FinalWidgetWithHook =
+			SNew(SVerticalBox)
+			+ SVerticalBox::Slot()
+			.HAlign(HAlign_Center)
+			.AutoHeight()
+			[
+				SNew(STextBlock)
+				.Visibility(bDisplayExtensionHooks ? EVisibility::Visible : EVisibility::Collapsed)
+				.ColorAndOpacity(StyleSet->GetColor("MultiboxHookColor"))
+				.Text(FText::FromName(Block.GetExtensionHook()))
+			]
+			+ SVerticalBox::Slot()
+			[
+				FinalWidget.ToSharedRef()
+			];
+	}
+	else
+	{
+		FinalWidgetWithHook = FinalWidget.ToSharedRef();
+	}
+	
 
 	switch (MultiBox->GetType())
 	{
 	case EMultiBoxType::MenuBar:
-	case EMultiBoxType::ToolMenuBar:
 	case EMultiBoxType::ToolBar:
+	case EMultiBoxType::SlimHorizontalToolBar:
 		{
-
 			HorizontalBox->AddSlot()
 			.AutoWidth()
 			.Padding(0)
@@ -667,7 +683,7 @@ void SMultiBoxWidget::AddBlockWidget( const FMultiBlock& Block, TSharedPtr<SHori
 		{
 			VerticalBox->AddSlot()
 			.AutoHeight()
-			.Padding( 1.0f, 0.0f, 1.0f, 0.0f )
+			.Padding( 0.0f, 0.0f, 0.0f, 0.0f )
 			[
 				SNew(SHorizontalBox)
 				+SHorizontalBox::Slot()
@@ -715,7 +731,16 @@ void SMultiBoxWidget::BuildMultiBoxWidget()
 	// Select background brush based on the type of multibox.
 	const ISlateStyle* const StyleSet = MultiBox->GetStyleSet();
 	const FName& StyleName = MultiBox->GetStyleName();
-	const FSlateBrush* BackgroundBrush = StyleSet->GetBrush( StyleName, ".Background" );
+
+	const FSlateBrush* BackgroundBrush = nullptr;
+	if (StyleSet->HasWidgetStyle<FToolBarStyle>(StyleName))
+	{
+		BackgroundBrush = &StyleSet->GetWidgetStyle<FToolBarStyle>(StyleName).BackgroundBrush;
+	}
+	else
+	{
+		BackgroundBrush = StyleSet->GetOptionalBrush(StyleName, ".Background", FStyleDefaults::GetNoBrush());
+	}
 
 	// Create a box panel that the various multiblocks will resides within
 	// @todo Slate MultiBox: Expose margins and other useful bits
@@ -732,11 +757,12 @@ void SMultiBoxWidget::BuildMultiBoxWidget()
 	switch (MultiBox->GetType())
 	{
 	case EMultiBoxType::MenuBar:
+		MainWidget = HorizontalBox = SNew(SHorizontalBox);
+		break;
 	case EMultiBoxType::ToolBar:
-	case EMultiBoxType::ToolMenuBar:
+	case EMultiBoxType::SlimHorizontalToolBar:
 		{
 			MainWidget = HorizontalBox = ClippedHorizontalBox = SNew(SClippingHorizontalBox)
-				.BackgroundBrush(BackgroundBrush)
 				.OnWrapButtonClicked(FOnGetContent::CreateSP(this, &SMultiBoxWidget::OnWrapButtonClicked))
 				.StyleSet(StyleSet)
 				.StyleName(StyleName);
@@ -909,39 +935,38 @@ void SMultiBoxWidget::BuildMultiBoxWidget()
 		ClippedHorizontalBox->AddWrapButton();
 	}
 
+	FMargin BorderPadding(0);
+	const FSlateBrush* BorderBrush = FStyleDefaults::GetNoBrush();
+	FSlateColor BorderForegroundColor = FSlateColor::UseForeground();
+
 	// Setup the root border widget
-	TSharedPtr< SBorder > RootBorder;
+	TSharedPtr<SWidget> RootBorder;
+
 	switch (MultiBox->GetType())
 	{
-	case EMultiBoxType::MenuBar:
-	case EMultiBoxType::ToolBar:
-	case EMultiBoxType::ToolMenuBar:
+	case EMultiBoxType::Menu:
 		{
-			RootBorder =
-				SNew( SBorder )
-				.Padding(0)
-				.BorderImage( FCoreStyle::Get().GetBrush("NoBorder") )
-				// Assign the box panel as the child
-				[
-					MainWidget.ToSharedRef()
-				];
+			BorderPadding = FMargin(0, 0, 0, 7);
 		}
 		break;
 	default:
-		{
-			RootBorder =
-				SNew( SBorder )
-				.Padding(0)
-				.BorderImage( BackgroundBrush )
-				.ForegroundColor( FCoreStyle::Get().GetSlateColor("DefaultForeground") )
-				// Assign the box panel as the child
-				[
-					MainWidget.ToSharedRef()
-				];
+		{ 
+			BorderBrush = BackgroundBrush;
+			BorderForegroundColor = FCoreStyle::Get().GetSlateColor("DefaultForeground");
 		}
 		break;
 	}
 
+	RootBorder =
+		SNew(SBorder)
+		.Padding(BorderPadding)
+		.BorderImage(BorderBrush)
+		.ForegroundColor(BorderForegroundColor)
+		// Assign the box panel as the child
+		[
+			MainWidget.ToSharedRef()
+		];
+
 	// Prevent tool-tips spawned by child widgets from drawing on top of our main widget
 	RootBorder->EnableToolTipForceField( true );
 
diff --git a/Engine/Source/Runtime/Slate/Private/Framework/MultiBox/MultiBoxBuilder.cpp b/Engine/Source/Runtime/Slate/Private/Framework/MultiBox/MultiBoxBuilder.cpp
--- a/Engine/Source/Runtime/Slate/Private/Framework/MultiBox/MultiBoxBuilder.cpp
+++ b/Engine/Source/Runtime/Slate/Private/Framework/MultiBox/MultiBoxBuilder.cpp
@@ -18,6 +18,7 @@
 #include "Framework/MultiBox/SGroupMarkerBlock.h"
 #include "Framework/MultiBox/ToolMenuBase.h"
 #include "Widgets/Layout/SScrollBox.h"
+#include "Styling/ToolBarStyle.h"
 
 
 FMultiBoxBuilder::FMultiBoxBuilder( const EMultiBoxType InType, FMultiBoxCustomization InCustomization, const bool bInShouldCloseWindowAfterMenuSelection, const TSharedPtr< const FUICommandList >& InCommandList, TSharedPtr<FExtender> InExtender, FName InTutorialHighlightName, FName InMenuName )
@@ -334,8 +335,11 @@ void FMenuBuilder::ApplySectionBeginning()
 {
 	if (bSectionNeedsToBeApplied)
 	{
+		bool bHasSearchWidget = MultiBox->SearchTextWidget.IsValid();
+		const int32 NumBlocksBeforeSeparator = bHasSearchWidget ? 2 : 1;
+
 		// Do not count search block, which starts as invisible
-		if( MultiBox->GetBlocks().Num() > 1 || FMultiBoxSettings::DisplayMultiboxHooks.Get() )
+		if( MultiBox->GetBlocks().Num() >= NumBlocksBeforeSeparator || FMultiBoxSettings::DisplayMultiboxHooks.Get() )
 		{
 			MultiBox->AddMultiBlock( MakeShareable( new FMenuSeparatorBlock(CurrentSectionExtensionHook, /* bInIsPartOfHeading=*/ true) ) );
 		}
@@ -395,6 +399,7 @@ void FToolBarBuilder::AddToolBarButton(const TSharedPtr< const FUICommandInfo >
 	NewToolBarButtonBlock->SetIsFocusable(bIsFocusable);
 	NewToolBarButtonBlock->SetForceSmallIcons(bForceSmallIcons);
 	NewToolBarButtonBlock->SetTutorialHighlightName(GenerateTutorialIdentfierName(TutorialHighlightName, InTutorialHighlightName, InCommand, MultiBox->GetBlocks().Num()));
+	NewToolBarButtonBlock->SetStyleNameOverride(CurrentStyleOverride);
 
 	MultiBox->AddMultiBlock( NewToolBarButtonBlock );
 
@@ -417,6 +422,7 @@ void FToolBarBuilder::AddToolBarButton(const FUIAction& InAction, FName InExtens
 	NewToolBarButtonBlock->SetIsFocusable(bIsFocusable);
 	NewToolBarButtonBlock->SetForceSmallIcons(bForceSmallIcons);
 	NewToolBarButtonBlock->SetTutorialHighlightName(GenerateTutorialIdentfierName(TutorialHighlightName, InTutorialHighlightName, nullptr, MultiBox->GetBlocks().Num()));
+	NewToolBarButtonBlock->SetStyleNameOverride(CurrentStyleOverride);
 
 	MultiBox->AddMultiBlock( NewToolBarButtonBlock );
 
@@ -436,6 +442,7 @@ void FToolBarBuilder::AddComboButton( const FUIAction& InAction, const FOnGetCon
 
 	NewToolBarComboButtonBlock->SetForceSmallIcons(bForceSmallIcons);
 	NewToolBarComboButtonBlock->SetTutorialHighlightName(GenerateTutorialIdentfierName(TutorialHighlightName, InTutorialHighlightName, nullptr, MultiBox->GetBlocks().Num()));
+	NewToolBarComboButtonBlock->SetStyleNameOverride(CurrentStyleOverride);
 
 	MultiBox->AddMultiBlock( NewToolBarComboButtonBlock );
 }
@@ -447,6 +454,8 @@ void FToolBarBuilder::AddToolBarWidget( TSharedRef<SWidget> InWidget, const TAtt
 	// If tutorial name specified, wrap in tutorial wrapper
 	const FName WrapperName = GenerateTutorialIdentfierName(InTutorialHighlightName, NAME_None, nullptr, MultiBox->GetBlocks().Num());
 
+	const FToolBarStyle& ToolBarStyle = GetStyleSet()->GetWidgetStyle<FToolBarStyle>(GetStyleName());
+
 	TSharedRef<SWidget> ChildWidget = InWidget;
 	InWidget = 
 		SNew( SVerticalBox )
@@ -461,7 +470,7 @@ void FToolBarBuilder::AddToolBarWidget( TSharedRef<SWidget> InWidget, const TAtt
 
 		+SVerticalBox::Slot()
 		.AutoHeight()
-		.Padding(GetStyleSet()->GetMargin(ISlateStyle::Join( GetStyleName(), ".Label.Padding" )))
+		.Padding(ToolBarStyle.LabelPadding)
 		.HAlign( HAlign_Center )
 		[
 			SNew( STextBlock )
@@ -473,8 +482,7 @@ void FToolBarBuilder::AddToolBarWidget( TSharedRef<SWidget> InWidget, const TAtt
 			})
 			.IsEnabled_Lambda( [ChildWidget] () -> bool { return ChildWidget->IsEnabled(); } )
 			.Text( InLabel )
-			.TextStyle( GetStyleSet(), ISlateStyle::Join( GetStyleName(), ".Label" ) )	// Smaller font for tool tip labels
-			.ShadowOffset( FVector2D::UnitVector )
+			.TextStyle(&ToolBarStyle.LabelStyle)	// Smaller font for tool tip labels
 		] ;
 	
 	TSharedRef< FWidgetBlock > NewWidgetBlock( new FWidgetBlock( InWidget, FText::GetEmpty(), true ) );
@@ -508,11 +516,10 @@ void FToolBarBuilder::AddSeparator(FName InExtensionHook)
 
 	ApplyHook(InExtensionHook, EExtensionHook::Before);
 
-	// Never add a menu separate as the first item, even if we were asked to
-	if( MultiBox->GetBlocks().Num() > 0 || FMultiBoxSettings::DisplayMultiboxHooks.Get() )
-	{
-		MultiBox->AddMultiBlock( MakeShareable( new FToolBarSeparatorBlock(InExtensionHook) ) );
-	}
+	TSharedRef<FToolBarSeparatorBlock> NewSeparatorBlock = MakeShared<FToolBarSeparatorBlock>(InExtensionHook);
+	NewSeparatorBlock->SetStyleNameOverride(CurrentStyleOverride);
+
+	MultiBox->AddMultiBlock(NewSeparatorBlock);
 
 	ApplyHook(InExtensionHook, EExtensionHook::After);
 }
@@ -565,7 +572,10 @@ void FToolBarBuilder::ApplySectionBeginning()
 	{
 		if( MultiBox->GetBlocks().Num() > 0 || FMultiBoxSettings::DisplayMultiboxHooks.Get() )
 		{
-			MultiBox->AddMultiBlock( MakeShareable( new FToolBarSeparatorBlock(CurrentSectionExtensionHook) ) );
+			TSharedRef<FToolBarSeparatorBlock> NewSeparatorBlock = MakeShared<FToolBarSeparatorBlock>(CurrentSectionExtensionHook);
+			NewSeparatorBlock->SetStyleNameOverride(CurrentStyleOverride);
+
+			MultiBox->AddMultiBlock(NewSeparatorBlock);
 		}
 		bSectionNeedsToBeApplied = false;
 	}
@@ -575,16 +585,28 @@ void FToolBarBuilder::EndBlockGroup()
 {
 	ApplySectionBeginning();
 
-	TSharedRef< FGroupEndBlock > NewGroupEndBlock( new FGroupEndBlock() );
-
+	TSharedRef<FGroupEndBlock> NewGroupEndBlock( new FGroupEndBlock() );
+	NewGroupEndBlock->SetStyleNameOverride(CurrentStyleOverride);
+	
 	MultiBox->AddMultiBlock( NewGroupEndBlock );
 }
 
+void FToolBarBuilder::BeginStyleOverride(FName StyleOverrideName)
+{
+	CurrentStyleOverride = StyleOverrideName;
+}
+
+void FToolBarBuilder::EndStyleOverride()
+{
+	CurrentStyleOverride = NAME_None;
+}
+
 void FToolBarBuilder::BeginBlockGroup()
 {
 	ApplySectionBeginning();
 
 	TSharedRef< FGroupStartBlock > NewGroupStartBlock( new FGroupStartBlock() );
+	NewGroupStartBlock->SetStyleNameOverride(CurrentStyleOverride);
 
 	MultiBox->AddMultiBlock( NewGroupStartBlock );
 }
diff --git a/Engine/Source/Runtime/Slate/Private/Framework/MultiBox/SClippingHorizontalBox.cpp b/Engine/Source/Runtime/Slate/Private/Framework/MultiBox/SClippingHorizontalBox.cpp
--- a/Engine/Source/Runtime/Slate/Private/Framework/MultiBox/SClippingHorizontalBox.cpp
+++ b/Engine/Source/Runtime/Slate/Private/Framework/MultiBox/SClippingHorizontalBox.cpp
@@ -6,22 +6,8 @@
 #include "Widgets/Images/SImage.h"
 #include "Widgets/Input/SComboButton.h"
 #include "Framework/Application/SlateApplication.h"
+#include "Styling/ToolBarStyle.h"
 
-void SClippingHorizontalBox::Tick( const FGeometry& AllottedGeometry, const double InCurrentTime, const float InDeltaTime )
-{
-	const int32 PrevNumClippedChildren = NumClippedChildren;
-	// Call ArrangeChildren so we can cache off the index of the first clipped tab
-	FArrangedChildren ArrangedChildren(EVisibility::Visible);
-	this->ArrangeChildren(AllottedGeometry, ArrangedChildren);
-	ClippedIdx = ArrangedChildren.Num() - 1;
-
-	if (NumClippedChildren != PrevNumClippedChildren && WrapButton->IsOpen())
-	{
-		// Number of children changed while the wrap button was open. An event while the wrap button menu is open has caused the contents of the multibox to change.  The menu can no longer be considered safe
-		// as it most likely lost its connection to the window
-		WrapButton->SetIsOpen(false);
-	}
-}
 
 void SClippingHorizontalBox::OnArrangeChildren( const FGeometry& AllottedGeometry, FArrangedChildren& ArrangedChildren ) const
 {
@@ -29,6 +15,8 @@ void SClippingHorizontalBox::OnArrangeChildren( const FGeometry& AllottedGeometr
 	// been called and this method isn't going to behave properly
 	check(WrapButton.IsValid());
 
+	LastClippedIdx = ClippedIdx;
+
 	NumClippedChildren = 0;
 	SHorizontalBox::OnArrangeChildren(AllottedGeometry, ArrangedChildren);
 
@@ -70,66 +58,18 @@ void SClippingHorizontalBox::OnArrangeChildren( const FGeometry& AllottedGeometr
 			}
 		}
 	}
+
+	ClippedIdx = ArrangedChildren.Num() - 1;
 }
 
 int32 SClippingHorizontalBox::OnPaint( const FPaintArgs& Args, const FGeometry& AllottedGeometry, const FSlateRect& MyCullingRect, FSlateWindowElementList& OutDrawElements, int32 LayerId, const FWidgetStyle& InWidgetStyle, bool bParentEnabled ) const
 {
-	// Get the clipped children info
-	FArrangedChildren ClippedArrangedChildren(EVisibility::Visible);
-	ArrangeChildren(AllottedGeometry, ClippedArrangedChildren);
-	
-	// Get the non-clipped children info
-	// @todo umg: One should not call the virtual OnArrangeChildren, one should only call ArrangeChildren.
-	FArrangedChildren ArrangedChildren(EVisibility::Visible);
-	SBoxPanel::OnArrangeChildren(AllottedGeometry, ArrangedChildren);
-	
-	if ((ClippedArrangedChildren.Num() != 0) && (ArrangedChildren.Num() != 0))
-	{
-		int32 IndexClippedAt = ClippedArrangedChildren.Num() - 1;
-		const FArrangedWidget& LastCippedChild = ClippedArrangedChildren[IndexClippedAt];
-		const FArrangedWidget& FirstChild = ArrangedChildren[0];
-		const FArrangedWidget& LastChild = ArrangedChildren[ArrangedChildren.Num() - 1];
-		float BorderLocalWidth = AllottedGeometry.GetLocalSize().X;
-	
-		// If only the last child/block, which is the wrap button, is being clipped
-		if (IndexClippedAt == ArrangedChildren.Num() - 2)
-		{
-			// Only recalculate the alloted geometry size if said size is fitted to the toolbar/menubar
-			if (FMath::TruncToInt(AllottedGeometry.AbsolutePosition.X + AllottedGeometry.GetLocalSize().X * AllottedGeometry.Scale) <= FMath::TruncToInt(LastChild.Geometry.AbsolutePosition.X + LastChild.Geometry.GetLocalSize().X * LastChild.Geometry.Scale))
-			{
-				// Calculate the size of the custom border
-				BorderLocalWidth = (LastCippedChild.Geometry.AbsolutePosition.X + LastCippedChild.Geometry.GetLocalSize().X * LastCippedChild.Geometry.Scale - FirstChild.Geometry.AbsolutePosition.X) / AllottedGeometry.Scale;
-			}
-		}
-		else
-		{
-			// Children/blocks are being clipped, calculate the size of the custom border
-			const FArrangedWidget& NextChild = (IndexClippedAt + 1 < ClippedArrangedChildren.Num())? ClippedArrangedChildren[IndexClippedAt + 1]: LastCippedChild;
-			BorderLocalWidth = (NextChild.Geometry.AbsolutePosition.X + NextChild.Geometry.GetLocalSize().X * NextChild.Geometry.Scale - FirstChild.Geometry.AbsolutePosition.X) / AllottedGeometry.Scale;
-		}
-	
-		bool bEnabled = ShouldBeEnabled( bParentEnabled );
-		ESlateDrawEffect DrawEffects = !bEnabled ? ESlateDrawEffect::DisabledEffect : ESlateDrawEffect::None;
-		FSlateColor BorderBackgroundColor = FLinearColor::White;
-
-		// Draw the custom border
-		FSlateDrawElement::MakeBox(
-			OutDrawElements,
-			LayerId,
-			AllottedGeometry.ToPaintGeometry(FVector2D(BorderLocalWidth, AllottedGeometry.GetLocalSize().Y), FSlateLayoutTransform()),
-			BackgroundBrush,
-			DrawEffects,
-			BackgroundBrush->GetTint( InWidgetStyle ) * InWidgetStyle.GetColorAndOpacityTint() * BorderBackgroundColor.GetColor(InWidgetStyle)
-			);
-	}
-
 	return SHorizontalBox::OnPaint(Args, AllottedGeometry, MyCullingRect, OutDrawElements, LayerId, InWidgetStyle, bParentEnabled);
 }
 
 FVector2D SClippingHorizontalBox::ComputeDesiredSize(float LayoutScaleMultiplier) const
 {
 	FVector2D Size = SBoxPanel::ComputeDesiredSize(LayoutScaleMultiplier);
-	if (ClippedIdx == (Children.Num() - 2))
 	{
 		// If the wrap button isn't being shown, subtract it's size from the total desired size
 		const SBoxPanel::FSlot& Child = Children[Children.Num() - 1];
@@ -141,27 +81,31 @@ FVector2D SClippingHorizontalBox::ComputeDesiredSize(float LayoutScaleMultiplier
 
 void SClippingHorizontalBox::Construct( const FArguments& InArgs )
 {
-	BackgroundBrush = InArgs._BackgroundBrush;
 	OnWrapButtonClicked = InArgs._OnWrapButtonClicked;
 	StyleSet = InArgs._StyleSet;
 	StyleName = InArgs._StyleName;
+
+	LastClippedIdx = ClippedIdx = INDEX_NONE;
 }
 
 void SClippingHorizontalBox::AddWrapButton()
 {
+	const FToolBarStyle& ToolBarStyle = StyleSet->GetWidgetStyle<FToolBarStyle>(StyleName);
+
 	// Construct the wrap button used in toolbars and menubars
 	WrapButton = 
 		SNew( SComboButton )
 		.HasDownArrow( false )
-		.ButtonStyle( StyleSet, ISlateStyle::Join(StyleName, ".Button") )
+		.ButtonStyle(&ToolBarStyle.ButtonStyle)
 		.ContentPadding( 0 )
 		.ToolTipText( NSLOCTEXT("Slate", "ExpandToolbar", "Click to expand toolbar") )
 		.OnGetMenuContent( OnWrapButtonClicked )
 		.Cursor( EMouseCursor::Default )
+		.OnMenuOpenChanged(this, &SClippingHorizontalBox::OnWrapButtonOpenChanged)
 		.ButtonContent()
 		[
-			SNew( SImage )
-			.Image( StyleSet->GetBrush(StyleName, ".Expand") )
+			SNew(SImage)
+			.Image(&ToolBarStyle.ExpandBrush)
 		];
 
 	// Add the wrap button
@@ -173,3 +117,28 @@ void SClippingHorizontalBox::AddWrapButton()
 	];
 }
 
+void SClippingHorizontalBox::OnWrapButtonOpenChanged(bool bIsOpen)
+{
+	if (bIsOpen && !WrapButtonOpenTimer.IsValid())
+	{
+		WrapButtonOpenTimer = RegisterActiveTimer(0.0f, FWidgetActiveTimerDelegate::CreateSP(this, &SClippingHorizontalBox::UpdateWrapButtonStatus));
+	}
+	else if(!bIsOpen && WrapButtonOpenTimer.IsValid())
+	{
+		UnRegisterActiveTimer(WrapButtonOpenTimer.ToSharedRef());
+		WrapButtonOpenTimer.Reset();
+	}
+}
+
+EActiveTimerReturnType SClippingHorizontalBox::UpdateWrapButtonStatus(double CurrentTime, float DeltaTime)
+{
+	if (LastClippedIdx != ClippedIdx || !WrapButton->IsOpen())
+	{
+		WrapButton->SetIsOpen(false);
+		WrapButtonOpenTimer.Reset();
+		return EActiveTimerReturnType::Stop;
+	}
+
+	return EActiveTimerReturnType::Continue;
+}
+
diff --git a/Engine/Source/Runtime/Slate/Private/Framework/MultiBox/SToolBarButtonBlock.cpp b/Engine/Source/Runtime/Slate/Private/Framework/MultiBox/SToolBarButtonBlock.cpp
--- a/Engine/Source/Runtime/Slate/Private/Framework/MultiBox/SToolBarButtonBlock.cpp
+++ b/Engine/Source/Runtime/Slate/Private/Framework/MultiBox/SToolBarButtonBlock.cpp
@@ -8,6 +8,9 @@
 #include "Widgets/Input/SButton.h"
 #include "Widgets/SToolTip.h"
 #include "Widgets/Input/SCheckBox.h"
+#include "Widgets/Layout/SBox.h"
+#include "Framework/MultiBox/SToolBarComboButtonBlock.h"
+#include "Styling/ToolBarStyle.h"
 
 
 FToolBarButtonBlock::FToolBarButtonBlock( const TSharedPtr< const FUICommandInfo > InCommand, TSharedPtr< const FUICommandList > InCommandList, const TAttribute<FText>& InLabelOverride, const TAttribute<FText>& InToolTipOverride, const TAttribute<FSlateIcon>& InIconOverride )
@@ -93,7 +96,7 @@ void SToolBarButtonBlock::Construct( const FArguments& InArgs )
 	}
 	else
 	{
-		LabelVisibility = TAttribute< EVisibility >::Create( TAttribute< EVisibility >::FGetter::CreateSP( SharedThis( this ), &SToolBarButtonBlock::GetIconVisibility, false ) );
+		LabelVisibility = TAttribute<EVisibility>::Create(TAttribute<EVisibility>::FGetter::CreateSP(SharedThis(this), &SToolBarButtonBlock::GetIconVisibility, false));
 	}
 
 	bIsFocusable = InArgs._IsFocusable;
@@ -107,6 +110,8 @@ void SToolBarButtonBlock::Construct( const FArguments& InArgs )
  */
 void SToolBarButtonBlock::BuildMultiBlockWidget(const ISlateStyle* StyleSet, const FName& StyleName)
 {
+	const FToolBarStyle& ToolBarStyle = StyleSet->GetWidgetStyle<FToolBarStyle>(StyleName);
+
 	struct Local
 	{
 		/** Appends the key binding to the end of the provided ToolTip */
@@ -128,9 +133,9 @@ void SToolBarButtonBlock::BuildMultiBlockWidget(const ISlateStyle* StyleSet, con
 	};
 
 
-	TSharedRef< const FMultiBox > MultiBox( OwnerMultiBoxWidget.Pin()->GetMultiBox() );
-	
-	TSharedRef< const FToolBarButtonBlock > ToolBarButtonBlock = StaticCastSharedRef< const FToolBarButtonBlock >( MultiBlock.ToSharedRef() );
+	TSharedRef<const FMultiBox> MultiBox = OwnerMultiBoxWidget.Pin()->GetMultiBox();
+
+	TSharedRef< const FToolBarButtonBlock > ToolBarButtonBlock = StaticCastSharedRef< const FToolBarButtonBlock >(MultiBlock.ToSharedRef());
 
 	TSharedPtr< const FUICommandInfo > UICommand = ToolBarButtonBlock->GetAction();
 
@@ -164,52 +169,75 @@ void SToolBarButtonBlock::BuildMultiBlockWidget(const ISlateStyle* StyleSet, con
 	ActualToolTip = TAttribute< FText >::Create( TAttribute< FText >::FGetter::CreateStatic( &Local::AppendKeyBindingToToolTip, ActualToolTip, Action ) );
 	
 	// If we were supplied an image than go ahead and use that, otherwise we use a null widget
-	TSharedRef< SWidget > IconWidget =
-		SNew( SImage )
-		.Visibility( this, &SToolBarButtonBlock::GetIconVisibility, false )	
-		.Image( this, &SToolBarButtonBlock::GetIconBrush );
-	TSharedRef< SWidget > SmallIconWidget =
-		SNew( SImage )
-		.Visibility( this, &SToolBarButtonBlock::GetIconVisibility, true )	
-		.Image( this, &SToolBarButtonBlock::GetSmallIconBrush );
+	TSharedRef<SImage> IconWidget =
+		SNew(SImage)
+		.ColorAndOpacity(FSlateColor::UseForeground())
+		.Image(this, &SToolBarButtonBlock::GetIconBrush);
 
-	// Create the content for our button
-	TSharedRef< SWidget > ButtonContent =
 
-		SNew(SHorizontalBox)
-		.AddMetaData<FTagMetaData>(FTagMetaData(TutorialHighlightName))
+	// Create the content for our button
+	TSharedRef<SWidget> ButtonContent = SNullWidget::NullWidget;
+	if (MultiBox->GetType() == EMultiBoxType::SlimHorizontalToolBar)
+	{
 
-		+ SHorizontalBox::Slot()
-		.FillWidth(1)
-		.VAlign(VAlign_Center)
-		[
-			SNew( SVerticalBox )
+		const FVector2D IconSize = ToolBarStyle.IconSize;
 
-			// Icon image
-			+ SVerticalBox::Slot()
-			.AutoHeight()
-			.HAlign( HAlign_Center )	// Center the icon horizontally, so that large labels don't stretch out the artwork
-			[
-				IconWidget
-			]
-			+ SVerticalBox::Slot().AutoHeight()
-			.HAlign( HAlign_Center )
+		ButtonContent =
+			SNew(SHorizontalBox)
+			.AddMetaData<FTagMetaData>(FTagMetaData(TutorialHighlightName))
+			+ SHorizontalBox::Slot()
+			.AutoWidth()
+			.VAlign(VAlign_Center)
+			.HAlign(HAlign_Center)
 			[
-				SmallIconWidget
+				SNew(SBox)
+				.WidthOverride(IconSize.X)
+				.HeightOverride(IconSize.Y)
+				[
+					IconWidget
+				]	
 			]
-
 			// Label text
-			+ SVerticalBox::Slot().AutoHeight()
-			.Padding(StyleSet->GetMargin(ISlateStyle::Join( StyleName, ".Label.Padding" )))
-			.HAlign( HAlign_Center )	// Center the label text horizontally
+			+ SHorizontalBox::Slot()
+			.AutoWidth()
+			.Padding(ToolBarStyle.LabelPadding)
+			.VAlign(VAlign_Center)
 			[
-				SNew( STextBlock )
-					.Visibility( LabelVisibility )
-					.Text( ActualLabel )
-					.TextStyle( StyleSet, ISlateStyle::Join( StyleName, ".Label" ) )	// Smaller font for tool tip labels
-					.ShadowOffset( FVector2D::UnitVector )
-			]
-		];
+				SNew(STextBlock)
+				.Visibility(LabelVisibility)
+				.Text(ActualLabel)
+				.TextStyle(&ToolBarStyle.LabelStyle)	// Smaller font for tool tip labels
+			];
+	}
+	else
+	{
+		ButtonContent =
+			SNew(SHorizontalBox)
+			.AddMetaData<FTagMetaData>(FTagMetaData(TutorialHighlightName))
+			+ SHorizontalBox::Slot()
+			.FillWidth(1)
+			.VAlign(VAlign_Center)
+			[
+				SNew(SVerticalBox)
+				// Icon image
+				+ SVerticalBox::Slot()
+				.AutoHeight()
+				.HAlign(HAlign_Center)	// Center the icon horizontally, so that large labels don't stretch out the artwork
+				[
+					IconWidget
+				]
+				// Label text
+				+ SVerticalBox::Slot().AutoHeight()
+				.Padding(ToolBarStyle.LabelPadding)
+				.HAlign(HAlign_Center)	// Center the label text horizontally
+				[
+					SNew(STextBlock)
+					.Visibility(LabelVisibility)
+					.Text(ActualLabel)
+					.TextStyle(&ToolBarStyle.LabelStyle)	// Smaller font for tool tip labels
+				]
+			];
+		}
 
 	EMultiBlockLocation::Type BlockLocation = GetMultiBlockLocation();
 	
@@ -224,58 +252,68 @@ void SToolBarButtonBlock::BuildMultiBlockWidget(const ISlateStyle* StyleSet, con
 	if( UserInterfaceType == EUserInterfaceActionType::Button )
 	{
 		FName BlockStyle = EMultiBlockLocation::ToName(ISlateStyle::Join( StyleName, ".Button" ), BlockLocation);
+		const FButtonStyle& ButtonStyle = BlockLocation == EMultiBlockLocation::None ? ToolBarStyle.ButtonStyle : StyleSet->GetWidgetStyle<FButtonStyle>(BlockStyle);
 
 		ChildSlot
 		[
 			// Create a button
-			SNew( SButton )
-				.ContentPadding(0 )
-				// Use the tool bar item style for this button
-				.ButtonStyle( StyleSet, BlockStyle )
-				.ForegroundColor( FSlateColor::UseForeground() )
-				.IsFocusable(bIsFocusable)
-				[
-					ButtonContent
-				]
-				// Bind the button's "on clicked" event to our object's method for this
-				.OnClicked( this, &SToolBarButtonBlock::OnClicked )
-				// Pass along the block's tool-tip string
-				.ToolTip( FMultiBoxSettings::ToolTipConstructor.Execute( ActualToolTip, nullptr, Action.Pin() ) )
+			SNew(SButton)
+			.ContentPadding(0)
+			// Use the tool bar item style for this button
+			.ButtonStyle(&ButtonStyle)
+			.OnClicked(this, &SToolBarButtonBlock::OnClicked)
+			.ToolTip(FMultiBoxSettings::ToolTipConstructor.Execute(ActualToolTip, nullptr, Action.Pin()))
+			.IsFocusable(bIsFocusable)
+			[
+				ButtonContent
+			]
 		];
 	}
 	else if( ensure( UserInterfaceType == EUserInterfaceActionType::ToggleButton || UserInterfaceType == EUserInterfaceActionType::RadioButton ) )
 	{
 		FName BlockStyle = EMultiBlockLocation::ToName(ISlateStyle::Join( StyleName, ".ToggleButton" ), BlockLocation);
-		FName CheckboxStyle = ISlateStyle::Join( StyleName, ".SToolBarButtonBlock.CheckBox.Padding" );
+	
+		const FCheckBoxStyle& CheckStyle = BlockLocation == EMultiBlockLocation::None ? ToolBarStyle.ToggleButton : StyleSet->GetWidgetStyle<FCheckBoxStyle>(BlockStyle);
 
 		ChildSlot
 		[
 			// Create a check box
 			SNew( SCheckBox )
+			// Use the tool bar style for this check box
+			.Style(&CheckStyle)
+			.IsFocusable(bIsFocusable)
+			.ToolTip( FMultiBoxSettings::ToolTipConstructor.Execute( ActualToolTip, nullptr, Action.Pin()))		
+			.OnCheckStateChanged(this, &SToolBarButtonBlock::OnCheckStateChanged )
+			.IsChecked(this, &SToolBarButtonBlock::OnIsChecked)
+			.Padding(ToolBarStyle.CheckBoxPadding)
+			[
+				ButtonContent
+			]
+		];
+	}
 
-				// Use the tool bar style for this check box
-				.Style( StyleSet, BlockStyle )
-
-				// User will have set the focusable attribute for the block, honor it
-				.IsFocusable( bIsFocusable )
-
-				// Pass along the block's tool-tip string
-				.ToolTip( FMultiBoxSettings::ToolTipConstructor.Execute( ActualToolTip, nullptr, Action.Pin() ) )
-				[
-					ButtonContent
-				]
-
-				// Bind the button's "on checked" event to our object's method for this
-				.OnCheckStateChanged( this, &SToolBarButtonBlock::OnCheckStateChanged )
 
-				// Bind the check box's "checked" state to our user interface action
-				.IsChecked( this, &SToolBarButtonBlock::OnIsChecked )
+	// Often buttons have a "simple" combo box next to it.  The button + simple combo button is designed to feel like a normal combo button but 
+	// when the button part is pressed some action happens independently of the combo dropdown and the dropdown changes settings the button uses.  
+	// We want this to feel like one widget so we space them closer together.
+	FMargin Padding = ToolBarStyle.ButtonPadding;
+	{
+		int32 MyIndex = MultiBox->GetBlocks().Find(ToolBarButtonBlock);
+		if (MultiBox->GetBlocks().IsValidIndex(MyIndex + 1))
+		{
+			const TSharedRef<const FMultiBlock>& NextBlock = MultiBox->GetBlocks()[MyIndex + 1];
+			if (NextBlock->GetType() == EMultiBlockType::ToolBarComboButton)
+			{
+				TSharedRef<const FToolBarComboButtonBlock> NextToolBarComboButtonBlock = StaticCastSharedRef<const FToolBarComboButtonBlock>(NextBlock);
+				if (NextToolBarComboButtonBlock->IsSimpleComboBox())
+				{
+					Padding.Right = 2.0f;
+				}
+			}
+		}
 
-				.Padding( StyleSet->GetMargin(CheckboxStyle) )
-		];
+		ChildSlot.Padding(Padding);
 	}
-	
-	ChildSlot.Padding(StyleSet->GetMargin(ISlateStyle::Join( StyleName, ".SToolBarButtonBlock.Padding" )));
 
 	// Bind our widget's enabled state to whether or not our action can execute
 	SetEnabled( TAttribute< bool >( this, &SToolBarButtonBlock::IsEnabled ) );
@@ -407,22 +445,29 @@ EVisibility SToolBarButtonBlock::GetIconVisibility(bool bIsASmallIcon) const
 
 const FSlateBrush* SToolBarButtonBlock::GetIconBrush() const
 {
-	TSharedRef< const FToolBarButtonBlock > ToolBarButtonBlock = StaticCastSharedRef< const FToolBarButtonBlock >( MultiBlock.ToSharedRef() );
+	return bForceSmallIcons || FMultiBoxSettings::UseSmallToolBarIcons.Get() ? GetSmallIconBrush() : GetNormalIconBrush();
+}
+
+const FSlateBrush* SToolBarButtonBlock::GetNormalIconBrush() const
+{
+	TSharedRef< const FToolBarButtonBlock > ToolBarButtonBlock = StaticCastSharedRef< const FToolBarButtonBlock >(MultiBlock.ToSharedRef());
 
 	const FSlateIcon ActionIcon = ToolBarButtonBlock->GetAction().IsValid() ? ToolBarButtonBlock->GetAction()->GetIcon() : FSlateIcon();
 	const FSlateIcon& ActualIcon = ToolBarButtonBlock->IconOverride.IsSet() ? ToolBarButtonBlock->IconOverride.Get() : ActionIcon;
-	
-	if( ActualIcon.IsSet() )
+
+	if (ActualIcon.IsSet())
 	{
 		return ActualIcon.GetIcon();
 	}
 	else
 	{
-		check( OwnerMultiBoxWidget.IsValid() );
+		check(OwnerMultiBoxWidget.IsValid());
 
 		TSharedPtr<SMultiBoxWidget> MultiBoxWidget = OwnerMultiBoxWidget.Pin();
 		const ISlateStyle* const StyleSet = MultiBoxWidget->GetStyleSet();
-		return StyleSet->GetBrush( "MultiBox.GenericToolBarIcon" );
+
+		static const FName IconName("MultiBox.GenericToolBarIcon");
+		return StyleSet->GetBrush(IconName);
 	}
 }
 
@@ -443,6 +488,8 @@ const FSlateBrush* SToolBarButtonBlock::GetSmallIconBrush() const
 
 		TSharedPtr<SMultiBoxWidget> MultiBoxWidget = OwnerMultiBoxWidget.Pin();
 		const ISlateStyle* const StyleSet = MultiBoxWidget->GetStyleSet();
-		return StyleSet->GetBrush( "MultiBox.GenericToolBarIcon.Small" );
+
+		static const FName IconName("MultiBox.GenericToolBarIcon.Small" );
+		return StyleSet->GetBrush(IconName);
 	}
 }
diff --git a/Engine/Source/Runtime/Slate/Private/Framework/MultiBox/SToolBarComboButtonBlock.cpp b/Engine/Source/Runtime/Slate/Private/Framework/MultiBox/SToolBarComboButtonBlock.cpp
--- a/Engine/Source/Runtime/Slate/Private/Framework/MultiBox/SToolBarComboButtonBlock.cpp
+++ b/Engine/Source/Runtime/Slate/Private/Framework/MultiBox/SToolBarComboButtonBlock.cpp
@@ -5,6 +5,8 @@
 #include "Widgets/Images/SImage.h"
 #include "Framework/MultiBox/MultiBoxBuilder.h"
 #include "Widgets/Input/SComboButton.h"
+#include "Widgets/Layout/SBox.h"
+#include "Styling/ToolBarStyle.h"
 
 
 FToolBarComboButtonBlock::FToolBarComboButtonBlock( const FUIAction& InAction, const FOnGetContent& InMenuContentGenerator, const TAttribute<FText>& InLabel, const TAttribute<FText>& InToolTip, const TAttribute<FSlateIcon>& InIcon, bool bInSimpleComboBox )
@@ -67,46 +69,35 @@ void SToolBarComboButtonBlock::BuildMultiBlockWidget(const ISlateStyle* StyleSet
 	
 	TSharedRef< const FToolBarComboButtonBlock > ToolBarComboButtonBlock = StaticCastSharedRef< const FToolBarComboButtonBlock >( MultiBlock.ToSharedRef() );
 
+	//TSharedPtr< const FUICommandInfo > UICommand = ToolBarComboButtonBlock->GetAction();
+
 	TAttribute<FText> Label;
 
-	// If we were supplied an image than go ahead and use that, otherwise we use a null widget
-	TSharedRef< SWidget > IconWidget = SNullWidget::NullWidget;
-	TSharedRef< SWidget > SmallIconWidget = SNullWidget::NullWidget;
+	const FToolBarStyle& ToolBarStyle = StyleSet->GetWidgetStyle<FToolBarStyle>(StyleName);
+
+	TSharedRef<SWidget> IconWidget = SNullWidget::NullWidget;
 	if (!ToolBarComboButtonBlock->bSimpleComboBox)
 	{
-		if( !HasDynamicIcon() )
+		if (MultiBox->GetType() == EMultiBoxType::SlimHorizontalToolBar)
 		{
-			// Not dynamic, so resolve the icons now
-			const FSlateBrush* IconBrush = GetIcon();
-			const FSlateBrush* SmallIconBrush = GetSmallIcon();
-
-			if( IconBrush->GetResourceName() != NAME_None )
-			{
-				IconWidget =
-					SNew( SImage )
-					.Visibility( this, &SToolBarComboButtonBlock::GetIconVisibility, false )
-					.Image( IconBrush );
-			}
-			if( SmallIconBrush->GetResourceName() != NAME_None )
-			{
-				SmallIconWidget =
-					SNew( SImage )
-					.Visibility( this, &SToolBarComboButtonBlock::GetIconVisibility, true )
-					.Image( SmallIconBrush );
-			}
+			const FVector2D IconSize = ToolBarStyle.IconSize;
+
+			IconWidget =
+				SNew(SBox)
+				.WidthOverride(IconSize.X)
+				.HeightOverride(IconSize.Y)
+				[
+					SNew(SImage)
+					.ColorAndOpacity(FSlateColor::UseForeground())
+					.Image(this, &SToolBarComboButtonBlock::GetIconBrush)
+				];
 		}
-		else 
+		else
 		{
-			// Dynamic, so we need to preserve the callback used
-			IconWidget =
-				SNew( SImage )
-				.Visibility( this, &SToolBarComboButtonBlock::GetIconVisibility, false )
-				.Image( this, &SToolBarComboButtonBlock::GetIcon );
-
-			SmallIconWidget =
-				SNew( SImage )
-				.Visibility( this, &SToolBarComboButtonBlock::GetIconVisibility, true )
-				.Image( this, &SToolBarComboButtonBlock::GetSmallIcon );
+			IconWidget = 
+				SNew(SImage)
+				.ColorAndOpacity(FSlateColor::UseForeground())
+				.Image(this, &SToolBarComboButtonBlock::GetIconBrush);
 		}
 
 		Label = ToolBarComboButtonBlock->Label;
@@ -114,7 +105,9 @@ void SToolBarComboButtonBlock::BuildMultiBlockWidget(const ISlateStyle* StyleSet
 
 	// Add this widget to the search list of the multibox
 	if (MultiBlock->GetSearchable())
+	{
 		OwnerMultiBoxWidget.Pin()->AddSearchElement(this->AsWidget(), Label.Get());
+	}
 
 	// Setup the string for the metatag
 	FName TagName;
@@ -128,66 +121,96 @@ void SToolBarComboButtonBlock::BuildMultiBlockWidget(const ISlateStyle* StyleSet
 	}
 	
 	// Create the content for our button
-	TSharedRef< SWidget > ButtonContent =
-		SNew( SVerticalBox )
-		.AddMetaData<FTagMetaData>(FTagMetaData(TagName))
-		// Icon image
-		+ SVerticalBox::Slot()
-		.AutoHeight()
-		.HAlign(HAlign_Center)
-		// Center the icon horizontally, so that large labels don't stretch out the artwork
-		[
-			IconWidget
-		]
-		+ SVerticalBox::Slot()
-		.AutoHeight()
-		[
-			SmallIconWidget
-		]
-
-		// Label text
-		+ SVerticalBox::Slot()
-		.AutoHeight()
-		.Padding(StyleSet->GetMargin(ISlateStyle::Join( StyleName, ".Label.Padding" )))
-		.HAlign( HAlign_Center )	// Center the label text horizontally
-		[
-			SNew( STextBlock )
-				.Visibility( LabelVisibility )
-				.Text( Label )
-
+	TSharedRef<SWidget> ButtonContent = SNullWidget::NullWidget;
+	if (MultiBox->GetType() == EMultiBoxType::SlimHorizontalToolBar)
+	{
+		ButtonContent =
+			SNew(SHorizontalBox)
+			// Icon image
+			+ SHorizontalBox::Slot()
+			.AutoWidth()
+			.VAlign(VAlign_Center)
+			.HAlign(HAlign_Center)
+			[
+				IconWidget
+			]
+			// Label text
+			+ SHorizontalBox::Slot()
+			.AutoWidth()
+			.Padding(ToolBarStyle.LabelPadding)
+			.VAlign(VAlign_Center)	// Center the label text horizontally
+			[
+				SNew(STextBlock)
+				.Visibility(ToolBarComboButtonBlock->bSimpleComboBox ? EVisibility::Collapsed : LabelVisibility)
+				.Text(Label)
 				// Smaller font for tool tip labels
-				.TextStyle( StyleSet, ISlateStyle::Join( StyleName, ".Label" ) )
-				.ShadowOffset(FVector2D::UnitVector)
-		]
-		;
+				.TextStyle(&ToolBarStyle.LabelStyle)
+			];
+
+	}
+	else
+	{
+		ButtonContent =
+			SNew(SVerticalBox)
+			// Icon image
+			+ SVerticalBox::Slot()
+			.AutoHeight()
+			.HAlign(HAlign_Center)	// Center the icon horizontally, so that large labels don't stretch out the artwork
+			[
+				IconWidget
+			]
+			// Label text
+			+ SVerticalBox::Slot().AutoHeight()
+			.Padding(ToolBarStyle.LabelPadding)
+			.AutoHeight()
+			.HAlign(HAlign_Center)	// Center the label text horizontally
+			[
+				SNew(STextBlock)
+				.Visibility(LabelVisibility)
+				.Text(Label)
+				.TextStyle(&ToolBarStyle.LabelStyle)
+			];
+	}
+	
 		
 	EMultiBlockLocation::Type BlockLocation = GetMultiBlockLocation();
-	FName BlockStyle = EMultiBlockLocation::ToName(ISlateStyle::Join( StyleName, ".Button" ), BlockLocation);
-	FName ColorStyle = ISlateStyle::Join( StyleName, ".SToolBarComboButtonBlock.ComboButton.Color" );
-
-	ChildSlot
-	[
-		SNew( SComboButton )
-			.ContentPadding(0)
+	FName BlockStyle = EMultiBlockLocation::ToName(ISlateStyle::Join(StyleName, ".Button"), BlockLocation);
+	const FButtonStyle& ButtonStyle = BlockLocation == EMultiBlockLocation::None ? ToolBarStyle.ButtonStyle : StyleSet->GetWidgetStyle<FButtonStyle>(BlockStyle);
 
-			// Use the tool bar item style for this button
-			.ButtonStyle( StyleSet, BlockStyle )
+	OpenForegroundColor = ButtonStyle.HoveredForeground;
 
-			// Pass along the block's tool-tip string
-			.ToolTipText( ToolBarComboButtonBlock->ToolTip )
+	const FComboButtonStyle* ComboStyle = &FAppStyle::Get().GetWidgetStyle<FComboButtonStyle>("ComboButton");
+	if (ToolBarComboButtonBlock->bSimpleComboBox)
+	{
+		ComboStyle = &ToolBarStyle.SettingsComboButton;
+	}
 
-			.ForegroundColor( StyleSet->GetSlateColor( ColorStyle ) )
+	ChildSlot
+	[
+		SAssignNew(ComboButtonWidget, SComboButton)
+		.AddMetaData<FTagMetaData>(FTagMetaData(TagName))
+		.ContentPadding(0)
+		.ComboButtonStyle(ComboStyle)
+		.ButtonStyle(&ButtonStyle)
+		.ToolTipText(ToolBarComboButtonBlock->ToolTip)
+		.ForegroundColor(this, &SToolBarComboButtonBlock::OnGetForegroundColor)
+		// Route the content generator event
+		.OnGetMenuContent(this, &SToolBarComboButtonBlock::OnGetMenuContent)
+		.ButtonContent()
+		[
+			ButtonContent
+		]
+	];
 
-			.ButtonContent()
-			[
-				ButtonContent
-			]
 
-			// Route the content generator event
-			.OnGetMenuContent( this, &SToolBarComboButtonBlock::OnGetMenuContent )
-	];
+	FMargin Padding = ToolBarStyle.ComboButtonPadding;
+	if (ToolBarComboButtonBlock->bSimpleComboBox)
+	{
+		Padding.Left = 0;
+		Padding.Right = 10;
+	}
 
-	ChildSlot.Padding(StyleSet->GetMargin(ISlateStyle::Join( StyleName, ".SToolBarComboButtonBlock.Padding" )));
+	ChildSlot.Padding(Padding);
 	// Bind our widget's enabled state to whether or not our action can execute
 	SetEnabled( TAttribute< bool >( this, &SToolBarComboButtonBlock::IsEnabled ) );
 
@@ -228,13 +251,18 @@ bool SToolBarComboButtonBlock::HasDynamicIcon() const
 	return Icon.IsBound();
 }
 
-const FSlateBrush* SToolBarComboButtonBlock::GetIcon() const
+const FSlateBrush* SToolBarComboButtonBlock::GetIconBrush() const
+{
+	return bForceSmallIcons || FMultiBoxSettings::UseSmallToolBarIcons.Get() ? GetSmallIconBrush() : GetNormalIconBrush();
+}
+
+const FSlateBrush* SToolBarComboButtonBlock::GetNormalIconBrush() const
 {
 	const FSlateIcon& ActualIcon = Icon.Get();
 	return ActualIcon.GetIcon();
 }
 
-const FSlateBrush* SToolBarComboButtonBlock::GetSmallIcon() const
+const FSlateBrush* SToolBarComboButtonBlock::GetSmallIconBrush() const
 {
 	const FSlateIcon& ActualIcon = Icon.Get();
 	return ActualIcon.GetSmallIcon();
@@ -244,3 +272,15 @@ EVisibility SToolBarComboButtonBlock::GetIconVisibility(bool bIsASmallIcon) cons
 {
 	return ((bForceSmallIcons || FMultiBoxSettings::UseSmallToolBarIcons.Get()) ^ bIsASmallIcon) ? EVisibility::Collapsed : EVisibility::Visible;
 }
+
+FSlateColor SToolBarComboButtonBlock::OnGetForegroundColor() const
+{
+	if (ComboButtonWidget->IsOpen())
+	{
+		return OpenForegroundColor;
+	}
+	else
+	{
+		return FSlateColor::UseStyle();
+	}
+}
diff --git a/Engine/Source/Runtime/Slate/Private/Framework/MultiBox/SToolBarSeparatorBlock.cpp b/Engine/Source/Runtime/Slate/Private/Framework/MultiBox/SToolBarSeparatorBlock.cpp
--- a/Engine/Source/Runtime/Slate/Private/Framework/MultiBox/SToolBarSeparatorBlock.cpp
+++ b/Engine/Source/Runtime/Slate/Private/Framework/MultiBox/SToolBarSeparatorBlock.cpp
@@ -4,6 +4,7 @@
 #include "Widgets/SBoxPanel.h"
 #include "Widgets/Layout/SSeparator.h"
 #include "Framework/MultiBox/MultiBoxBuilder.h"
+#include "Styling/ToolBarStyle.h"
 
 
 /**
@@ -50,17 +51,19 @@ void SToolBarSeparatorBlock::Construct( const FArguments& InArgs )
  */
 void SToolBarSeparatorBlock::BuildMultiBlockWidget(const ISlateStyle* StyleSet, const FName& StyleName)
 {
+	const FToolBarStyle& ToolBarStyle = StyleSet->GetWidgetStyle<FToolBarStyle>(StyleName);
+	
 	ChildSlot
 	[
 		SNew( SHorizontalBox )
 		+SHorizontalBox::Slot()
 		.AutoWidth()
-		.Padding( StyleSet->GetMargin( ISlateStyle::Join( StyleName, ".Separator.Padding" ) ) )
+		.Padding(ToolBarStyle.SeparatorPadding)
 		[
 			SNew(SSeparator)
 				.Orientation(Orient_Vertical)
 				.Thickness(2.0f)
-				.SeparatorImage( StyleSet->GetBrush( ISlateStyle::Join( StyleName, ".Separator" ) ) )
+				.SeparatorImage(&ToolBarStyle.SeparatorBrush)
 		]
 	];
 
diff --git a/Engine/Source/Runtime/Slate/Private/Framework/MultiBox/SUniformToolbarPanel.cpp b/Engine/Source/Runtime/Slate/Private/Framework/MultiBox/SUniformToolbarPanel.cpp
--- a/Engine/Source/Runtime/Slate/Private/Framework/MultiBox/SUniformToolbarPanel.cpp
+++ b/Engine/Source/Runtime/Slate/Private/Framework/MultiBox/SUniformToolbarPanel.cpp
@@ -4,6 +4,7 @@
 #include "Layout/LayoutUtils.h"
 #include "Widgets/Input/SComboButton.h"
 #include "Widgets/Images/SImage.h"
+#include "Styling/ToolBarStyle.h"
 
 SUniformToolbarPanel::SUniformToolbarPanel()
 : Children(this)
@@ -31,6 +32,8 @@ void SUniformToolbarPanel::Construct( const FArguments& InArgs )
 		Children.Add( ChildSlot );
 	}
 
+	const FToolBarStyle& ToolBarStyle = StyleSet->GetWidgetStyle<FToolBarStyle>(StyleName);
+	
 	// Add the optional dropdown arrow as a child.  
 	FSlot& NewSlot = *(new FSlot());
 
@@ -38,7 +41,7 @@ void SUniformToolbarPanel::Construct( const FArguments& InArgs )
 	[
 		SAssignNew(Dropdown, SComboButton)
 		.HasDownArrow(false)
-		.ButtonStyle(StyleSet, ISlateStyle::Join(StyleName, ".Button"))
+		.ButtonStyle(&ToolBarStyle.ButtonStyle)
 		.ContentPadding(0)
 		.ToolTipText(NSLOCTEXT("Slate", "ExpandToolbar", "Click to expand toolbar"))
 		.OnGetMenuContent(InArgs._OnDropdownOpened)
@@ -46,7 +49,7 @@ void SUniformToolbarPanel::Construct( const FArguments& InArgs )
 		.ButtonContent()
 		[
 			SNew(SImage)
-			.Image(StyleSet->GetBrush(StyleName, ".Expand"))
+			.Image(&ToolBarStyle.ExpandBrush)
 		]
 	];
 
diff --git a/Engine/Source/Runtime/Slate/Private/Framework/MultiBox/SWidgetBlock.cpp b/Engine/Source/Runtime/Slate/Private/Framework/MultiBox/SWidgetBlock.cpp
--- a/Engine/Source/Runtime/Slate/Private/Framework/MultiBox/SWidgetBlock.cpp
+++ b/Engine/Source/Runtime/Slate/Private/Framework/MultiBox/SWidgetBlock.cpp
@@ -3,6 +3,7 @@
 #include "Framework/MultiBox/SWidgetBlock.h"
 #include "Widgets/SBoxPanel.h"
 #include "Framework/MultiBox/MultiBoxBuilder.h"
+#include "Styling/ToolBarStyle.h"
 
 
 /**
@@ -56,17 +57,43 @@ void SWidgetBlock::BuildMultiBlockWidget(const ISlateStyle* StyleSet, const FNam
 {
 	TSharedRef< const FWidgetBlock > WidgetBlock = StaticCastSharedRef< const FWidgetBlock >( MultiBlock.ToSharedRef() );
 
+	// Support menus which do not have a defined widget style yet
 	bool bHasLabel = !WidgetBlock->Label.IsEmpty();
-	FMargin Padding = WidgetBlock->bNoIndent ? StyleSet->GetMargin( StyleName, ".Block.Padding" ) : StyleSet->GetMargin( StyleName, ".Block.IndentedPadding" );
+	FMargin Padding;
+	const FTextBlockStyle* LabelStyle = nullptr;
+	if(StyleSet->HasWidgetStyle<FToolBarStyle>(StyleName))
+	{
+		const FToolBarStyle& ToolBarStyle = StyleSet->GetWidgetStyle<FToolBarStyle>(StyleName);
+
+		Padding = WidgetBlock->bNoIndent ? ToolBarStyle.BlockPadding : ToolBarStyle.IndentedBlockPadding;
+		LabelStyle = &ToolBarStyle.LabelStyle;
+
+	}
+	else
+	{
+		Padding = WidgetBlock->bNoIndent ? StyleSet->GetMargin(StyleName, ".Block.Padding") : StyleSet->GetMargin(StyleName, ".Block.IndentedPadding");
+
+		LabelStyle = &StyleSet->GetWidgetStyle<FTextBlockStyle>(ISlateStyle::Join(StyleName, ".Label"));
+	}
+
+	TSharedPtr<SMultiBoxWidget> OwnerMultiBoxWidgetPinned = OwnerMultiBoxWidget.Pin();
+
+	if(OwnerMultiBoxWidgetPinned->GetMultiBox()->GetType() == EMultiBoxType::Menu)
+	{
+		// Account for checkmark used in other menu blocks but not used in for widget rows
+		Padding = Padding + FMargin(12, 0, 8, 0);
+	}
 
 	// Add this widget to the search list of the multibox
 	if (MultiBlock->GetSearchable())
-		OwnerMultiBoxWidget.Pin()->AddSearchElement(this->AsWidget(), WidgetBlock->Label);
+	{
+		OwnerMultiBoxWidgetPinned->AddSearchElement(this->AsWidget(), WidgetBlock->Label);
+	}
 
 	// This widget holds the search text, set it as the search block widget
-	if (OwnerMultiBoxWidget.Pin()->GetSearchTextWidget() == WidgetBlock->ContentWidget)
+	if (OwnerMultiBoxWidgetPinned->GetSearchTextWidget() == WidgetBlock->ContentWidget)
 	{
-		OwnerMultiBoxWidget.Pin()->SetSearchBlockWidget(this->AsWidget());
+		OwnerMultiBoxWidgetPinned->SetSearchBlockWidget(this->AsWidget());
 		this->AsWidget()->SetVisibility(EVisibility::Collapsed);
 	}
 
@@ -85,8 +112,8 @@ void SWidgetBlock::BuildMultiBlockWidget(const ISlateStyle* StyleSet, const FNam
 			.VAlign( VAlign_Center )
 			[
 				SNew( STextBlock )
-				.TextStyle( StyleSet, ISlateStyle::Join( StyleName, ".Label" ) )
-				.Text( WidgetBlock->Label )
+				.TextStyle(LabelStyle)
+				.Text(WidgetBlock->Label)
 			]
 		]
 		+SHorizontalBox::Slot()
diff --git a/Engine/Source/Runtime/Slate/Public/Framework/MultiBox/MultiBoxBuilder.h b/Engine/Source/Runtime/Slate/Public/Framework/MultiBox/MultiBoxBuilder.h
--- a/Engine/Source/Runtime/Slate/Public/Framework/MultiBox/MultiBoxBuilder.h
+++ b/Engine/Source/Runtime/Slate/Public/Framework/MultiBox/MultiBoxBuilder.h
@@ -362,6 +362,7 @@ public:
 	FMenuBarBuilder( TSharedPtr< const FUICommandList > InCommandList, TSharedPtr<FExtender> InExtender = TSharedPtr<FExtender>(), const ISlateStyle* InStyleSet = &FCoreStyle::Get(), FName InMenuName = NAME_None)
 		: FBaseMenuBuilder( EMultiBoxType::MenuBar, false, InCommandList, false, InExtender, InStyleSet, NAME_None, InMenuName )
 	{
+		MultiBox->SetStyle(InStyleSet, "WindowMenuBar");
 	}
 
 
@@ -386,16 +387,25 @@ protected:
  */
 class SLATE_API FToolBarBuilder : public FMultiBoxBuilder
 {
-
+	friend class UToolMenus;
 public:
 
+	UE_DEPRECATED(4.26, "FToolBarBuilder constructor that takes in an EOrientation is deprecated.  Use one of the specific per-type FToolbarBuilder overrides instead.")
+	FToolBarBuilder(TSharedPtr< const FUICommandList > InCommandList, FMultiBoxCustomization InCustomization, TSharedPtr<FExtender> InExtender, EOrientation Orientation, const bool InForceSmallIcons = false, const bool bUniform = false)
+		: FMultiBoxBuilder(bUniform ? EMultiBoxType::UniformToolBar : (Orientation == Orient_Horizontal) ? EMultiBoxType::ToolBar : EMultiBoxType::VerticalToolBar, InCustomization, false, InCommandList, InExtender)
+		, bSectionNeedsToBeApplied(false)
+		, bIsFocusable(false)
+		, bForceSmallIcons(InForceSmallIcons)
+	{
+	}
+
 	/**
 	 * Constructor
 	 *
 	 * @param	InCommandList	The action list that maps command infos to delegates that should be called for each command associated with a multiblock widget
 	 */
-	FToolBarBuilder( TSharedPtr< const FUICommandList > InCommandList, FMultiBoxCustomization InCustomization, TSharedPtr<FExtender> InExtender = TSharedPtr<FExtender>(), EOrientation Orientation = Orient_Horizontal, const bool InForceSmallIcons = false, const bool bUniform = false)
-		: FMultiBoxBuilder(bUniform ? EMultiBoxType::UniformToolBar : (Orientation == Orient_Horizontal) ? EMultiBoxType::ToolBar : EMultiBoxType::VerticalToolBar, InCustomization, false, InCommandList, InExtender )
+	FToolBarBuilder(TSharedPtr<const FUICommandList> InCommandList, FMultiBoxCustomization InCustomization, TSharedPtr<FExtender> InExtender = nullptr, const bool InForceSmallIcons = false)
+		: FMultiBoxBuilder(EMultiBoxType::ToolBar, InCustomization, false, InCommandList, InExtender)
 		, bSectionNeedsToBeApplied(false)
 		, bIsFocusable(false)
 		, bForceSmallIcons(InForceSmallIcons)
@@ -490,7 +500,24 @@ public:
 	 */
 	void EndBlockGroup();
 
+	/**
+	 * Overrides the style being used by this toolbar with a different one for the
+	 * The override will be used for any added blocks until EndStyleOverride is called
+	 */
+	void BeginStyleOverride(FName StyleOverrideName);
+
+	void EndStyleOverride();
 protected:
+
+	FToolBarBuilder(EMultiBoxType InType, TSharedPtr<const FUICommandList> InCommandList, FMultiBoxCustomization InCustomization, TSharedPtr<FExtender> InExtender = TSharedPtr<FExtender>(), const bool InForceSmallIcons = false)
+		: FMultiBoxBuilder(InType, InCustomization, false, InCommandList, InExtender)
+		, bSectionNeedsToBeApplied(false)
+		, bIsFocusable(false)
+		, bForceSmallIcons(InForceSmallIcons)
+	{
+	}
+
+
 	/** FMultiBoxBuilder interface */
 	virtual void ApplyHook(FName InExtensionHook, EExtensionHook::Position HookPosition) override;
 	virtual void ApplySectionBeginning() override;
@@ -499,6 +526,8 @@ private:
 	/** Current extension hook name for sections to determine where sections begin and end */
 	FName CurrentSectionExtensionHook;
 
+	FName CurrentStyleOverride;
+
 	TOptional< EVisibility > LabelVisibility;
 
 	/** True if there is a pending section that needs to be applied */
@@ -511,6 +540,49 @@ private:
 	bool bForceSmallIcons;
 };
 
+
+class SLATE_API FVerticalToolBarBuilder : public FToolBarBuilder
+{
+public:
+	/**
+	 * Constructor
+	 *
+	 * @param	InCommandList	The action list that maps command infos to delegates that should be called for each command associated with a multiblock widget
+	 */
+	FVerticalToolBarBuilder(TSharedPtr<const FUICommandList> InCommandList, FMultiBoxCustomization InCustomization, TSharedPtr<FExtender> InExtender = nullptr, const bool InForceSmallIcons = false)
+		: FToolBarBuilder(EMultiBoxType::VerticalToolBar, InCommandList, InCustomization, InExtender, InForceSmallIcons)
+	{
+	}
+};
+
+class SLATE_API FUniformToolBarBuilder : public FToolBarBuilder
+{
+public:
+	/**
+	 * Constructor
+	 *
+	 * @param	InCommandList	The action list that maps command infos to delegates that should be called for each command associated with a multiblock widget
+	 */
+	FUniformToolBarBuilder(TSharedPtr<const FUICommandList> InCommandList, FMultiBoxCustomization InCustomization, TSharedPtr<FExtender> InExtender = nullptr, const bool InForceSmallIcons = false)
+		: FToolBarBuilder(EMultiBoxType::UniformToolBar, InCommandList, InCustomization, InExtender, InForceSmallIcons)
+	{
+	}
+};
+
+class SLATE_API FSlimHorizontalToolBarBuilder : public FToolBarBuilder
+{
+public:
+	/**
+	 * Constructor
+	 *
+	 * @param	InCommandList	The action list that maps command infos to delegates that should be called for each command associated with a multiblock widget
+	 */
+	FSlimHorizontalToolBarBuilder(TSharedPtr<const FUICommandList> InCommandList, FMultiBoxCustomization InCustomization, TSharedPtr<FExtender> InExtender = nullptr, const bool InForceSmallIcons = false)
+		: FToolBarBuilder(EMultiBoxType::SlimHorizontalToolBar, InCommandList, InCustomization, InExtender, InForceSmallIcons)
+	{
+	}
+};
+
 /**
  * Button grid builder
  */
@@ -522,8 +594,8 @@ public:
 	 *
 	 * @param	InCommandList	The action list that maps command infos to delegates that should be called for each command associated with a multiblock widget
 	 */
-	FButtonRowBuilder( TSharedPtr< const FUICommandList > InCommandList, TSharedPtr<FExtender> InExtender = TSharedPtr<FExtender>() )
-		: FMultiBoxBuilder( EMultiBoxType::ButtonRow, FMultiBoxCustomization::None, false, InCommandList, InExtender )
+	FButtonRowBuilder(TSharedPtr< const FUICommandList > InCommandList, TSharedPtr<FExtender> InExtender = TSharedPtr<FExtender>())
+		: FMultiBoxBuilder(EMultiBoxType::ButtonRow, FMultiBoxCustomization::None, false, InCommandList, InExtender)
 	{
 	}
 
@@ -536,7 +608,7 @@ public:
 	 * @param	InToolTipOverride		Optional tool tip override.	 If omitted, then the action's label will be used instead.
 	 * @param	InIconOverride			Optional icon to use for the tool bar image.  If omitted, then the action's icon will be used instead.
 	 */
-	void AddButton( const TSharedPtr< const FUICommandInfo > InCommand, const TAttribute<FText>& InLabelOverride = TAttribute<FText>(), const TAttribute<FText>& InToolTipOverride = TAttribute<FText>(), const FSlateIcon& InIconOverride = FSlateIcon() );
+	void AddButton(const TSharedPtr< const FUICommandInfo > InCommand, const TAttribute<FText>& InLabelOverride = TAttribute<FText>(), const TAttribute<FText>& InToolTipOverride = TAttribute<FText>(), const FSlateIcon& InIconOverride = FSlateIcon());
 
 	/**
 	 * Adds a button to a row
@@ -547,8 +619,8 @@ public:
 	 * @param	InIcon					The icon for the button
 	 * @param	UserInterfaceActionType	The style of button to display
 	 */
-	void AddButton( const FText& InLabel, const FText& InToolTip, const FUIAction& UIAction, const FSlateIcon& InIcon = FSlateIcon(), const EUserInterfaceActionType UserInterfaceActionType = EUserInterfaceActionType::Button );
-	
+	void AddButton(const FText& InLabel, const FText& InToolTip, const FUIAction& UIAction, const FSlateIcon& InIcon = FSlateIcon(), const EUserInterfaceActionType UserInterfaceActionType = EUserInterfaceActionType::Button);
+
 protected:
 	/** FMultiBoxBuilder interface */
 	virtual void ApplyHook(FName InExtensionHook, EExtensionHook::Position HookPosition) override {}
diff --git a/Engine/Source/Runtime/SlateCore/Private/Styling/SlateStyleSet.cpp b/Engine/Source/Runtime/SlateCore/Private/Styling/SlateStyleSet.cpp
--- a/Engine/Source/Runtime/SlateCore/Private/Styling/SlateStyleSet.cpp
+++ b/Engine/Source/Runtime/SlateCore/Private/Styling/SlateStyleSet.cpp
@@ -12,6 +12,7 @@
 #include "Styling/StyleDefaults.h"
 #include "Styling/ISlateStyle.h"
 #include "Styling/SlateStyle.h"
+#include "Styling/SlateStyleRegistry.h"
 
 
 FSlateStyleSet::FSlateStyleSet(const FName& InStyleSetName)
@@ -34,7 +35,6 @@ FSlateStyleSet::~FSlateStyleSet()
 	}
 }
 
-
 const FName& FSlateStyleSet::GetStyleSetName() const
 {
 	return StyleSetName;
@@ -164,96 +164,131 @@ FString FSlateStyleSet::RootToCoreContentDir(const FString& RelativePath)
 	return ( CoreContentRootDir / RelativePath );
 }
 
-float FSlateStyleSet::GetFloat(const FName PropertyName, const ANSICHAR* Specifier, float DefaultValue) const
+void FSlateStyleSet::SetParentStyleName(const FName& InParentStyleName)
+{
+	ParentStyleName = InParentStyleName;	
+}
+
+const ISlateStyle* FSlateStyleSet::GetParentStyle() const
+{
+	return ParentStyleName.IsNone() ? nullptr : FSlateStyleRegistry::FindSlateStyle(ParentStyleName);
+}
+
+float FSlateStyleSet::GetFloat(const FName PropertyName, const ANSICHAR* Specifier, float DefaultValue, const ISlateStyle* RequestingStyle) const
 {
 	const float* Result = FloatValues.Find(Join(PropertyName, Specifier));
 
-#if DO_GUARD_SLOW
-	if ((Result == nullptr) && !MissingResources.Contains(PropertyName))
+	if (Result == nullptr) 
 	{
-		MissingResources.Add(PropertyName);
-		Log(ISlateStyle::Warning, FText::Format(NSLOCTEXT("SlateStyleSet", "UknownSlateFloat", "Unable to find float property '{0}' in style."), FText::FromName(PropertyName)));
+		if (const ISlateStyle* ParentStyle = GetParentStyle())
+		{
+			return ParentStyle->GetFloat(PropertyName, Specifier, DefaultValue, RequestingStyle ? RequestingStyle : this);
+		}
+	}
+
+	if (Result == nullptr) 
+	{
+		const ISlateStyle* ReportingStyle = RequestingStyle != nullptr ? RequestingStyle : this;
+		ReportingStyle->LogMissingResource(ISlateStyle::Warning, FText::Format(NSLOCTEXT("SlateStyleSet", "UknownSlateFloat", "Unable to find float property '{0}' in style."), FText::FromName(PropertyName)), PropertyName);
 	}
-#endif
 
 	return Result ? *Result : DefaultValue;
 }
 
-FVector2D FSlateStyleSet::GetVector(const FName PropertyName, const ANSICHAR* Specifier, FVector2D DefaultValue) const
+FVector2D FSlateStyleSet::GetVector(const FName PropertyName, const ANSICHAR* Specifier, FVector2D DefaultValue, const ISlateStyle* RequestingStyle) const
 {
 	const FVector2D* const Result = Vector2DValues.Find(Join(PropertyName, Specifier));
 
-#if DO_GUARD_SLOW
-	if ((Result == nullptr) && !MissingResources.Contains(PropertyName))
+	if (Result == nullptr) 
+	{
+		if (const ISlateStyle* ParentStyle = GetParentStyle())
+		{
+			return ParentStyle->GetVector(PropertyName, Specifier, DefaultValue);
+		}
+	}
+
+	if (Result == nullptr)
 	{
-		MissingResources.Add(PropertyName);
-		Log(ISlateStyle::Warning, FText::Format(NSLOCTEXT("SlateStyleSet", "UknownSlateVector", "Unable to find vector property '{0}' in style."), FText::FromName(PropertyName)));
+		const ISlateStyle* ReportingStyle = RequestingStyle != nullptr ? RequestingStyle : this;
+		ReportingStyle->LogMissingResource(ISlateStyle::Warning, FText::Format(NSLOCTEXT("SlateStyleSet", "UknownSlateVector", "Unable to find vector property '{0}' in style."), FText::FromName(PropertyName)), PropertyName);
 	}
-#endif
 
 	return Result ? *Result : DefaultValue;
 }
 
-const FLinearColor& FSlateStyleSet::GetColor(const FName PropertyName, const ANSICHAR* Specifier, const FLinearColor& DefaultValue) const
+const FLinearColor& FSlateStyleSet::GetColor(const FName PropertyName, const ANSICHAR* Specifier, const FLinearColor& DefaultValue, const ISlateStyle* RequestingStyle) const
 {
 	const FName LookupName(Join(PropertyName, Specifier));
 	const FLinearColor* Result = ColorValues.Find(LookupName);
 
-#if DO_GUARD_SLOW
-	if ( Result == nullptr && !MissingResources.Contains(LookupName) )
+	if (Result == nullptr) 
 	{
-		MissingResources.Add(LookupName);
-		Log(ISlateStyle::Warning, FText::Format(NSLOCTEXT("SlateStyleSet", "UknownColor", "Unable to find Color '{0}'."), FText::FromName(LookupName)));
+		if (const ISlateStyle* ParentStyle = GetParentStyle())
+		{
+			return ParentStyle->GetColor(PropertyName, Specifier, DefaultValue);
+		}
+	}
+
+	if ( Result == nullptr) 
+	{
+		const ISlateStyle* ReportingStyle = RequestingStyle != nullptr ? RequestingStyle : this;
+		ReportingStyle->LogMissingResource(ISlateStyle::Warning, FText::Format(NSLOCTEXT("SlateStyleSet", "UknownColor", "Unable to find Color '{0}'."), FText::FromName(LookupName)), LookupName);
 	}
-#endif
 
 	return Result ? *Result : DefaultValue;
 }
 
-const FSlateColor FSlateStyleSet::GetSlateColor(const FName PropertyName, const ANSICHAR* Specifier, const FSlateColor& DefaultValue) const
+const FSlateColor FSlateStyleSet::GetSlateColor(const FName PropertyName, const ANSICHAR* Specifier, const FSlateColor& DefaultValue, const ISlateStyle* RequestingStyle) const
 {
 	const FName StyleName(Join(PropertyName, Specifier));
 	const FSlateColor* Result = SlateColorValues.Find(StyleName);
+	const FLinearColor* LinearResult = nullptr;
 
 	if ( Result == nullptr )
 	{
-		const FLinearColor* LinearColorLookup = ColorValues.Find(StyleName);
-		if ( LinearColorLookup == nullptr )
+		LinearResult = ColorValues.Find(StyleName);
+	}
+
+	if ( Result == nullptr && LinearResult == nullptr) 
+	{
+		if (const ISlateStyle* ParentStyle = GetParentStyle())
 		{
-			return DefaultValue;
+			return ParentStyle->GetSlateColor(PropertyName, Specifier, DefaultValue);
 		}
-
-		return *LinearColorLookup;
 	}
 
-#if DO_GUARD_SLOW
-	if ( Result == nullptr && !MissingResources.Contains(StyleName) )
+	if ( Result == nullptr && LinearResult == nullptr)
 	{
-		MissingResources.Add(StyleName);
-		Log(ISlateStyle::Warning, FText::Format(NSLOCTEXT("SlateStyleSet", "UknownSlateColor", "Unable to find SlateColor '{0}'."), FText::FromName(StyleName)));
+		const ISlateStyle* ReportingStyle = RequestingStyle != nullptr ? RequestingStyle : this;
+		ReportingStyle->LogMissingResource(ISlateStyle::Warning, FText::Format(NSLOCTEXT("SlateStyleSet", "UknownSlateColor", "Unable to find SlateColor '{0}'."), FText::FromName(StyleName)), StyleName);
 	}
-#endif
 
-	return *Result;
+	return Result ? *Result : LinearResult ? *LinearResult : DefaultValue;
 }
 
-const FMargin& FSlateStyleSet::GetMargin(const FName PropertyName, const ANSICHAR* Specifier, const FMargin& DefaultValue) const
+const FMargin& FSlateStyleSet::GetMargin(const FName PropertyName, const ANSICHAR* Specifier, const FMargin& DefaultValue, const ISlateStyle* RequestingStyle) const
 {
 	const FName StyleName(Join(PropertyName, Specifier));
 	const FMargin* const Result = MarginValues.Find(StyleName);
 
-#if DO_GUARD_SLOW
-	if ( Result == nullptr && !MissingResources.Contains(StyleName) )
+	if (Result == nullptr) 
+	{
+		if (const ISlateStyle* ParentStyle = GetParentStyle())
+		{
+			return ParentStyle->GetMargin(PropertyName, Specifier, DefaultValue);
+		}
+	}
+
+	if ( Result == nullptr )
 	{
-		MissingResources.Add(StyleName);
-		Log(ISlateStyle::Warning, FText::Format(NSLOCTEXT("SlateStyleSet", "UknownMargin", "Unable to find Margin '{0}'."), FText::FromName(StyleName)));
+		const ISlateStyle* ReportingStyle = RequestingStyle != nullptr ? RequestingStyle : this;
+		ReportingStyle->LogMissingResource(ISlateStyle::Warning, FText::Format(NSLOCTEXT("SlateStyleSet", "UknownMargin", "Unable to find Margin '{0}'."), FText::FromName(StyleName)), StyleName);
 	}
-#endif
 
 	return Result ? *Result : DefaultValue;
 }
 
-const FSlateBrush* FSlateStyleSet::GetBrush(const FName PropertyName, const ANSICHAR* Specifier) const
+const FSlateBrush* FSlateStyleSet::GetBrush(const FName PropertyName, const ANSICHAR* Specifier, const ISlateStyle* RequestingStyle) const
 {
 	const FName StyleName = Join(PropertyName, Specifier);
 	const FSlateBrush* Result = BrushResources.FindRef(StyleName);
@@ -269,13 +304,19 @@ const FSlateBrush* FSlateStyleSet::GetBrush(const FName PropertyName, const ANSI
 		}
 	}
 
-#if DO_GUARD_SLOW
-	if ( Result == nullptr && !MissingResources.Contains(StyleName))
+	if (Result == nullptr) 
 	{
-		MissingResources.Add(StyleName);
-		Log(ISlateStyle::Warning, FText::Format(NSLOCTEXT("SlateStyleSet", "UknownBrush", "Unable to find Brush '{0}'."), FText::FromName(StyleName)));
+		if (const ISlateStyle* ParentStyle = GetParentStyle())
+		{
+			return ParentStyle->GetBrush(PropertyName, Specifier);
+		}
+	}
+
+	if ( Result == nullptr )
+	{
+		const ISlateStyle* ReportingStyle = RequestingStyle != nullptr ? RequestingStyle : this;
+		ReportingStyle->LogMissingResource(ISlateStyle::Warning, FText::Format(NSLOCTEXT("SlateStyleSet", "UknownBrush", "Unable to find Brush '{0}'."), FText::FromName(StyleName)), StyleName);
 	}
-#endif
 
 	return Result ? Result : GetDefaultBrush();
 }
@@ -296,20 +337,28 @@ const FSlateBrush* FSlateStyleSet::GetOptionalBrush(const FName PropertyName, co
 		}
 	}
 
+	if (Result == nullptr) 
+	{
+		if (const ISlateStyle* ParentStyle = GetParentStyle())
+		{
+			return ParentStyle->GetOptionalBrush(PropertyName, Specifier, InDefaultBrush);
+		}
+	}
+
 	return Result ? Result : InDefaultBrush;
 }
 
-const TSharedPtr< FSlateDynamicImageBrush > FSlateStyleSet::GetDynamicImageBrush(const FName BrushTemplate, const FName TextureName, const ANSICHAR* Specifier)
+const TSharedPtr< FSlateDynamicImageBrush > FSlateStyleSet::GetDynamicImageBrush(const FName BrushTemplate, const FName TextureName, const ANSICHAR* Specifier, const ISlateStyle* RequestingStyle) const
 {
 	return GetDynamicImageBrush(BrushTemplate, Specifier, nullptr, TextureName);
 }
 
-const TSharedPtr< FSlateDynamicImageBrush > FSlateStyleSet::GetDynamicImageBrush(const FName BrushTemplate, const ANSICHAR* Specifier, UTexture2D* TextureResource, const FName TextureName)
+const TSharedPtr< FSlateDynamicImageBrush > FSlateStyleSet::GetDynamicImageBrush(const FName BrushTemplate, const ANSICHAR* Specifier, UTexture2D* TextureResource, const FName TextureName, const ISlateStyle* RequestingStyle) const
 {
 	return GetDynamicImageBrush(Join(BrushTemplate, Specifier), TextureResource, TextureName);
 }
 
-const TSharedPtr< FSlateDynamicImageBrush > FSlateStyleSet::GetDynamicImageBrush(const FName BrushTemplate, UTexture2D* TextureResource, const FName TextureName)
+const TSharedPtr< FSlateDynamicImageBrush > FSlateStyleSet::GetDynamicImageBrush(const FName BrushTemplate, UTexture2D* TextureResource, const FName TextureName, const ISlateStyle* RequestingStyle) const
 {
 	//create a resource name
 	FName ResourceName;
@@ -317,24 +366,24 @@ const TSharedPtr< FSlateDynamicImageBrush > FSlateStyleSet::GetDynamicImageBrush
 
 	//see if we already have that brush
 	TWeakPtr< FSlateDynamicImageBrush > WeakImageBrush = DynamicBrushes.FindRef(ResourceName);
-
-	//if we don't have the image brush, then make it
 	TSharedPtr< FSlateDynamicImageBrush > ReturnBrush = WeakImageBrush.Pin();
 
-	if ( !ReturnBrush.IsValid() )
+	// if no brush was found, check referenced styles
+	if (!ReturnBrush.IsValid()) 
 	{
-		const FSlateBrush* Result = BrushResources.FindRef(Join(BrushTemplate, nullptr));
-
-		if ( Result == nullptr )
+		if (const ISlateStyle* ParentStyle = GetParentStyle())
 		{
-			Result = GetDefaultBrush();
+			return ParentStyle->GetDynamicImageBrush(BrushTemplate, TextureResource, TextureName);
 		}
+	}
 
-		//create the new brush
-		ReturnBrush = MakeShareable(new FSlateDynamicImageBrush(TextureResource, Result->ImageSize, ResourceName));
-
-		//add it to the dynamic brush list
-		DynamicBrushes.Add(ResourceName, ReturnBrush);
+	// If no DynamicImageBrush was found and this style has no more referenced styles to check, 
+	// then notify the RequestingStyle to make the DynamicImageBrush.
+	// This is to avoid all DynamicImageBrushes being made on the RootStyle (usually CoreStyle)
+	if ( !ReturnBrush.IsValid() )
+	{
+		const ISlateStyle* ReportingStyle = RequestingStyle != nullptr ? RequestingStyle : this;
+		ReportingStyle->MakeDynamicImageBrush(BrushTemplate, TextureResource, TextureName);
 	}
 
 	return ReturnBrush;
@@ -345,18 +394,24 @@ FSlateBrush* FSlateStyleSet::GetDefaultBrush() const
 	return DefaultBrush;
 }
 
-const FSlateSound& FSlateStyleSet::GetSound(const FName PropertyName, const ANSICHAR* Specifier) const
+const FSlateSound& FSlateStyleSet::GetSound(const FName PropertyName, const ANSICHAR* Specifier, const ISlateStyle* RequestingStyle) const
 {
 	const FName StyleName = Join(PropertyName, Specifier);
 	const FSlateSound* Result = Sounds.Find(StyleName);
 
-#if DO_GUARD_SLOW
-	if ( Result == nullptr && !MissingResources.Contains(StyleName) )
+	if (Result == nullptr) 
+	{
+		if (const ISlateStyle* ParentStyle = GetParentStyle())
+		{
+			return ParentStyle->GetSound(PropertyName, Specifier);
+		}
+	}
+
+	if ( Result == nullptr )
 	{
-		MissingResources.Add(StyleName);
-		Log(ISlateStyle::Warning, FText::Format(NSLOCTEXT("SlateStyleSet", "UknownSound", "Unable to find Sound '{0}'."), FText::FromName(StyleName)));
+		const ISlateStyle* ReportingStyle = RequestingStyle != nullptr ? RequestingStyle : this;
+		ReportingStyle->LogMissingResource(ISlateStyle::Warning, FText::Format(NSLOCTEXT("SlateStyleSet", "UknownSound", "Unable to find Sound '{0}'."), FText::FromName(StyleName)), StyleName);
 	}
-#endif
 
 	return Result ? *Result : FStyleDefaults::GetSound();
 }
@@ -365,30 +420,120 @@ FSlateFontInfo FSlateStyleSet::GetFontStyle(const FName PropertyName, const ANSI
 {
 	const FSlateFontInfo* Result = FontInfoResources.Find(Join(PropertyName, Specifier));
 
+	if (Result == nullptr) 
+	{
+		if (const ISlateStyle* ParentStyle = GetParentStyle())
+		{
+			return ParentStyle->GetFontStyle(PropertyName, Specifier);
+		}
+	}
+
 	return Result ? *Result : FStyleDefaults::GetFontInfo();
 }
 
-const FSlateWidgetStyle* FSlateStyleSet::GetWidgetStyleInternal(const FName DesiredTypeName, const FName StyleName) const
+const FSlateWidgetStyle* FSlateStyleSet::GetWidgetStyleInternal(const FName DesiredTypeName, const FName StyleName, const FSlateWidgetStyle* DefaultStyle, bool bWarnIfNotFound) const
 {
 	const TSharedRef< struct FSlateWidgetStyle >* StylePtr = WidgetStyleValues.Find(StyleName);
 
+	if (StylePtr == nullptr) 
+	{
+		if (const ISlateStyle* ParentStyle = GetParentStyle())
+		{
+			if (const FSlateWidgetStyle* Result = ParentStyle->GetWidgetStyleInternal(DesiredTypeName, StyleName, DefaultStyle, bWarnIfNotFound))
+			{
+				return Result;
+			}
+		}
+	}
+
 	if ( StylePtr == nullptr )
 	{
-		Log(ISlateStyle::Warning, FText::Format(NSLOCTEXT("SlateStyleSet", "UnknownWidgetStyle", "Unable to find Slate Widget Style '{0}'. Using {1} defaults instead."), FText::FromName(StyleName), FText::FromName(DesiredTypeName)));
-		return nullptr;
+		if (bWarnIfNotFound)
+		{
+			Log(ISlateStyle::Warning, FText::Format(NSLOCTEXT("SlateStyleSet", "UnknownWidgetStyle", "Unable to find Slate Widget Style '{0}'. Using {1} defaults instead."), FText::FromName(StyleName), FText::FromName(DesiredTypeName)));
+		}
+
+		return DefaultStyle;
 	}
 
 	const TSharedRef< struct FSlateWidgetStyle > Style = *StylePtr;
 
 	if ( Style->GetTypeName() != DesiredTypeName )
 	{
-		Log(ISlateStyle::Error, FText::Format(NSLOCTEXT("SlateStyleSet", "WrongWidgetStyleType", "The Slate Widget Style '{0}' is not of the desired type. Desired: '{1}', Actual: '{2}'"), FText::FromName(StyleName), FText::FromName(DesiredTypeName), FText::FromName(Style->GetTypeName())));
+		if (bWarnIfNotFound)
+		{
+			Log(ISlateStyle::Error, FText::Format(NSLOCTEXT("SlateStyleSet", "WrongWidgetStyleType", "The Slate Widget Style '{0}' is not of the desired type. Desired: '{1}', Actual: '{2}'"), FText::FromName(StyleName), FText::FromName(DesiredTypeName), FText::FromName(Style->GetTypeName())));
+		}
+
 		return nullptr;
 	}
 
 	return &Style.Get();
 }
 
+TSet<FName> FSlateStyleSet::GetStyleKeys() const
+{
+	TSet<FName> AllKeys;
+
+	{
+		TArray<FName> Keys;
+		WidgetStyleValues.GenerateKeyArray(Keys);
+		AllKeys.Append(Keys);
+	}
+
+	{
+		TArray<FName> Keys;
+		FloatValues.GenerateKeyArray(Keys);
+		AllKeys.Append(Keys);
+	}
+
+	{
+		TArray<FName> Keys;
+		Vector2DValues.GenerateKeyArray(Keys);
+		AllKeys.Append(Keys);
+	}
+
+	{
+		TArray<FName> Keys;
+		ColorValues.GenerateKeyArray(Keys);
+		AllKeys.Append(Keys);
+	}
+
+	{
+		TArray<FName> Keys;
+		SlateColorValues.GenerateKeyArray(Keys);
+		AllKeys.Append(Keys);
+	}
+
+	{
+		TArray<FName> Keys;
+		MarginValues.GenerateKeyArray(Keys);
+		AllKeys.Append(Keys);
+	}
+
+	{
+		TArray<FName> Keys;
+		BrushResources.GenerateKeyArray(Keys);
+		AllKeys.Append(Keys);
+	}
+
+	{
+		TArray<FName> Keys;
+		Sounds.GenerateKeyArray(Keys);
+		AllKeys.Append(Keys);
+	}
+
+	{
+		TArray<FName> Keys;
+		FontInfoResources.GenerateKeyArray(Keys);
+		AllKeys.Append(Keys);
+	}
+
+	return AllKeys;
+
+}
+
+
 void FSlateStyleSet::Log(ISlateStyle::EStyleMessageSeverity Severity, const FText& Message) const
 {
 	if ( Severity == ISlateStyle::Error )
@@ -413,6 +558,42 @@ void FSlateStyleSet::Log(ISlateStyle::EStyleMessageSeverity Severity, const FTex
 	}
 }
 
+void FSlateStyleSet::LogMissingResource(EStyleMessageSeverity Severity, const FText& Message, const FName& MissingResource) const
+{
+	if (!MissingResources.Contains(MissingResource))
+	{
+		FText StyleAndMessage = FText::Format(NSLOCTEXT("SlateStyleSet", "SlateSetResourceMissing", "Missing Resource from '{0}' Style: '{1}'"), FText::FromName(GetStyleSetName()), Message);
+
+		MissingResources.Add(MissingResource);
+		Log(Severity, StyleAndMessage);
+
+		ensure(false);
+	}
+}
+
+const TSharedPtr< FSlateDynamicImageBrush > FSlateStyleSet::MakeDynamicImageBrush( const FName BrushTemplate, UTexture2D* TextureResource, const FName TextureName ) const
+{
+	// Double check we don't have this brush already defined
+	FName ResourceName = (TextureName == NAME_None) ? BrushTemplate : FName(*( BrushTemplate.ToString() + TextureName.ToString() ));
+	TWeakPtr< FSlateDynamicImageBrush > WeakImageBrush = DynamicBrushes.FindRef(ResourceName);
+	TSharedPtr< FSlateDynamicImageBrush > ReturnBrush = WeakImageBrush.Pin();
+	if (ReturnBrush.IsValid())
+	{
+		return ReturnBrush;
+	}
+
+	const FSlateBrush* TemplateResult = GetOptionalBrush(BrushTemplate, nullptr, GetDefaultBrush());
+
+	// create the new brush
+	ReturnBrush = MakeShareable(new FSlateDynamicImageBrush(TextureResource, TemplateResult->ImageSize, ResourceName));
+
+	// add it to the dynamic brush list
+	DynamicBrushes.Add(ResourceName, ReturnBrush);
+
+	return ReturnBrush;
+
+}
+
 void FSlateStyleSet::LogUnusedBrushResources()
 {
 	TArray<FString> Filenames;
@@ -464,6 +645,7 @@ bool FSlateStyleSet::IsBrushFromFile(const FString& FilePath, const FSlateBrush*
 {
 	FString Path = Brush->GetResourceName().ToString();
 	FPaths::MakeStandardFilename(Path);
+
 	if ( Path.Compare(FilePath, ESearchCase::IgnoreCase) == 0 )
 	{
 		return true;
diff --git a/Engine/Source/Runtime/SlateCore/Private/Styling/StarshipCoreStyle.cpp b/Engine/Source/Runtime/SlateCore/Private/Styling/StarshipCoreStyle.cpp
--- a/Engine/Source/Runtime/SlateCore/Private/Styling/StarshipCoreStyle.cpp
+++ b/Engine/Source/Runtime/SlateCore/Private/Styling/StarshipCoreStyle.cpp
@@ -12,6 +12,7 @@
 #include "Styling/SlateTypes.h"
 #include "Fonts/LegacySlateFontInfoCache.h"
 #include "Styling/SlateStyleMacros.h"
+#include "Styling/ToolBarStyle.h"
 
 // This is to fix the issue that SlateStyleMacros like IMAGE_BRUSH look for RootToContentDir but Style->RootToContentDir is how the core style is set up
 #define RootToContentDir Style->RootToContentDir
@@ -1172,72 +1173,21 @@ TSharedRef<ISlateStyle> FStarshipCoreStyle::Create()
 
 	// ToolBar
 	{
-		Style->Set("ToolBar.Background", new FSlateRoundedBoxBrush(FStyleColors::Background, 4) );
-		Style->Set("ToolBar.Icon", new IMAGE_BRUSH("Icons/icon_tab_toolbar_16px", Icon16x16));
-		Style->Set("ToolBar.Expand", new IMAGE_BRUSH("Icons/toolbar_expand_16x", Icon16x16));
-		Style->Set("ToolBar.SubMenuIndicator", new IMAGE_BRUSH("Common/SubmenuArrow", Icon8x8));
-		Style->Set("ToolBar.SToolBarComboButtonBlock.Padding", FMargin(4.0f, 0.0f));
-		Style->Set("ToolBar.SToolBarButtonBlock.Padding", FMargin(2.0f,0.f));
-		Style->Set("ToolBar.SToolBarCheckComboButtonBlock.Padding", FMargin(4.0f, 0.f));
-		Style->Set("ToolBar.SToolBarButtonBlock.CheckBox.Padding", FMargin(4.0f, 0.f));
-		Style->Set("ToolBar.SToolBarComboButtonBlock.ComboButton.Color", DefaultForeground);
+		FToolBarStyle NormalToolbarStyle =
+			FToolBarStyle()
+			.SetBackground(FSlateColorBrush(FStyleColors::Background))
+			.SetExpandBrush(IMAGE_BRUSH("Icons/toolbar_expand_16x", Icon16x16))
+			.SetSubMenuIndicator(IMAGE_BRUSH("Common/SubmenuArrow", Icon8x8))
+			.SetComboButtonPadding(FMargin(4.0f, 0.0f))
+			.SetButtonPadding(FMargin(2.0f, 0.f))
+			.SetCheckBoxPadding(FMargin(4.0f, 0.f))
+			.SetSeparatorBrush(FSlateColorBrush(FStyleColors::Input))
+			.SetSeparatorPadding(FMargin(1.f, 0.f, 1.f, 0.f))
+			.SetLabelStyle(FTextBlockStyle(NormalText).SetFont(StyleFonts.Normal))
+			.SetEditableTextStyle(FEditableTextBoxStyle(NormalEditableTextBoxStyle).SetFont(StyleFonts.Normal))
+			.SetBlockPadding(FMargin(2.0f, 2.0f, 4.0f, 4.0f))
+			.SetIndentedBlockPadding(FMargin(18.0f, 2.0f, 4.0f, 4.0f));
 
-		Style->Set("ToolBar.Block.IndentedPadding", FMargin(18.0f, 2.0f, 4.0f, 4.0f));
-		Style->Set("ToolBar.Block.Padding", FMargin(2.0f, 2.0f, 4.0f, 4.0f));
-
-		Style->Set("ToolBar.Separator", new FSlateColorBrush(FStyleColors::Input));
-		Style->Set("ToolBar.Separator.Padding", FMargin(8.f, 0.f, 8.f, 0.f));
-
-		Style->Set("ToolBar.Label", FTextBlockStyle(NormalText).SetFont(StyleFonts.Normal));
-		Style->Set("ToolBar.EditableText", FEditableTextBoxStyle(NormalEditableTextBoxStyle).SetFont(StyleFonts.Normal));
-		Style->Set("ToolBar.Keybinding", FTextBlockStyle(NormalText).SetFont(StyleFonts.Small));
-		Style->Set("ToolBar.Label.DropShadowSize", FVector2D::ZeroVector);
-
-		Style->Set("ToolBar.Heading", FTextBlockStyle(NormalText)
-			.SetFont(StyleFonts.Small)
-			.SetColorAndOpacity(FLinearColor(0.4f, 0.4, 0.4f, 1.0f)));
-
-		/* Create style for "ToolBar.CheckBox" widget ... */
-
-/*
-		const FCheckBoxStyle ToolBarCheckBoxStyle = FCheckBoxStyle()
-			.SetUncheckedImage(IMAGE_BRUSH("Common/SmallCheckBox", Icon14x14))
-			.SetCheckedImage(IMAGE_BRUSH("Common/SmallCheckBox_Checked", Icon14x14))
-			.SetUncheckedHoveredImage(IMAGE_BRUSH("Common/SmallCheckBox_Hovered", Icon14x14))
-			.SetCheckedImage(IMAGE_BRUSH("Common/SmallCheckBox_Checked_Hovered", Icon14x14))
-			.SetUncheckedPressedImage(IMAGE_BRUSH("Common/SmallCheckBox_Hovered", Icon14x14, FLinearColor(0.5f, 0.5f, 0.5f)))
-			.SetCheckedPressedImage(IMAGE_BRUSH("Common/SmallCheckBox_Checked_Hovered", Icon14x14, FLinearColor(0.5f, 0.5f, 0.5f)))
-			.SetUndeterminedImage(IMAGE_BRUSH("Common/CheckBox_Undetermined", Icon14x14))
-			.SetUndeterminedHoveredImage(IMAGE_BRUSH("Common/CheckBox_Undetermined_Hovered", Icon14x14))
-			.SetUndeterminedPressedImage(IMAGE_BRUSH("Common/CheckBox_Undetermined_Hovered", Icon14x14, FLinearColor(0.5f, 0.5f, 0.5f)));
-		/ * ... and add new style * /
-		Style->Set("ToolBar.CheckBox", ToolBarCheckBoxStyle);*/
-
-		/* Read-only checkbox that appears next to a menu item */
-		/* Set images for various SCheckBox states associated with read-only toolbar check box items... */
-		const FCheckBoxStyle BasicToolBarCheckStyle = FCheckBoxStyle()
-			.SetUncheckedImage(IMAGE_BRUSH("Icons/Empty_14x", Icon14x14))
-			.SetUncheckedHoveredImage(IMAGE_BRUSH("Icons/Empty_14x", Icon14x14))
-			.SetUncheckedPressedImage(IMAGE_BRUSH("Common/SmallCheckBox_Hovered", Icon14x14))
-			.SetCheckedImage(IMAGE_BRUSH("Common/SmallCheck", Icon14x14))
-			.SetCheckedHoveredImage(IMAGE_BRUSH("Common/SmallCheck", Icon14x14))
-			.SetCheckedPressedImage(IMAGE_BRUSH("Common/SmallCheck", Icon14x14))
-			.SetUndeterminedImage(IMAGE_BRUSH("Icons/Empty_14x", Icon14x14))
-			.SetUndeterminedHoveredImage(FSlateNoResource())
-			.SetUndeterminedPressedImage(FSlateNoResource());
-		Style->Set("ToolBar.Check", BasicToolBarCheckStyle);
-
-		// This radio button is actually just a check box with different images
-		/* Create style for "ToolBar.RadioButton" widget ... */
-		const FCheckBoxStyle ToolbarRadioButtonCheckBoxStyle = FCheckBoxStyle()
-			.SetUncheckedImage(IMAGE_BRUSH("Common/RadioButton_Unselected_16x", Icon16x16))
-			.SetCheckedImage(IMAGE_BRUSH("Common/RadioButton_Selected_16x", Icon16x16))
-			.SetUncheckedHoveredImage(IMAGE_BRUSH("Common/RadioButton_Unselected_16x", Icon16x16, SelectionColor))
-			.SetCheckedHoveredImage(IMAGE_BRUSH("Common/RadioButton_Unselected_16x", Icon16x16, SelectionColor))
-			.SetUncheckedPressedImage(IMAGE_BRUSH("Common/RadioButton_Unselected_16x", Icon16x16, SelectionColor_Pressed))
-			.SetCheckedPressedImage(IMAGE_BRUSH("Common/RadioButton_Selected_16x", Icon16x16, SelectionColor_Pressed));
-		/* ... and add new style */
-		Style->Set("ToolBar.RadioButton", ToolbarRadioButtonCheckBoxStyle);
 
 		/* Create style for "ToolBar.ToggleButton" widget ... */
 		const FCheckBoxStyle ToolBarToggleButtonCheckBoxStyle = FCheckBoxStyle()
@@ -1248,8 +1198,8 @@ TSharedRef<ISlateStyle> FStarshipCoreStyle::Create()
 			.SetCheckedImage(BOX_BRUSH("Common/RoundedSelection_16x", 4.0f / 16.0f, SelectionColor_Pressed))
 			.SetCheckedHoveredImage(BOX_BRUSH("Common/RoundedSelection_16x", 4.0f / 16.0f, SelectionColor_Pressed))
 			.SetCheckedPressedImage(BOX_BRUSH("Common/RoundedSelection_16x", 4.0f / 16.0f, SelectionColor));
-		/* ... and add new style */
-		Style->Set("ToolBar.ToggleButton", ToolBarToggleButtonCheckBoxStyle);
+
+		NormalToolbarStyle.SetToggleButtonStyle(ToolBarToggleButtonCheckBoxStyle);
 
 		const FButtonStyle ToolbarButton = FButtonStyle()
 			.SetNormal(FSlateNoResource())
@@ -1262,19 +1212,12 @@ TSharedRef<ISlateStyle> FStarshipCoreStyle::Create()
 			.SetHoveredForeground(FSlateColor::UseForeground())
 			.SetDisabledForeground(FSlateColor::UseForeground());
 
+		NormalToolbarStyle.SetButtonStyle(ToolbarButton);
 
-		Style->Set("ToolBar.Button", ToolbarButton);
-
-		//Style->Set("ToolBar.Button.Normal", new FSlateNoResource());
-		//Style->Set("ToolBar.Button.Pressed", new BOX_BRUSH("Common/RoundedSelection_16x", 4.0f / 16.0f, SelectionColor_Pressed));
-		//Style->Set("ToolBar.Button.Hovered", new BOX_BRUSH("Common/RoundedSelection_16x", 4.0f / 16.0f, SelectionColor));
+		NormalToolbarStyle.SetSettingsComboButtonStyle(Style->GetWidgetStyle<FComboButtonStyle>("ComboButton"));
+		NormalToolbarStyle.SetIconSize(Icon40x40);
 
-		Style->Set("ToolBar.Button.Checked", new BOX_BRUSH("Common/RoundedSelection_16x", 4.0f / 16.0f, SelectionColor_Pressed));
-		Style->Set("ToolBar.Button.Checked_Hovered", new BOX_BRUSH("Common/RoundedSelection_16x", 4.0f / 16.0f, SelectionColor_Pressed));
-		Style->Set("ToolBar.Button.Checked_Pressed", new BOX_BRUSH("Common/RoundedSelection_16x", 4.0f / 16.0f, SelectionColor));
-
-		Style->Set("ToolBar.SimpleComboButton", Style->GetWidgetStyle<FComboButtonStyle>("ComboButton"));
-		Style->Set("ToolBar.IconSize", Icon40x40);
+		Style->Set("ToolBar", NormalToolbarStyle);
 
 		// Slim Toolbar
 
@@ -1286,23 +1229,19 @@ TSharedRef<ISlateStyle> FStarshipCoreStyle::Create()
 			.SetHoveredForeground(FStyleColors::ForegroundHover)
 			.SetDisabledForeground(FStyleColors::Foreground);
 
-		Style->Set("SlimToolBar.Background", new FSlateColorBrush(FStyleColors::Background));
-		Style->Set("SlimToolBar.Expand", new IMAGE_BRUSH("Icons/toolbar_expand_16x", Icon16x16));
-		Style->Set("SlimToolBar.SubMenuIndicator", new IMAGE_BRUSH("Common/SubmenuArrow", Icon8x8));
-		Style->Set("SlimToolBar.SToolBarComboButtonBlock.Padding", FMargin(6.0f, 0.0f));
-		Style->Set("SlimToolBar.SToolBarButtonBlock.Padding", FMargin(4.0f, 0.0f));
-		Style->Set("SlimToolBar.SToolBarButtonBlock.CheckBox.Padding", FMargin(10.0f, 0.0f));
-
-		Style->Set("SlimToolBar.Separator", new FSlateColorBrush(FStyleColors::Input));
-		Style->Set("SlimToolBar.Separator.Padding", FMargin(8.f, 0));
-
-		Style->Set("SlimToolBar.Label", FTextBlockStyle(NormalText).SetFont(DEFAULT_FONT("Regular", 10)));
-		Style->Set("SlimToolBar.Label.Padding", FMargin(5,9,0,9));
-		Style->Set("SlimToolBar.Label.DropShadowSize", FVector2D::ZeroVector);
-
-		Style->Set("SlimToolBar.EditableText", FEditableTextBoxStyle(NormalEditableTextBoxStyle).SetFont(DEFAULT_FONT("Regular", 9)));
-
-	//	Style->Set("SlimToolBar.CheckBox", ToolBarCheckBoxStyle);
+		FToolBarStyle SlimToolbarStyle =
+			FToolBarStyle()
+			.SetBackground(FSlateColorBrush(FStyleColors::Background))
+			.SetExpandBrush(IMAGE_BRUSH("Icons/toolbar_expand_16x", Icon16x16))
+			.SetSubMenuIndicator(IMAGE_BRUSH("Common/SubmenuArrow", Icon8x8))
+			.SetComboButtonPadding(FMargin(6.0f, 0.0f))
+			.SetButtonPadding(FMargin(4.0f, 0.0f))
+			.SetCheckBoxPadding(FMargin(10.0f, 0.0f))
+			.SetSeparatorBrush(FSlateColorBrush(FStyleColors::Input))
+			.SetSeparatorPadding(FMargin(8.f, 0))
+			.SetLabelStyle(FTextBlockStyle(NormalText))
+			.SetLabelPadding(FMargin(5, 9, 0, 9))
+			.SetEditableTextStyle(FEditableTextBoxStyle(NormalEditableTextBoxStyle));
 
 		const FCheckBoxStyle SlimToolBarToggleButtonCheckBoxStyle = FCheckBoxStyle()
 			.SetCheckBoxType(ESlateCheckBoxType::ToggleButton)
@@ -1319,16 +1258,18 @@ TSharedRef<ISlateStyle> FStarshipCoreStyle::Create()
 			.SetCheckedPressedForegroundColor(FStyleColors::PrimaryPress)
 			.SetCheckedHoveredForegroundColor(FStyleColors::PrimaryHover);
 
-		Style->Set("SlimToolBar.ToggleButton", SlimToolBarToggleButtonCheckBoxStyle);
-		Style->Set("SlimToolBar.Button", SlimToolBarButton);
+		SlimToolbarStyle.SetToggleButtonStyle(SlimToolBarToggleButtonCheckBoxStyle);
+		SlimToolbarStyle.SetButtonStyle(SlimToolBarButton);
 
 		FComboButtonStyle SlimToolBarComboButton = FComboButtonStyle(Style->GetWidgetStyle<FComboButtonStyle>("ComboButton"))
 			.SetContentPadding(0)
 			.SetButtonStyle(SlimToolBarButton)
 			.SetDownArrowImage(IMAGE_BRUSH_SVG("Starship/Common/ellipsis-vertical-narrow", FVector2D(6,24)));
 
-		Style->Set("SlimToolBar.SimpleComboButton", SlimToolBarComboButton);
-		Style->Set("SlimToolBar.IconSize", Icon20x20);
+		SlimToolbarStyle.SetSettingsComboButtonStyle(SlimToolBarComboButton);
+		SlimToolbarStyle.SetIconSize(Icon20x20);
+
+		Style->Set("SlimToolBar", SlimToolbarStyle);
 	}
 
 	// MenuBar
diff --git a/Engine/Source/Runtime/SlateCore/Private/Styling/ToolBarStyle.cpp b/Engine/Source/Runtime/SlateCore/Private/Styling/ToolBarStyle.cpp
new file mode 100644
--- /dev/null
+++ b/Engine/Source/Runtime/SlateCore/Private/Styling/ToolBarStyle.cpp
@@ -0,0 +1,47 @@
+// Copyright Epic Games, Inc. All Rights Reserved.
+
+#include "Styling/ToolBarStyle.h"
+#include "Brushes/SlateNoResource.h"
+
+const FName FToolBarStyle::TypeName(TEXT("FToolbarStyle"));
+
+FToolBarStyle::FToolBarStyle()
+	: BackgroundBrush(FSlateNoResource())
+	, ExpandBrush(FSlateNoResource())
+	, SubMenuIndicator(FSlateNoResource())
+	, SeparatorBrush(FSlateNoResource())
+	, LabelStyle()
+	, EditableTextStyle()
+	, ToggleButton()
+	, SettingsComboButton()
+	, ButtonStyle()
+	, LabelPadding(0)
+	, SeparatorPadding(0)
+	, ComboButtonPadding(0)
+	, ButtonPadding(0)
+	, CheckBoxPadding(0)
+	, BlockPadding(0)
+	, IndentedBlockPadding(0)
+	, IconSize(16,16)
+{}
+
+const FToolBarStyle& FToolBarStyle::GetDefault()
+{
+	static FToolBarStyle Default;
+	return Default;
+}
+
+void FToolBarStyle::GetResources(TArray<const FSlateBrush*>& OutBrushes) const
+{
+	OutBrushes.Add(&BackgroundBrush);
+	OutBrushes.Add(&ExpandBrush);
+	OutBrushes.Add(&SubMenuIndicator);
+	OutBrushes.Add(&SeparatorBrush);
+
+	LabelStyle.GetResources(OutBrushes);
+	EditableTextStyle.GetResources(OutBrushes);
+	ToggleButton.GetResources(OutBrushes);
+	SettingsComboButton.GetResources(OutBrushes);
+	ButtonStyle.GetResources(OutBrushes);
+}
+
diff --git a/Engine/Source/Runtime/SlateCore/Public/Styling/ISlateStyle.h b/Engine/Source/Runtime/SlateCore/Public/Styling/ISlateStyle.h
--- a/Engine/Source/Runtime/SlateCore/Public/Styling/ISlateStyle.h
+++ b/Engine/Source/Runtime/SlateCore/Public/Styling/ISlateStyle.h
@@ -33,6 +33,7 @@ public:
 	 * @param OutResources - the array to populate.
 	 */
 	virtual void GetResources( TArray< const FSlateBrush* >& OutResources ) const = 0;
+	virtual FString GetContentRootDir() const = 0;
 
 	/**
 	 * Gets the names of every style entry using a brush in this style
@@ -44,9 +45,10 @@ public:
 	virtual TArray<FName> GetEntriesUsingBrush(const FName BrushName) const = 0;
 
 	template< typename WidgetStyleType >            
-	const WidgetStyleType& GetWidgetStyle( FName PropertyName, const ANSICHAR* Specifier = nullptr ) const 
+	const WidgetStyleType& GetWidgetStyle( FName PropertyName, const ANSICHAR* Specifier = nullptr, const WidgetStyleType* DefaultValue = nullptr ) const 
 	{
-		const FSlateWidgetStyle* ExistingStyle = GetWidgetStyleInternal( WidgetStyleType::TypeName, Join( PropertyName, Specifier ) );
+		const bool bWarnIfNotFound = true;
+		const FSlateWidgetStyle* ExistingStyle = GetWidgetStyleInternal(WidgetStyleType::TypeName, Join( PropertyName, Specifier ), DefaultValue, bWarnIfNotFound);
 
 		if ( ExistingStyle == nullptr )
 		{
@@ -59,7 +61,8 @@ public:
 	template< typename WidgetStyleType >            
 	bool HasWidgetStyle( FName PropertyName, const ANSICHAR* Specifier = nullptr ) const
 	{
-		return GetWidgetStyleInternal( WidgetStyleType::TypeName, Join( PropertyName, Specifier ) ) != nullptr;
+		const bool bWarnIfNotFound = false;
+		return GetWidgetStyleInternal( WidgetStyleType::TypeName, Join( PropertyName, Specifier ), nullptr, bWarnIfNotFound) != nullptr;
 	}
 
 	/**
@@ -69,7 +72,7 @@ public:
 	 *
 	 * @return a float property.
 	 */
-	virtual float GetFloat(const FName PropertyName, const ANSICHAR* Specifier = nullptr, float DefaultValue = FStyleDefaults::GetFloat()) const = 0;
+	virtual float GetFloat(const FName PropertyName, const ANSICHAR* Specifier = nullptr, float DefaultValue = FStyleDefaults::GetFloat(), const ISlateStyle* RequestingStyle = nullptr) const = 0;
 
 	/**
 	 * @param PropertyName   Name of the property to get
@@ -78,7 +81,7 @@ public:
 	 *
 	 * @return a FVector2D property.
 	 */
-	virtual FVector2D GetVector(const FName PropertyName, const ANSICHAR* Specifier = nullptr, FVector2D DefaultValue = FStyleDefaults::GetVector2D()) const = 0;
+	virtual FVector2D GetVector(const FName PropertyName, const ANSICHAR* Specifier = nullptr, FVector2D DefaultValue = FStyleDefaults::GetVector2D(), const ISlateStyle* RequestingStyle = nullptr) const = 0;
 
 	/**
 	 * @param PropertyName   Name of the property to get
@@ -87,7 +90,7 @@ public:
 	 *
 	 * @return a FLinearColor property.
 	 */
-	virtual const FLinearColor& GetColor(const FName PropertyName, const ANSICHAR* Specifier = nullptr, const FLinearColor& DefaultValue = FStyleDefaults::GetColor()) const = 0;
+	virtual const FLinearColor& GetColor(const FName PropertyName, const ANSICHAR* Specifier = nullptr, const FLinearColor& DefaultValue = FStyleDefaults::GetColor(), const ISlateStyle* RequestingStyle = nullptr) const = 0;
 
 	/**
 	 * @param PropertyName   Name of the property to get
@@ -96,7 +99,7 @@ public:
 	 *
 	 * @return a FLinearColor property.
 	 */
-	virtual const FSlateColor GetSlateColor(const FName PropertyName, const ANSICHAR* Specifier = nullptr, const FSlateColor& DefaultValue = FStyleDefaults::GetSlateColor()) const = 0;
+	virtual const FSlateColor GetSlateColor(const FName PropertyName, const ANSICHAR* Specifier = nullptr, const FSlateColor& DefaultValue = FStyleDefaults::GetSlateColor(), const ISlateStyle* RequestingStyle = nullptr) const = 0;
 
 	/**
 	 * @param PropertyName   Name of the property to get
@@ -105,7 +108,7 @@ public:
 	 *
 	 * @return a FMargin property.
 	 */
-	virtual const FMargin& GetMargin(const FName PropertyName, const ANSICHAR* Specifier = nullptr, const FMargin& DefaultValue = FStyleDefaults::GetMargin()) const = 0;
+	virtual const FMargin& GetMargin(const FName PropertyName, const ANSICHAR* Specifier = nullptr, const FMargin& DefaultValue = FStyleDefaults::GetMargin(), const ISlateStyle* RequestingStyle = nullptr) const = 0;
 
 	/**
 	 * @param PropertyName   Name of the property to get
@@ -113,7 +116,7 @@ public:
 	 *
 	 * @return an FSlateBrush property.
 	 */		
-	virtual const FSlateBrush* GetBrush( const FName PropertyName, const ANSICHAR* Specifier = nullptr ) const = 0;
+	virtual const FSlateBrush* GetBrush( const FName PropertyName, const ANSICHAR* Specifier = nullptr, const ISlateStyle* RequestingStyle = nullptr) const = 0;
 
 	/**
 	 * Just like GetBrush, but returns DefaultBrush instead of the
@@ -124,13 +127,13 @@ public:
 	 *
 	 * @return an FSlateBrush property.
 	 */		
-	virtual const FSlateBrush* GetOptionalBrush( const FName PropertyName, const ANSICHAR* Specifier = nullptr, const FSlateBrush* const DefaultBrush = FStyleDefaults::GetNoBrush() ) const = 0;
+	virtual const FSlateBrush* GetOptionalBrush( const FName PropertyName, const ANSICHAR* Specifier = nullptr, const FSlateBrush* const DefaultBrush = FStyleDefaults::GetNoBrush()) const = 0;
 
-	virtual const TSharedPtr< FSlateDynamicImageBrush > GetDynamicImageBrush( const FName BrushTemplate, const FName TextureName, const ANSICHAR* Specifier = nullptr ) = 0;
+	virtual const TSharedPtr< FSlateDynamicImageBrush > GetDynamicImageBrush( const FName BrushTemplate, const FName TextureName, const ANSICHAR* Specifier = nullptr, const ISlateStyle* RequestingStyle = nullptr ) const = 0;
 
-	virtual const TSharedPtr< FSlateDynamicImageBrush > GetDynamicImageBrush( const FName BrushTemplate, const ANSICHAR* Specifier, UTexture2D* TextureResource, const FName TextureName ) = 0;
+	virtual const TSharedPtr< FSlateDynamicImageBrush > GetDynamicImageBrush( const FName BrushTemplate, const ANSICHAR* Specifier, UTexture2D* TextureResource, const FName TextureName, const ISlateStyle* RequestingStyle = nullptr ) const = 0;
 
-	virtual const TSharedPtr< FSlateDynamicImageBrush > GetDynamicImageBrush( const FName BrushTemplate, UTexture2D* TextureResource, const FName TextureName ) = 0;
+	virtual const TSharedPtr< FSlateDynamicImageBrush > GetDynamicImageBrush( const FName BrushTemplate, UTexture2D* TextureResource, const FName TextureName, const ISlateStyle* RequestingStyle = nullptr ) const = 0;
 
 	/**
 	 * Get default FSlateBrush.
@@ -139,7 +142,7 @@ public:
 	virtual FSlateBrush* GetDefaultBrush() const = 0;
 
 	/** Look up a sound property specified by PropertyName and optional Specifier. */
-	virtual const FSlateSound& GetSound( const FName PropertyName, const ANSICHAR* Specifier = nullptr ) const = 0;
+	virtual const FSlateSound& GetSound( const FName PropertyName, const ANSICHAR* Specifier = nullptr , const ISlateStyle* RequestingStyle = nullptr) const = 0;
 
 	/**
 	 * @param Propertyname   Name of the property to get
@@ -161,6 +164,7 @@ public:
 		}
 	}
 
+	virtual TSet<FName> GetStyleKeys() const = 0;
 protected:
 
 	/** 
@@ -179,13 +183,20 @@ protected:
 
 protected:
 
+	friend class FSlateStyleSet;
+	friend class ISlateStyle;
+
 	/**
 	 * 
 	 */
-	virtual const FSlateWidgetStyle* GetWidgetStyleInternal( const FName DesiredTypeName, const FName StyleName ) const = 0;
+	virtual const FSlateWidgetStyle* GetWidgetStyleInternal( const FName DesiredTypeName, const FName StyleName, const FSlateWidgetStyle* DefaultStyle, bool bWarnIfNotFound) const = 0;
 
 	/**
 	 * 
 	 */
 	virtual void Log( EStyleMessageSeverity Severity, const FText& Message ) const = 0;
+	
+	virtual void LogMissingResource(EStyleMessageSeverity Severity, const FText& Message, const FName& MissingResource) const = 0;
+
+	virtual const TSharedPtr< FSlateDynamicImageBrush > MakeDynamicImageBrush( const FName BrushTemplate, UTexture2D* TextureResource, const FName TextureName ) const = 0;
 };
diff --git a/Engine/Source/Runtime/SlateCore/Public/Styling/SlateStyle.h b/Engine/Source/Runtime/SlateCore/Public/Styling/SlateStyle.h
--- a/Engine/Source/Runtime/SlateCore/Public/Styling/SlateStyle.h
+++ b/Engine/Source/Runtime/SlateCore/Public/Styling/SlateStyle.h
@@ -12,6 +12,7 @@
 #include "Brushes/SlateBorderBrush.h"
 #include "Brushes/SlateBoxBrush.h"
 #include "Brushes/SlateColorBrush.h"
+#include "Brushes/SlateRoundedBoxBrush.h"
 #include "Brushes/SlateImageBrush.h"
 #include "Brushes/SlateDynamicImageBrush.h"
 
@@ -30,7 +31,7 @@ public:
 	 * @param InStyleSetName The name used to identity this style set
 	 */
 	FSlateStyleSet(const FName& InStyleSetName);
-	
+
 	/** Destructor. */
 	virtual ~FSlateStyleSet();
 
@@ -44,6 +45,7 @@ public:
 	virtual FString RootToContentDir(const ANSICHAR* RelativePath);
 	virtual FString RootToContentDir(const WIDECHAR* RelativePath);
 	virtual FString RootToContentDir(const FString& RelativePath);
+	virtual FString GetContentRootDir() const { return ContentRootDir; }
 
 	virtual void SetCoreContentRoot(const FString& InCoreContentRootDir);
 
@@ -54,30 +56,32 @@ public:
 	virtual FString RootToCoreContentDir(const WIDECHAR* RelativePath);
 	virtual FString RootToCoreContentDir(const FString& RelativePath);
 
-	virtual float GetFloat(const FName PropertyName, const ANSICHAR* Specifier = nullptr, float DefaultValue = FStyleDefaults::GetFloat()) const override;
-	virtual FVector2D GetVector(const FName PropertyName, const ANSICHAR* Specifier = nullptr, FVector2D DefaultValue = FStyleDefaults::GetVector2D()) const override;
-	virtual const FLinearColor& GetColor(const FName PropertyName, const ANSICHAR* Specifier = nullptr, const FLinearColor& DefaultValue = FStyleDefaults::GetColor()) const override;
-	virtual const FSlateColor GetSlateColor(const FName PropertyName, const ANSICHAR* Specifier = nullptr, const FSlateColor& DefaultValue = FStyleDefaults::GetSlateColor()) const override;
-	virtual const FMargin& GetMargin(const FName PropertyName, const ANSICHAR* Specifier = nullptr, const FMargin& DefaultValue = FStyleDefaults::GetMargin()) const override;
+	virtual float GetFloat(const FName PropertyName, const ANSICHAR* Specifier = nullptr, float DefaultValue = FStyleDefaults::GetFloat(), const ISlateStyle* RequestingStyle = nullptr) const override;
+	virtual FVector2D GetVector(const FName PropertyName, const ANSICHAR* Specifier = nullptr, FVector2D DefaultValue = FStyleDefaults::GetVector2D(), const ISlateStyle* RequestingStyle = nullptr) const override;
+	virtual const FLinearColor& GetColor(const FName PropertyName, const ANSICHAR* Specifier = nullptr, const FLinearColor& DefaultValue = FStyleDefaults::GetColor(), const ISlateStyle* RequestingStyle = nullptr) const override;
+	virtual const FSlateColor GetSlateColor(const FName PropertyName, const ANSICHAR* Specifier = nullptr, const FSlateColor& DefaultValue = FStyleDefaults::GetSlateColor(), const ISlateStyle* RequestingStyle = nullptr) const override;
+	virtual const FMargin& GetMargin(const FName PropertyName, const ANSICHAR* Specifier = nullptr, const FMargin& DefaultValue = FStyleDefaults::GetMargin(), const ISlateStyle* RequestingStyle = nullptr) const override;
 
-	virtual const FSlateBrush* GetBrush(const FName PropertyName, const ANSICHAR* Specifier = nullptr) const override;
+	virtual const FSlateBrush* GetBrush(const FName PropertyName, const ANSICHAR* Specifier = nullptr, const ISlateStyle* RequestingStyle = nullptr) const override;
 	virtual const FSlateBrush* GetOptionalBrush(const FName PropertyName, const ANSICHAR* Specifier = nullptr, const FSlateBrush* const DefaultBrush = FStyleDefaults::GetNoBrush()) const override;
 
-	virtual const TSharedPtr< FSlateDynamicImageBrush > GetDynamicImageBrush(const FName BrushTemplate, const FName TextureName, const ANSICHAR* Specifier = nullptr) override;
-	virtual const TSharedPtr< FSlateDynamicImageBrush > GetDynamicImageBrush(const FName BrushTemplate, const ANSICHAR* Specifier, UTexture2D* TextureResource, const FName TextureName) override;
-	virtual const TSharedPtr< FSlateDynamicImageBrush > GetDynamicImageBrush(const FName BrushTemplate, UTexture2D* TextureResource, const FName TextureName) override;
+	virtual const TSharedPtr<FSlateDynamicImageBrush> GetDynamicImageBrush(const FName BrushTemplate, const FName TextureName, const ANSICHAR* Specifier = nullptr, const ISlateStyle* RequestingStyle = nullptr) const override;
+	virtual const TSharedPtr<FSlateDynamicImageBrush> GetDynamicImageBrush(const FName BrushTemplate, const ANSICHAR* Specifier, UTexture2D* TextureResource, const FName TextureName, const ISlateStyle* RequestingStyle = nullptr) const override;
+	virtual const TSharedPtr<FSlateDynamicImageBrush> GetDynamicImageBrush(const FName BrushTemplate, UTexture2D* TextureResource, const FName TextureName, const ISlateStyle* RequestingStyle = nullptr) const override;
 
 	virtual FSlateBrush* GetDefaultBrush() const override;
 
-	virtual const FSlateSound& GetSound(const FName PropertyName, const ANSICHAR* Specifier = nullptr) const override;
+	virtual const FSlateSound& GetSound(const FName PropertyName, const ANSICHAR* Specifier = nullptr, const ISlateStyle* RequestingStyle = nullptr) const override;
 	virtual FSlateFontInfo GetFontStyle(const FName PropertyName, const ANSICHAR* Specifier = nullptr) const override;
 
+	/** Name a Parent Style to fallback on if the requested styles are specified in this style **/
+	void SetParentStyleName(const FName& InParentStyleName);
 public:
 
-	template< typename DefinitionType >            
-	FORCENOINLINE void Set( const FName PropertyName, const DefinitionType& InStyleDefintion )
+	template< typename DefinitionType >
+	FORCENOINLINE void Set(const FName PropertyName, const DefinitionType& InStyleDefintion)
 	{
-		WidgetStyleValues.Add( PropertyName, MakeShareable( new DefinitionType( InStyleDefintion ) ) );
+		WidgetStyleValues.Add(PropertyName, MakeShared<DefinitionType>(InStyleDefintion));
 	}
 
 	/**
@@ -85,9 +89,9 @@ public:
 	 * @param PropertyName - Name of the property to set.
 	 * @param InFloat - The value to set.
 	 */
-	FORCENOINLINE void Set( const FName PropertyName, const float InFloat )
+	FORCENOINLINE void Set(const FName PropertyName, const float InFloat)
 	{
-		FloatValues.Add( PropertyName, InFloat );
+		FloatValues.Add(PropertyName, InFloat);
 	}
 
 	/**
@@ -95,9 +99,9 @@ public:
 	 * @param PropertyName - Name of the property to set.
 	 * @param InVector - The value to set.
 	 */
-	FORCENOINLINE void Set( const FName PropertyName, const FVector2D InVector )
+	FORCENOINLINE void Set(const FName PropertyName, const FVector2D InVector)
 	{
-		Vector2DValues.Add( PropertyName, InVector );
+		Vector2DValues.Add(PropertyName, InVector);
 	}
 
 	/**
@@ -105,14 +109,14 @@ public:
 	 * @param PropertyName - Name of the property to set.
 	 * @param InColor - The value to set.
 	 */
-	FORCENOINLINE void Set( const FName PropertyName, const FLinearColor& InColor )
+	FORCENOINLINE void Set(const FName PropertyName, const FLinearColor& InColor)
 	{
-		ColorValues.Add( PropertyName, InColor );
+		ColorValues.Add(PropertyName, InColor);
 	}
 
-	FORCENOINLINE void Set( const FName PropertyName, const FColor& InColor )
+	FORCENOINLINE void Set(const FName PropertyName, const FColor& InColor)
 	{
-		ColorValues.Add( PropertyName, InColor );
+		ColorValues.Add(PropertyName, InColor);
 	}
 
 	/**
@@ -120,9 +124,9 @@ public:
 	 * @param PropertyName - Name of the property to add.
 	 * @param InColor - The value to add.
 	 */
-	FORCENOINLINE void Set( const FName PropertyName, const FSlateColor& InColor )
+	FORCENOINLINE void Set(const FName PropertyName, const FSlateColor& InColor)
 	{
-		SlateColorValues.Add( PropertyName, InColor );
+		SlateColorValues.Add(PropertyName, InColor);
 	}
 
 	/**
@@ -130,9 +134,9 @@ public:
 	 * @param PropertyName - Name of the property to add.
 	 * @param InMargin - The value to add.
 	 */
-	FORCENOINLINE void Set( const FName PropertyName, const FMargin& InMargin )
+	FORCENOINLINE void Set(const FName PropertyName, const FMargin& InMargin)
 	{
-		MarginValues.Add( PropertyName, InMargin );
+		MarginValues.Add(PropertyName, InMargin);
 	}
 
 	/**
@@ -140,71 +144,49 @@ public:
 	 * @param PropertyName - Name of the property to set.
 	 * @param InBrush - The brush to set.
 	 */
-	FORCENOINLINE void Set( const FName PropertyName, FSlateBrush* InBrush )
-	{
-		BrushResources.Add( PropertyName, InBrush );
-	}
-
-	FORCENOINLINE void Set( const FName PropertyName, FSlateNoResource* InBrush )
-	{
-		BrushResources.Add( PropertyName, InBrush );
-	}
-
-	FORCENOINLINE void Set( const FName PropertyName, FSlateBoxBrush* InBrush )
-	{
-		BrushResources.Add( PropertyName, InBrush );
-	}
-
-	FORCENOINLINE void Set( const FName PropertyName, FSlateBorderBrush* InBrush )
-	{
-		BrushResources.Add( PropertyName, InBrush );
-	}
-
-	FORCENOINLINE void Set( const FName PropertyName, FSlateImageBrush* InBrush )
-	{
-		BrushResources.Add( PropertyName, InBrush );
-	}
-
-	FORCENOINLINE void Set( const FName PropertyName, FSlateDynamicImageBrush* InBrush )
+	template<typename BrushType>
+	FORCENOINLINE void Set(const FName PropertyName, BrushType* InBrush)
 	{
-		BrushResources.Add( PropertyName, InBrush );
+		BrushResources.Add(PropertyName, InBrush);
 	}
-
-	FORCENOINLINE void Set( const FName PropertyName, FSlateColorBrush* InBrush )
-	{
-		BrushResources.Add( PropertyName, InBrush );
-	}
-
 	/**
 	 * Set FSlateSound properties
 	 *
 	 * @param PropertyName  Name of the property to set
 	 * @param InSound		Sound to set
 	 */
-	FORCENOINLINE void Set( FName PropertyName, const FSlateSound& InSound )
+	FORCENOINLINE void Set(FName PropertyName, const FSlateSound& InSound)
 	{
-		Sounds.Add( PropertyName, InSound );
+		Sounds.Add(PropertyName, InSound);
 	}
-		
+
 	/**
 	 * Set FSlateFontInfo properties
 	 *
 	 * @param PropertyName  Name of the property to set
 	 * @param InFontStyle   The value to set
 	 */
-	FORCENOINLINE void Set( FName PropertyName, const FSlateFontInfo& InFontInfo )
+	FORCENOINLINE void Set(FName PropertyName, const FSlateFontInfo& InFontInfo)
 	{
-		FontInfoResources.Add( PropertyName, InFontInfo );
+		FontInfoResources.Add(PropertyName, InFontInfo);
 	}
 
 protected:
 
-	virtual const FSlateWidgetStyle* GetWidgetStyleInternal(const FName DesiredTypeName, const FName StyleName) const override;
+	friend class FSlateStyleSet;
+	friend class ISlateStyle;
+
+	virtual const FSlateWidgetStyle* GetWidgetStyleInternal(const FName DesiredTypeName, const FName StyleName, const FSlateWidgetStyle* DefaultStyle, bool bWarnIfNotFound) const override;
 
 	virtual void Log(ISlateStyle::EStyleMessageSeverity Severity, const FText& Message) const override;
 
+	virtual void LogMissingResource(EStyleMessageSeverity Severity, const FText& Message, const FName& MissingResource) const override;
+
+	virtual const TSharedPtr< FSlateDynamicImageBrush > MakeDynamicImageBrush(const FName BrushTemplate, UTexture2D* TextureResource, const FName TextureName) const override;
+
 	virtual void LogUnusedBrushResources();
 
+	TSet<FName> GetStyleKeys() const;
 protected:
 
 	bool IsBrushFromFile(const FString& FilePath, const FSlateBrush* Brush);
@@ -228,7 +210,7 @@ protected:
 
 	/** Color property storage. */
 	TMap< FName, FLinearColor > ColorValues;
-	
+
 	/** FSlateColor property storage. */
 	TMap< FName, FSlateColor > SlateColorValues;
 
@@ -241,13 +223,17 @@ protected:
 
 	/** SlateSound property storage */
 	TMap< FName, FSlateSound > Sounds;
-	
+
 	/** FSlateFontInfo property storage. */
 	TMap< FName, FSlateFontInfo > FontInfoResources;
 
 	/** A list of dynamic brushes */
-	TMap< FName, TWeakPtr< FSlateDynamicImageBrush > > DynamicBrushes;
+	mutable TMap< FName, TWeakPtr< FSlateDynamicImageBrush > > DynamicBrushes;
 
 	/** A set of resources that were requested, but not found. */
 	mutable TSet< FName > MissingResources;
+
+	FName ParentStyleName;
+	const ISlateStyle* GetParentStyle() const;
 };
+
diff --git a/Engine/Source/Runtime/SlateCore/Public/Styling/ToolBarStyle.h b/Engine/Source/Runtime/SlateCore/Public/Styling/ToolBarStyle.h
new file mode 100644
--- /dev/null
+++ b/Engine/Source/Runtime/SlateCore/Public/Styling/ToolBarStyle.h
@@ -0,0 +1,99 @@
+// Copyright Epic Games, Inc. All Rights Reserved.
+
+#pragma once
+
+#include "CoreMinimal.h"
+#include "UObject/ObjectMacros.h"
+#include "Styling/SlateTypes.h"
+
+#include "ToolbarStyle.generated.h"
+
+/**
+ * Represents the appearance of a toolbar 
+ */
+USTRUCT(BlueprintType)
+struct SLATECORE_API FToolBarStyle : public FSlateWidgetStyle
+{
+	GENERATED_BODY()
+
+	FToolBarStyle();
+
+	virtual ~FToolBarStyle() {}
+
+	virtual void GetResources(TArray<const FSlateBrush*>& OutBrushes) const override;
+
+	static const FName TypeName;
+	virtual const FName GetTypeName() const override { return TypeName; };
+
+	static const FToolBarStyle& GetDefault();
+
+	/** The brush used for the background of the toolbar */
+	UPROPERTY(EditAnywhere, BlueprintReadWrite, Category = Appearance)
+	FSlateBrush BackgroundBrush;
+	FToolBarStyle& SetBackground(const FSlateBrush& InBackground) { BackgroundBrush = InBackground; return *this; }
+
+	/** The brush used for the expand arrow when the toolbar runs out of room and needs to display toolbar items in a menu*/
+	UPROPERTY(EditAnywhere, BlueprintReadWrite, Category = Appearance)
+	FSlateBrush ExpandBrush;
+	FToolBarStyle& SetExpandBrush(const FSlateBrush& InExpandBrush) { ExpandBrush = InExpandBrush; return *this; }
+
+	UPROPERTY(EditAnywhere, BlueprintReadWrite, Category = Appearance)
+	FSlateBrush SubMenuIndicator;
+	FToolBarStyle& SetSubMenuIndicator(const FSlateBrush& InSubMenuIndicator) { SubMenuIndicator = InSubMenuIndicator; return *this; }
+
+	UPROPERTY(EditAnywhere, BlueprintReadWrite, Category = Appearance)
+	FSlateBrush SeparatorBrush;
+	FToolBarStyle& SetSeparatorBrush(const FSlateBrush& InSeparatorBrush) { SeparatorBrush = InSeparatorBrush; return *this; }
+
+	UPROPERTY(EditAnywhere, BlueprintReadWrite, Category = Appearance)
+	FTextBlockStyle LabelStyle;
+	FToolBarStyle& SetLabelStyle(const FTextBlockStyle& InLabelStyle) { LabelStyle = InLabelStyle; return *this; }
+
+	UPROPERTY(EditAnywhere, BlueprintReadWrite, Category = Appearance)
+	FEditableTextBoxStyle EditableTextStyle;
+	FToolBarStyle& SetEditableTextStyle(const FEditableTextBoxStyle& InEditableTextStyle) { EditableTextStyle = InEditableTextStyle; return *this; }
+
+	UPROPERTY(EditAnywhere, BlueprintReadWrite, Category = Appearance)
+	FCheckBoxStyle ToggleButton;
+	FToolBarStyle& SetToggleButtonStyle(const FCheckBoxStyle& InToggleButton) { ToggleButton = InToggleButton; return *this; }
+
+	UPROPERTY(EditAnywhere, BlueprintReadWrite, Category = Appearance)
+	FComboButtonStyle SettingsComboButton;
+	FToolBarStyle& SetSettingsComboButtonStyle(const FComboButtonStyle& InSettingsComboButton) { SettingsComboButton = InSettingsComboButton; return *this; }
+
+	UPROPERTY(EditAnywhere, BlueprintReadWrite, Category = Appearance)
+	FButtonStyle ButtonStyle;
+	FToolBarStyle& SetButtonStyle(const FButtonStyle& InButtonStyle) { ButtonStyle = InButtonStyle; return *this; }
+
+	UPROPERTY(EditAnywhere, BlueprintReadWrite, Category = Appearance)
+	FMargin LabelPadding;
+	FToolBarStyle& SetLabelPadding(const FMargin& InLabelPadding) { LabelPadding = InLabelPadding; return *this; }
+
+	UPROPERTY(EditAnywhere, BlueprintReadWrite, Category = Appearance)
+	FMargin SeparatorPadding;
+	FToolBarStyle& SetSeparatorPadding(const FMargin& InSeparatorPadding) { SeparatorPadding = InSeparatorPadding; return *this; }
+
+	UPROPERTY(EditAnywhere, BlueprintReadWrite, Category = Appearance)
+	FMargin ComboButtonPadding;
+	FToolBarStyle& SetComboButtonPadding(const FMargin& InComboButtonPadding) { ComboButtonPadding = InComboButtonPadding; return *this; }
+
+	UPROPERTY(EditAnywhere, BlueprintReadWrite, Category = Appearance)
+	FMargin ButtonPadding;
+	FToolBarStyle& SetButtonPadding(const FMargin& InButtonPadding) { ButtonPadding = InButtonPadding; return *this; }
+
+	UPROPERTY(EditAnywhere, BlueprintReadWrite, Category = Appearance)
+	FMargin CheckBoxPadding;
+	FToolBarStyle& SetCheckBoxPadding(const FMargin& InCheckBoxPadding) { CheckBoxPadding = InCheckBoxPadding; return *this; }
+
+	UPROPERTY(EditAnywhere, BlueprintReadWrite, Category = Appearance)
+	FMargin BlockPadding;
+	FToolBarStyle& SetBlockPadding(const FMargin& InBlockPadding) { BlockPadding = InBlockPadding; return *this; }
+
+	UPROPERTY(EditAnywhere, BlueprintReadWrite, Category = Appearance)
+	FMargin IndentedBlockPadding;
+	FToolBarStyle& SetIndentedBlockPadding(const FMargin& InIndentedBlockPadding) { IndentedBlockPadding = InIndentedBlockPadding; return *this; }
+
+	UPROPERTY(EditAnywhere, BlueprintReadWrite, Category = Appearance)
+	FVector2D IconSize;
+	FToolBarStyle& SetIconSize(const FVector2D& InIconSize) { IconSize = InIconSize; return *this; }
+};
