Renamed existing viewport toolbar style to LegacyViewportToolbar.  Keeping it around for backwards compat but it wil no longer be used in the editor

diff --git a/Engine/Plugins/Enterprise/DataprepEditor/Source/DataprepEditor/Private/DataprepEditorStyle.cpp b/Engine/Plugins/Enterprise/DataprepEditor/Source/DataprepEditor/Private/DataprepEditorStyle.cpp
--- a/Engine/Plugins/Enterprise/DataprepEditor/Source/DataprepEditor/Private/DataprepEditorStyle.cpp
+++ b/Engine/Plugins/Enterprise/DataprepEditor/Source/DataprepEditor/Private/DataprepEditorStyle.cpp
@@ -91,7 +91,7 @@ void FDataprepEditorStyle::Initialize()
 			TilteTextStyle.SetFont(FCoreStyle::GetDefaultFontStyle("Bold", 11));
 			StyleSet->Set("DataprepAction.TitleTextStyle", TilteTextStyle);
 
-			FEditableTextBoxStyle TitleEditableText = FEditorStyle::Get().GetWidgetStyle<FToolBarStyle>("ViewportMenu").EditableTextStyle;
+			FEditableTextBoxStyle TitleEditableText = FEditorStyle::Get().GetWidgetStyle<FToolBarStyle>("LegacyViewportMenu").EditableTextStyle;
 			TitleEditableText.SetFont(FCoreStyle::GetDefaultFontStyle("Bold", 11));
 
 			StyleSet->Set( "DataprepAction.TitleInlineEditableText", FInlineEditableTextBlockStyle()
diff --git a/Engine/Source/Editor/EditorStyle/Private/StarshipStyle.cpp b/Engine/Source/Editor/EditorStyle/Private/StarshipStyle.cpp
--- a/Engine/Source/Editor/EditorStyle/Private/StarshipStyle.cpp
+++ b/Engine/Source/Editor/EditorStyle/Private/StarshipStyle.cpp
@@ -2320,7 +2320,7 @@ void FStarshipEditorStyle::FStyle::SetupViewportStyles()
 		const FLinearColor NormalColor(1, 1, 1, 0.75f);
 		const FLinearColor PressedColor(1, 1, 1, 1.f);
 
-		/* Create style for "ViewportMenu.ToggleButton" ... */
+		/* Create style for "LegacyViewportMenu.ToggleButton" ... */
 		const FCheckBoxStyle ViewportMenuToggleButtonStyle = FCheckBoxStyle()
 			.SetCheckBoxType(ESlateCheckBoxType::ToggleButton)
 			.SetUncheckedImage(BOX_BRUSH(*SmallRoundedButton, FMargin(7.f / 16.f), NormalColor))
@@ -2332,7 +2332,7 @@ void FStarshipEditorStyle::FStyle::SetupViewportStyles()
 		/* ... and add new style */
 		ViewportToolbarStyle.SetToggleButtonStyle(ViewportMenuToggleButtonStyle);
 
-		/* Create style for "ViewportMenu.ToggleButton.Start" ... */
+		/* Create style for "LegacyViewportMenu.ToggleButton.Start" ... */
 		const FCheckBoxStyle ViewportMenuToggleStartButtonStyle = FCheckBoxStyle()
 			.SetCheckBoxType(ESlateCheckBoxType::ToggleButton)
 			.SetUncheckedImage(BOX_BRUSH(*SmallRoundedButtonStart, FMargin(7.f / 16.f), NormalColor))
@@ -2342,9 +2342,9 @@ void FStarshipEditorStyle::FStyle::SetupViewportStyles()
 			.SetCheckedPressedImage(BOX_BRUSH(*SmallRoundedButtonStart, FMargin(7.f / 16.f), SelectionColor_Pressed))
 			.SetCheckedImage(BOX_BRUSH(*SmallRoundedButtonStart, FMargin(7.f / 16.f), SelectionColor_Pressed));
 		/* ... and add new style */
-		Set("ViewportMenu.ToggleButton.Start", ViewportMenuToggleStartButtonStyle);
+		Set("LegacyViewportMenu.ToggleButton.Start", ViewportMenuToggleStartButtonStyle);
 
-		/* Create style for "ViewportMenu.ToggleButton.Middle" ... */
+		/* Create style for "LegacyViewportMenu.ToggleButton.Middle" ... */
 		const FCheckBoxStyle ViewportMenuToggleMiddleButtonStyle = FCheckBoxStyle()
 			.SetCheckBoxType(ESlateCheckBoxType::ToggleButton)
 			.SetUncheckedImage(BOX_BRUSH(*SmallRoundedButtonMiddle, FMargin(7.f / 16.f), NormalColor))
@@ -2354,9 +2354,9 @@ void FStarshipEditorStyle::FStyle::SetupViewportStyles()
 			.SetCheckedPressedImage(BOX_BRUSH(*SmallRoundedButtonMiddle, FMargin(7.f / 16.f), SelectionColor_Pressed))
 			.SetCheckedImage(BOX_BRUSH(*SmallRoundedButtonMiddle, FMargin(7.f / 16.f), SelectionColor_Pressed));
 		/* ... and add new style */
-		Set("ViewportMenu.ToggleButton.Middle", ViewportMenuToggleMiddleButtonStyle);
+		Set("LegacyViewportMenu.ToggleButton.Middle", ViewportMenuToggleMiddleButtonStyle);
 
-		/* Create style for "ViewportMenu.ToggleButton.End" ... */
+		/* Create style for "LegacyViewportMenu.ToggleButton.End" ... */
 		const FCheckBoxStyle ViewportMenuToggleEndButtonStyle = FCheckBoxStyle()
 			.SetCheckBoxType(ESlateCheckBoxType::ToggleButton)
 			.SetUncheckedImage(BOX_BRUSH(*SmallRoundedButtonEnd, FMargin(7.f / 16.f), NormalColor))
@@ -2366,7 +2366,7 @@ void FStarshipEditorStyle::FStyle::SetupViewportStyles()
 			.SetCheckedPressedImage(BOX_BRUSH(*SmallRoundedButtonEnd, FMargin(7.f / 16.f), SelectionColor_Pressed))
 			.SetCheckedImage(BOX_BRUSH(*SmallRoundedButtonEnd, FMargin(7.f / 16.f), SelectionColor_Pressed));
 		/* ... and add new style */
-		Set("ViewportMenu.ToggleButton.End", ViewportMenuToggleEndButtonStyle);
+		Set("LegacyViewportMenu.ToggleButton.End", ViewportMenuToggleEndButtonStyle);
 
 		const FMargin NormalPadding = FMargin(4.0f, 4.0f, 4.0f, 4.0f);
 		const FMargin PressedPadding = FMargin(4.0f, 4.0f, 4.0f, 4.0f);
@@ -2380,25 +2380,25 @@ void FStarshipEditorStyle::FStyle::SetupViewportStyles()
 
 		ViewportToolbarStyle.SetButtonStyle(ViewportMenuButton);
 
-		Set("ViewportMenu.Button.Start", FButtonStyle(ViewportMenuButton)
+		Set("LegacyViewportMenu.Button.Start", FButtonStyle(ViewportMenuButton)
 			.SetNormal(BOX_BRUSH(*SmallRoundedButtonStart, 7.0f / 16.0f, NormalColor))
 			.SetPressed(BOX_BRUSH(*SmallRoundedButtonStart, 7.0f / 16.0f, PressedColor))
 			.SetHovered(BOX_BRUSH(*SmallRoundedButtonStart, 7.0f / 16.0f, PressedColor))
 		);
 
-		Set("ViewportMenu.Button.Middle", FButtonStyle(ViewportMenuButton)
+		Set("LegacyViewportMenu.Button.Middle", FButtonStyle(ViewportMenuButton)
 			.SetNormal(BOX_BRUSH(*SmallRoundedButtonMiddle, 7.0f / 16.0f, NormalColor))
 			.SetPressed(BOX_BRUSH(*SmallRoundedButtonMiddle, 7.0f / 16.0f, PressedColor))
 			.SetHovered(BOX_BRUSH(*SmallRoundedButtonMiddle, 7.0f / 16.0f, PressedColor))
 		);
 
-		Set("ViewportMenu.Button.End", FButtonStyle(ViewportMenuButton)
+		Set("LegacyViewportMenu.Button.End", FButtonStyle(ViewportMenuButton)
 			.SetNormal(BOX_BRUSH(*SmallRoundedButtonEnd, 7.0f / 16.0f, NormalColor))
 			.SetPressed(BOX_BRUSH(*SmallRoundedButtonEnd, 7.0f / 16.0f, PressedColor))
 			.SetHovered(BOX_BRUSH(*SmallRoundedButtonEnd, 7.0f / 16.0f, PressedColor))
 		);
 
-		Set("ViewportMenu", ViewportToolbarStyle);
+		Set("LegacyViewportMenu", ViewportToolbarStyle);
 	}
 
 	// Viewport actor preview's pin/unpin and attach/detach buttons
diff --git a/Engine/Source/Editor/LevelEditor/Private/SActorPilotViewportToolbar.h b/Engine/Source/Editor/LevelEditor/Private/SActorPilotViewportToolbar.h
--- a/Engine/Source/Editor/LevelEditor/Private/SActorPilotViewportToolbar.h
+++ b/Engine/Source/Editor/LevelEditor/Private/SActorPilotViewportToolbar.h
@@ -61,10 +61,10 @@ public:
 
 		auto& ViewportCommands = FLevelViewportCommands::Get();
 
-		FToolBarBuilder ToolbarBuilder(InArgs._Viewport->GetCommandList(), FMultiBoxCustomization::None, nullptr/*InExtenders*/);
+		FToolBarBuilder ToolbarBuilder(InArgs._Viewport->GetCommandList(), FMultiBoxCustomization::None);
 
 		// Use a custom style
-		FName ToolBarStyle = "ViewportMenu";
+		FName ToolBarStyle = "LegacyViewportMenu";
 		ToolbarBuilder.SetStyle(&FEditorStyle::Get(), ToolBarStyle);
 		ToolbarBuilder.SetStyle(&FEditorStyle::Get(), ToolBarStyle);
 		ToolbarBuilder.SetLabelVisibility(EVisibility::Collapsed);
diff --git a/Engine/Source/Editor/MaterialEditor/Private/SMaterialEditorViewportToolBar.cpp b/Engine/Source/Editor/MaterialEditor/Private/SMaterialEditorViewportToolBar.cpp
--- a/Engine/Source/Editor/MaterialEditor/Private/SMaterialEditorViewportToolBar.cpp
+++ b/Engine/Source/Editor/MaterialEditor/Private/SMaterialEditorViewportToolBar.cpp
@@ -14,10 +14,10 @@ void SMaterialEditorViewportPreviewShapeToolBar::Construct(const FArguments& InA
 {
 	// Force this toolbar to have small icons, as the preview panel is only small so we have limited space
 	const bool bForceSmallIcons = true;
-	FToolBarBuilder ToolbarBuilder(InViewport->GetCommandList(), FMultiBoxCustomization::None, nullptr, Orient_Horizontal, bForceSmallIcons);
+	FToolBarBuilder ToolbarBuilder(InViewport->GetCommandList(), FMultiBoxCustomization::None, nullptr, bForceSmallIcons);
 
 	// Use a custom style
-	ToolbarBuilder.SetStyle(&FEditorStyle::Get(), "ViewportMenu");
+	ToolbarBuilder.SetStyle(&FEditorStyle::Get(), "LegacyViewportMenu");
 	ToolbarBuilder.SetLabelVisibility(EVisibility::Collapsed);
 	ToolbarBuilder.SetIsFocusable(false);
 	
diff --git a/Engine/Source/Editor/TextureEditor/Private/Widgets/STextureEditorViewportToolbar.cpp b/Engine/Source/Editor/TextureEditor/Private/Widgets/STextureEditorViewportToolbar.cpp
--- a/Engine/Source/Editor/TextureEditor/Private/Widgets/STextureEditorViewportToolbar.cpp
+++ b/Engine/Source/Editor/TextureEditor/Private/Widgets/STextureEditorViewportToolbar.cpp
@@ -39,7 +39,7 @@ void STextureEditorViewportToolbar::Construct( const FArguments& InArgs, const T
 							.ClickMethod(EButtonClickMethod::MouseDown)
 							.ContentPadding(FMargin(5.0f, 2.0f))
 							.VAlign(VAlign_Center)
-							.ButtonStyle(FEditorStyle::Get(), "ViewPortMenu.Button")
+							.ButtonStyle(FEditorStyle::Get(), "LegacyViewportMenu.Button")
 							.OnClicked(this, &STextureEditorViewportToolbar::HandleViewOptionsMenuButtonClicked)
 							[
 								SNew(SHorizontalBox)
@@ -49,7 +49,7 @@ void STextureEditorViewportToolbar::Construct( const FArguments& InArgs, const T
 									[
 										SNew(STextBlock)
 										.Text(LOCTEXT("ViewButtonText", "View"))
-										.TextStyle(FEditorStyle::Get(), "ViewportMenu.Label")
+										.TextStyle(FEditorStyle::Get(), "LegacyViewportMenu.Label")
 									]
 
 								+ SHorizontalBox::Slot()
diff --git a/Engine/Source/Editor/UMGEditor/Private/Designer/SDesignerToolBar.cpp b/Engine/Source/Editor/UMGEditor/Private/Designer/SDesignerToolBar.cpp
--- a/Engine/Source/Editor/UMGEditor/Private/Designer/SDesignerToolBar.cpp
+++ b/Engine/Source/Editor/UMGEditor/Private/Designer/SDesignerToolBar.cpp
@@ -39,7 +39,7 @@ TSharedRef< SWidget > SDesignerToolBar::MakeToolBar(const TSharedPtr< FExtender
 	FToolBarBuilder ToolbarBuilder( CommandList, FMultiBoxCustomization::None, InExtenders );
 
 	// Use a custom style
-	FName ToolBarStyle = "ViewportMenu";
+	FName ToolBarStyle = "LegacyViewportMenu";
 	ToolbarBuilder.SetStyle(&FEditorStyle::Get(), ToolBarStyle);
 	ToolbarBuilder.SetLabelVisibility(EVisibility::Collapsed);
 
diff --git a/Engine/Source/Editor/UMGEditor/Private/Designer/SDesignerView.cpp b/Engine/Source/Editor/UMGEditor/Private/Designer/SDesignerView.cpp
--- a/Engine/Source/Editor/UMGEditor/Private/Designer/SDesignerView.cpp
+++ b/Engine/Source/Editor/UMGEditor/Private/Designer/SDesignerView.cpp
@@ -538,7 +538,7 @@ EActiveTimerReturnType SDesignerView::EnsureTick(double InCurrentTime, float InD
 
 TSharedRef<SWidget> SDesignerView::CreateOverlayUI()
 {
-	const FToolBarStyle& ToolBarStyle = FEditorStyle::Get().GetWidgetStyle<FToolBarStyle>("ViewportMenu");
+	const FToolBarStyle& ToolBarStyle = FEditorStyle::Get().GetWidgetStyle<FToolBarStyle>("LegacyViewportMenu");
 
 	return SNew(SOverlay)
 
@@ -629,7 +629,7 @@ TSharedRef<SWidget> SDesignerView::CreateOverlayUI()
 				.ButtonStyle(&ToolBarStyle.ButtonStyle)
 				.ToolTipText(LOCTEXT("ZoomToFit_ToolTip", "Zoom To Fit"))
 				.OnClicked(this, &SDesignerView::HandleZoomToFitClicked)
-				.ContentPadding(FEditorStyle::Get().GetMargin("ViewportMenu.SToolBarButtonBlock.Button.Padding"))
+				.ContentPadding(FEditorStyle::Get().GetMargin("LegacyViewportMenu.SToolBarButtonBlock.Button.Padding"))
 				[
 					SNew(SImage)
 					.Image(FEditorStyle::GetBrush("UMGEditor.ZoomToFit"))
@@ -644,7 +644,7 @@ TSharedRef<SWidget> SDesignerView::CreateOverlayUI()
 				.ButtonStyle(&ToolBarStyle.ButtonStyle)
 				.ToolTipText(LOCTEXT("SwapAspectRatio_ToolTip", "Switch between Landscape and Portrait"))
 				.OnClicked(this, &SDesignerView::HandleSwapAspectRatioClicked)
-				.ContentPadding(FEditorStyle::Get().GetMargin("ViewportMenu.SToolBarButtonBlock.Button.Padding"))
+				.ContentPadding(FEditorStyle::Get().GetMargin("LegacyViewportMenu.SToolBarButtonBlock.Button.Padding"))
 				.IsEnabled(this, &SDesignerView::GetAspectRatioSwitchEnabled)
 				[
 					SNew(SImage)
@@ -659,7 +659,7 @@ TSharedRef<SWidget> SDesignerView::CreateOverlayUI()
 				.ButtonStyle(&ToolBarStyle.ButtonStyle)
 				.ToolTipText(LOCTEXT("Mirror_ToolTip", "Flip the current safe zones"))
 				.OnClicked(this, &SDesignerView::HandleFlipSafeZonesClicked)
-				.ContentPadding(FEditorStyle::Get().GetMargin("ViewportMenu.SToolBarButtonBlock.Button.Padding"))	
+				.ContentPadding(FEditorStyle::Get().GetMargin("LegacyViewportMenu.SToolBarButtonBlock.Button.Padding"))	
 				.IsEnabled(this, &SDesignerView::GetFlipDeviceEnabled)
 				[
 					SNew(SImage)
@@ -676,12 +676,12 @@ TSharedRef<SWidget> SDesignerView::CreateOverlayUI()
 			.ButtonStyle(&ToolBarStyle.ButtonStyle)
 			.ForegroundColor(FLinearColor::Black)
 			.OnGetMenuContent(this, &SDesignerView::GetResolutionsMenu)
-			.ContentPadding(FEditorStyle::Get().GetMargin("ViewportMenu.SToolBarButtonBlock.Button.Padding"))
+			.ContentPadding(FEditorStyle::Get().GetMargin("LegacyViewportMenu.SToolBarButtonBlock.Button.Padding"))
 			.ButtonContent()
 			[
 				SNew(STextBlock)
 				.Text(LOCTEXT("ScreenSize", "Screen Size"))
-				.TextStyle(FEditorStyle::Get(), "ViewportMenu.Label")
+				.TextStyle(FEditorStyle::Get(), "LegacyViewportMenu.Label")
 			]
 		]
 
@@ -694,12 +694,12 @@ TSharedRef<SWidget> SDesignerView::CreateOverlayUI()
 			.ButtonStyle(&ToolBarStyle.ButtonStyle)
 			.ForegroundColor(FLinearColor::Black)
 			.OnGetMenuContent(this, &SDesignerView::GetScreenSizingFillMenu)
-			.ContentPadding(FEditorStyle::Get().GetMargin("ViewportMenu.SToolBarButtonBlock.Button.Padding"))
+			.ContentPadding(FEditorStyle::Get().GetMargin("LegacyViewportMenu.SToolBarButtonBlock.Button.Padding"))
 			.ButtonContent()
 			[
 				SNew(STextBlock)
 				.Text(this, &SDesignerView::GetScreenSizingFillText)
-				.TextStyle(FEditorStyle::Get(), "ViewportMenu.Label")
+				.TextStyle(FEditorStyle::Get(), "LegacyViewportMenu.Label")
 			]
 		]
 
diff --git a/Engine/Source/Editor/UnrealEd/Private/STransformViewportToolbar.cpp b/Engine/Source/Editor/UnrealEd/Private/STransformViewportToolbar.cpp
--- a/Engine/Source/Editor/UnrealEd/Private/STransformViewportToolbar.cpp
+++ b/Engine/Source/Editor/UnrealEd/Private/STransformViewportToolbar.cpp
@@ -199,7 +199,7 @@ TSharedRef< SWidget > STransformViewportToolBar::MakeTransformToolBar( const TSh
 	FToolBarBuilder ToolbarBuilder( CommandList, FMultiBoxCustomization::None, InExtenders );
 
 	// Use a custom style
-	FName ToolBarStyle = "ViewportMenu";
+	FName ToolBarStyle = "LegacyViewportMenu";
 	ToolbarBuilder.SetStyle(&FEditorStyle::Get(), ToolBarStyle);
 	ToolbarBuilder.SetLabelVisibility(EVisibility::Collapsed);
 
