Scene Outliner Styling

diff --git a/Engine/Source/Editor/EditorStyle/Private/StarshipStyle.cpp b/Engine/Source/Editor/EditorStyle/Private/StarshipStyle.cpp
--- a/Engine/Source/Editor/EditorStyle/Private/StarshipStyle.cpp
+++ b/Engine/Source/Editor/EditorStyle/Private/StarshipStyle.cpp
@@ -1116,20 +1116,13 @@ void FStarshipEditorStyle::FStyle::SetupGeneralStyles()
 
 	// Scene Outliner
 	{
-		Set( "SceneOutliner.FilterSearch", new IMAGE_BRUSH( "Old/FilterSearch", Icon16x16 ) );
-		Set( "SceneOutliner.FilterCancel", new IMAGE_BRUSH( "Old/FilterCancel", Icon16x16 ) );
-		Set( "SceneOutliner.FolderClosed", new IMAGE_BRUSH( "Icons/FolderClosed", Icon16x16 ) );
-		Set( "SceneOutliner.FolderOpen", new IMAGE_BRUSH( "Icons/FolderOpen", Icon16x16 ) );
-		Set( "SceneOutliner.NewFolderIcon", new IMAGE_BRUSH("Icons/icon_AddFolder_16x", Icon16x16 ) );
-		Set( "SceneOutliner.MoveToRoot", new IMAGE_BRUSH("Icons/icon_NoFolder_16x", Icon16x16 ) );
-		Set( "SceneOutliner.ChangedItemHighlight", new BOX_BRUSH( "Common/EditableTextSelectionBackground", FMargin(4.f/16.f) ) );
-		Set( "SceneOutliner.World", new IMAGE_BRUSH( "Icons/icon_world_16x", Icon16x16 ) );
-
-		// Selection color should still be orange to align with the editor viewport.
-		// But must also give the hint that the tree is no longer focused.
-		Set( "SceneOutliner.TableViewRow", FTableRowStyle(NormalTableRowStyle)
-			.SetInactiveBrush( IMAGE_BRUSH( "Common/Selection", Icon8x8, SelectionColor_Subdued ) )
-		);
+		Set( "SceneOutliner.NewFolderIcon", new CORE_IMAGE_BRUSH_SVG("Starship/Common/folder-plus", Icon16x16 ) );
+		Set( "SceneOutliner.FolderClosed",  new CORE_IMAGE_BRUSH_SVG("Starship/SceneOutliner/FolderClosed", Icon16x16, FStyleColors::AccentFolder ) );
+		Set( "SceneOutliner.FolderOpen",    new CORE_IMAGE_BRUSH_SVG("Starship/SceneOutliner/FolderOpen", Icon16x16, FStyleColors::AccentFolder ) );
+		Set( "SceneOutliner.World", 		new CORE_IMAGE_BRUSH_SVG("Starship/SceneOutliner/World", Icon16x16 ) );
+		Set( "SceneOutliner.ChangedItemHighlight", new FSlateRoundedBoxBrush( FStyleColors::White, 1.0) );
+
+		Set( "SceneOutliner.TableViewRow", FTableRowStyle(NormalTableRowStyle));
 	}
 
 	// Socket chooser
@@ -1885,10 +1878,11 @@ void FStarshipEditorStyle::FStyle::SetupLevelGeneralStyles()
 {
 // Levels General
 	{
-	Set("Level.VisibleIcon16x", new IMAGE_BRUSH("Icons/icon_levels_visible_16px", Icon16x16));
-	Set("Level.VisibleHighlightIcon16x", new IMAGE_BRUSH("Icons/icon_levels_visible_hi_16px", Icon16x16));
-	Set("Level.NotVisibleIcon16x", new IMAGE_BRUSH("Icons/icon_levels_invisible_16px", Icon16x16));
-	Set("Level.NotVisibleHighlightIcon16x", new IMAGE_BRUSH("Icons/icon_levels_invisible_hi_16px", Icon16x16));
+	Set("Level.VisibleIcon16x",             new CORE_IMAGE_BRUSH_SVG("Starship/SceneOutliner/Visible", Icon16x16));
+	Set("Level.VisibleHighlightIcon16x",    new CORE_IMAGE_BRUSH_SVG("Starship/SceneOutliner/Visible", Icon16x16));
+	Set("Level.NotVisibleIcon16x",          new CORE_IMAGE_BRUSH_SVG("Starship/SceneOutliner/Hidden",  Icon16x16));
+	Set("Level.NotVisibleHighlightIcon16x", new CORE_IMAGE_BRUSH_SVG("Starship/SceneOutliner/Hidden",  Icon16x16));
+
 	Set("Level.LightingScenarioIcon16x", new IMAGE_BRUSH("Icons/icon_levels_LightingScenario_16px", Icon16x16));
 	Set("Level.LightingScenarioNotIcon16x", new IMAGE_BRUSH("Icons/icon_levels_LightingScenarioNot_16px", Icon16x16));
 	Set("Level.LockedIcon16x", new IMAGE_BRUSH("Icons/icon_locked_16px", Icon16x16));
diff --git a/Engine/Source/Editor/SceneOutliner/Private/SSceneOutliner.cpp b/Engine/Source/Editor/SceneOutliner/Private/SSceneOutliner.cpp
--- a/Engine/Source/Editor/SceneOutliner/Private/SSceneOutliner.cpp
+++ b/Engine/Source/Editor/SceneOutliner/Private/SSceneOutliner.cpp
@@ -377,7 +377,7 @@ namespace SceneOutliner
 				[
 					SNew(SImage)
 	                .ColorAndOpacity(FSlateColor::UseForeground())
-					.Image( FAppStyle::Get().GetBrush("Icons.cogwheel") )
+					.Image( FAppStyle::Get().GetBrush("Icons.Settings") )
 				]
 			];
 		}
diff --git a/Engine/Source/Editor/SceneOutliner/Private/SceneOutlinerItemLabelColumn.cpp b/Engine/Source/Editor/SceneOutliner/Private/SceneOutlinerItemLabelColumn.cpp
--- a/Engine/Source/Editor/SceneOutliner/Private/SceneOutlinerItemLabelColumn.cpp
+++ b/Engine/Source/Editor/SceneOutliner/Private/SceneOutlinerItemLabelColumn.cpp
@@ -99,6 +99,7 @@ struct SActorTreeLabel : FCommonLabelData, public SCompoundWidget
 
 			// Main actor label
 			+ SHorizontalBox::Slot()
+			.VAlign(VAlign_Center)
 			[
 				SAssignNew(InlineTextBlock, SInlineEditableTextBlock)
 				.Text(this, &SActorTreeLabel::GetDisplayText)
@@ -115,6 +116,7 @@ struct SActorTreeLabel : FCommonLabelData, public SCompoundWidget
 			]
 
 			+ SHorizontalBox::Slot()
+			.VAlign(VAlign_Center)
 			.AutoWidth()
 			.Padding( 0.0f, 0.f, 3.0f, 0.0f )
 			[
@@ -176,7 +178,7 @@ struct SActorTreeLabel : FCommonLabelData, public SCompoundWidget
 			+ SHorizontalBox::Slot()
 			.FillWidth(1.0f)
 			.VAlign(VAlign_Center)
-			.Padding(0.0f, 2.0f)
+			.Padding(0.0f, 0.0f)
 			[
 				MainContent
 			]
