Details View Style - First Pass

diff --git a/Engine/Source/Editor/EditorStyle/Private/StarshipStyle.cpp b/Engine/Source/Editor/EditorStyle/Private/StarshipStyle.cpp
--- a/Engine/Source/Editor/EditorStyle/Private/StarshipStyle.cpp
+++ b/Engine/Source/Editor/EditorStyle/Private/StarshipStyle.cpp
@@ -3044,8 +3044,8 @@ void FStarshipEditorStyle::FStyle::SetupPropertyEditorStyles()
 		Set( "PropertyWindow.Button_PickActor", new IMAGE_BRUSH( "Icons/levels_16x", Icon12x12 ) );
 		Set( "PropertyWindow.Button_PickActorInteractive", new IMAGE_BRUSH( "Icons/eyedropper_16px", Icon12x12 ) );
 		Set( "PropertyWindow.Button_Refresh", new IMAGE_BRUSH("Icons/refresh_12x", Icon12x12 ) );
-
 		Set( "PropertyWindow.WindowBorder", new BOX_BRUSH( "Common/GroupBorder", FMargin(4.0f/16.0f) ) );
+
 		Set( "DetailsView.NameChangeCommitted", new BOX_BRUSH( "Common/EditableTextSelectionBackground", FMargin(4.f/16.f) ) );
 		Set( "DetailsView.HyperlinkStyle", FTextBlockStyle(NormalText) .SetFont( DEFAULT_FONT( "Regular", 8 ) ) );
 
@@ -3066,43 +3066,43 @@ void FStarshipEditorStyle::FStyle::SetupPropertyEditorStyles()
 		Set("DetailsView.BPMessageTextStyle", BPWarningMessageTextStyle);
 
 
-		Set( "DetailsView.GroupSection", new BOX_BRUSH( "Common/RoundedSelection_16x", FMargin(4.0f/16.0f) ) );
+		Set( "DetailsView.GroupSection",              new FSlateNoResource());
 
-		Set( "DetailsView.PulldownArrow.Down", new IMAGE_BRUSH( "PropertyView/AdvancedButton_Down", FVector2D(10,8) ) );
-		Set( "DetailsView.PulldownArrow.Down.Hovered", new IMAGE_BRUSH( "PropertyView/AdvancedButton_Down_Hovered", FVector2D(10,8) ) );
-		Set( "DetailsView.PulldownArrow.Up", new IMAGE_BRUSH( "PropertyView/AdvancedButton_Up", FVector2D(10,8) ) );
-		Set( "DetailsView.PulldownArrow.Up.Hovered", new IMAGE_BRUSH( "PropertyView/AdvancedButton_Up_Hovered", FVector2D(10,8) ) );
+		Set( "DetailsView.PulldownArrow.Down",        new CORE_IMAGE_BRUSH_SVG("Starship/Common/chevron-down", Icon16x16, FStyleColors::Foreground)); 
+		Set( "DetailsView.PulldownArrow.Down.Hovered",new CORE_IMAGE_BRUSH_SVG("Starship/Common/chevron-down", Icon16x16, FStyleColors::ForegroundHover)); 
+		Set( "DetailsView.PulldownArrow.Up",          new CORE_IMAGE_BRUSH_SVG("Starship/Common/chevron-up",   Icon16x16, FStyleColors::Foreground)); 
+		Set( "DetailsView.PulldownArrow.Up.Hovered",  new CORE_IMAGE_BRUSH_SVG("Starship/Common/chevron-up",   Icon16x16, FStyleColors::ForegroundHover)); 
 
-		Set( "DetailsView.EditRawProperties", new IMAGE_BRUSH( "Icons/icon_PropertyMatrix_16px", Icon16x16, FLinearColor::Black ) );
-		Set( "DetailsView.EditConfigProperties", new IMAGE_BRUSH( "Icons/icon_PropertyMatrix_16px", Icon16x16, FLinearColor::White ) );
+		Set( "DetailsView.EditRawProperties",         new IMAGE_BRUSH("Icons/icon_PropertyMatrix_16px",  Icon16x16, FLinearColor::Black ) );
+		Set( "DetailsView.EditConfigProperties",      new IMAGE_BRUSH("Icons/icon_PropertyMatrix_16px",  Icon16x16, FLinearColor::White ) );
 
-		Set( "DetailsView.CollapsedCategory", new BOX_BRUSH( "Common/GroupBorder", FMargin(4.0f/16.0f) ) );
-		Set( "DetailsView.CategoryTop", new BOX_BRUSH( "PropertyView/DetailCategoryTop", FMargin( 4/16.0f, 8.0f/16.0f, 4/16.0f, 4/16.0f ) ) );
-		Set( "DetailsView.CollapsedCategory_Hovered", new BOX_BRUSH( "Common/GroupBorder", FMargin(4.0f/16.0f), FLinearColor(0.5f,0.5f,0.5f,1.0f)  ) );
-		Set( "DetailsView.CategoryTop_Hovered", new BOX_BRUSH( "PropertyView/DetailCategoryTop", FMargin( 4/16.0f, 8.0f/16.0f, 4/16.0f, 4/16.0f ), FLinearColor(0.5f,0.5f,0.5f,1.0f) ) );
-		Set( "DetailsView.CategoryBottom", new BOX_BRUSH( "PropertyView/DetailCategoryBottom", FMargin(4.0f/16.0f) ) );
-		Set( "DetailsView.CategoryMiddle", new IMAGE_BRUSH( "PropertyView/DetailCategoryMiddle", FVector2D( 16, 16 ) ) );
-		Set( "DetailsView.CategoryMiddle_Hovered", new IMAGE_BRUSH( "PropertyView/DetailCategoryMiddle_Hovered", FVector2D( 16, 16 ) ) );
-		Set( "DetailsView.CategoryMiddle_Highlighted", new BOX_BRUSH( "Common/TextBox_Special_Active", FMargin(8.0f/32.0f) ) );
-		Set( "DetailsView.CategoryMiddle_Active", new BOX_BRUSH( "Common/TextBox_Special_Active", FMargin(8.0f/32.0f), SelectionColor_Pressed ) );
+		Set( "DetailsView.CollapsedCategory",         new FSlateColorBrush(FStyleColors::Background));
+		Set( "DetailsView.CategoryTop",               new FSlateColorBrush(FStyleColors::Background));
+		Set( "DetailsView.CollapsedCategory_Hovered", new FSlateColorBrush(FStyleColors::Header));
+		Set( "DetailsView.CategoryTop_Hovered",       new FSlateColorBrush(FStyleColors::Header));
+		Set( "DetailsView.CategoryBottom",            new FSlateColorBrush(FStyleColors::Input));
+
+		Set( "DetailsView.CategoryMiddle",            new FSlateColorBrush(FStyleColors::Background));
+		Set( "DetailsView.CategoryMiddle_Hovered",    new FSlateColorBrush(FStyleColors::Header));
+		Set( "DetailsView.CategoryMiddle_Highlighted",new FSlateColorBrush(FStyleColors::Background));
+		Set( "DetailsView.CategoryMiddle_Active",     new FSlateColorBrush(FStyleColors::Background));
 		
 		Set( "DetailsView.PropertyIsFavorite", new IMAGE_BRUSH("PropertyView/Favorites_Enabled", Icon12x12));
 		Set( "DetailsView.PropertyIsNotFavorite", new IMAGE_BRUSH("PropertyView/Favorites_Disabled", Icon12x12));
 		Set( "DetailsView.NoFavoritesSystem", new IMAGE_BRUSH("PropertyView/NoFavoritesSystem", Icon12x12));
 
 		Set( "DetailsView.Splitter", FSplitterStyle()
-			.SetHandleNormalBrush( IMAGE_BRUSH( "Common/SplitterHandleHighlight", Icon8x8, FLinearColor::Black ) )
-			.SetHandleHighlightBrush( IMAGE_BRUSH( "Common/SplitterHandleHighlight", Icon8x8, FLinearColor::White ) )
-			);
+			.SetHandleNormalBrush(   FSlateNoResource())                   
+			.SetHandleHighlightBrush(FSlateColorBrush(FStyleColors::Input))
+		);
 
-		Set( "DetailsView.AdvancedDropdownBorder", new BOX_BRUSH( "PropertyView/DetailCategoryAdvanced", FMargin(4.0f/16.0f) ) );
-		Set( "DetailsView.AdvancedDropdownBorder.Open", new IMAGE_BRUSH( "Common/ScrollBoxShadowTop", FVector2D(64,8) ) );
-		Set( "DetailsView.CategoryFontStyle", DEFAULT_FONT( "Bold", 10 ) );
+		Set( "DetailsView.AdvancedDropdownBorder",      new FSlateColorBrush(FStyleColors::Background));
+		Set( "DetailsView.AdvancedDropdownBorder.Open", new FSlateColorBrush(FStyleColors::Background));
 
-		Set( "DetailsView.CategoryTextStyle", 
-			FTextBlockStyle(NormalText)
+		Set( "DetailsView.CategoryFontStyle", FStyleFonts::Get().NormalBold);
+		Set( "DetailsView.CategoryTextStyle", FTextBlockStyle(NormalText)
 			.SetFont(GetFontStyle("DetailsView.CategoryFontStyle"))
-			.SetShadowOffset(FVector2D(1.0f, 1.0f))
+			.SetColorAndOpacity(FStyleColors::White)
 		);
 
 		Set( "DetailsView.TreeView.TableRow", FTableRowStyle()
@@ -4017,7 +4017,7 @@ void FStarshipEditorStyle::FStyle::SetupGraphEditorStyles()
 			.SetShadowOffset(FVector2D(1, 1))
 			.SetShadowColorAndOpacity(FLinearColor(0, 0, 0, 0.9f)));
 
-		Set("SCSEditor.TreePanel", new BOX_BRUSH("Common/GroupBorder_FlatTop", FMargin(4.0f / 16.0f)));
+		Set("SCSEditor.TreePanel", new FSlateNoResource());
 
 		//
 
diff --git a/Engine/Source/Editor/PropertyEditor/Private/DetailAdvancedDropdownNode.cpp b/Engine/Source/Editor/PropertyEditor/Private/DetailAdvancedDropdownNode.cpp
--- a/Engine/Source/Editor/PropertyEditor/Private/DetailAdvancedDropdownNode.cpp
+++ b/Engine/Source/Editor/PropertyEditor/Private/DetailAdvancedDropdownNode.cpp
@@ -124,7 +124,7 @@ public:
 
 			ContentWidget =
 				SNew( SBorder )
-				.BorderImage( FEditorStyle::GetBrush("DetailsView.CategoryBottom") )
+				.BorderImage( FEditorStyle::GetBrush("DetailsView.CategoryMiddle") )
 				.Padding( FMargin( 0.0f, 0.0f, SDetailTableRowBase::ScrollbarPaddingSize, 2.0f ) )
 				[
 					SplitterArea.ToSharedRef()	
@@ -133,7 +133,21 @@ public:
 		
 		ChildSlot
 		[
-			ContentWidget.ToSharedRef()
+			SNew( SVerticalBox)
+			+ SVerticalBox::Slot()
+			[
+				ContentWidget.ToSharedRef()
+			]
+			+ SVerticalBox::Slot()
+			.AutoHeight()
+			[
+        		SNew(SBox)
+        		.HeightOverride(2.f)
+        		[ 
+					SNew( SBorder )
+					.BorderImage( FEditorStyle::GetBrush("DetailsView.CategoryBottom") )
+				]
+			]
 		];
 
 		STableRow< TSharedPtr< FDetailTreeNode > >::ConstructInternal(
diff --git a/Engine/Source/Editor/PropertyEditor/Private/DetailCategoryGroupNode.cpp b/Engine/Source/Editor/PropertyEditor/Private/DetailCategoryGroupNode.cpp
--- a/Engine/Source/Editor/PropertyEditor/Private/DetailCategoryGroupNode.cpp
+++ b/Engine/Source/Editor/PropertyEditor/Private/DetailCategoryGroupNode.cpp
@@ -44,9 +44,10 @@ void SDetailCategoryTableRow::Construct( const FArguments& InArgs, TSharedRef<FD
 			.AutoWidth()
 			[
 				SNew(STextBlock)
+				.TransformPolicy(ETextTransformPolicy::ToUpper)
 				.Text(InArgs._DisplayName)
-				.Font(FEditorStyle::GetFontStyle(bIsInnerCategory ? "PropertyWindow.NormalFont" : "DetailsView.CategoryFontStyle"))
-				.ShadowOffset(bIsInnerCategory ? FVector2D::ZeroVector : FVector2D(1.0f, 1.0f))
+				.Font(FAppStyle::Get().GetFontStyle(bIsInnerCategory ? "PropertyWindow.NormalFont" : "DetailsView.CategoryFontStyle"))
+				.TextStyle(FAppStyle::Get(), "DetailsView.CategoryTextStyle")
 			];
 
 	// If column size data was specified, add a separator
@@ -86,7 +87,6 @@ void SDetailCategoryTableRow::Construct( const FArguments& InArgs, TSharedRef<FD
 		SNew( SBorder )
 		.BorderImage( this, &SDetailCategoryTableRow::GetBackgroundImage )
 		.Padding( FMargin( 0.0f, BorderVerticalPadding, SDetailTableRowBase::ScrollbarPaddingSize, BorderVerticalPadding ) )
-		.BorderBackgroundColor( FLinearColor( .6,.6,.6, 1.0f ) )
 		[
 			Widget.ToSharedRef()
 		]
diff --git a/Engine/Source/Editor/PropertyEditor/Private/DetailPropertyRow.cpp b/Engine/Source/Editor/PropertyEditor/Private/DetailPropertyRow.cpp
--- a/Engine/Source/Editor/PropertyEditor/Private/DetailPropertyRow.cpp
+++ b/Engine/Source/Editor/PropertyEditor/Private/DetailPropertyRow.cpp
@@ -666,6 +666,7 @@ void FDetailPropertyRow::MakeNameOrKeyWidget( FDetailWidgetRow& Row, const TShar
 		.AutoWidth()
 		.Padding( 0.0f, 0.0f )
 		.VAlign(VAlign_Center)
+		.HAlign(HAlign_Right)
 		[
 			SNew( SEditConditionWidget, PropertyEditor )
 			.CustomEditCondition( CustomEditCondition.IsValid() ? *CustomEditCondition : FCustomEditCondition() )
@@ -702,6 +703,7 @@ void FDetailPropertyRow::MakeNameOrKeyWidget( FDetailWidgetRow& Row, const TShar
 	}
 
 	SHorizontalBox::FSlot& Slot = NameHorizontalBox->AddSlot()
+	.HAlign(HAlign_Right)
 	[
 		NameWidget.ToSharedRef()
 	];
@@ -721,7 +723,7 @@ void FDetailPropertyRow::MakeNameOrKeyWidget( FDetailWidgetRow& Row, const TShar
 	}
 
 	Row.NameContent()
-	.HAlign( HorizontalAlignment )
+	.HAlign( HAlign_Right )
 	.VAlign( VerticalAlignment )
 	[
 		NameHorizontalBox
