Fix wrong border on viewport toolbar hover

diff --git a/Engine/Source/Editor/EditorStyle/Private/StarshipStyle.cpp b/Engine/Source/Editor/EditorStyle/Private/StarshipStyle.cpp
--- a/Engine/Source/Editor/EditorStyle/Private/StarshipStyle.cpp
+++ b/Engine/Source/Editor/EditorStyle/Private/StarshipStyle.cpp
@@ -2350,8 +2350,6 @@ void FStarshipEditorStyle::FStyle::SetupViewportStyles()
 		);
 
 		
-		Set("EditorViewportToolBar.Button", HoverHintOnly);
-
 		Set("EditorViewportToolBar.MenuDropdown", new IMAGE_BRUSH("Common/ComboArrow", Icon8x8));
 		Set("EditorViewportToolBar.Maximize.Normal", new IMAGE_BRUSH_SVG("Starship/EditorViewport/maximize", Icon16x16));
 		Set("EditorViewportToolBar.Maximize.Checked", new IMAGE_BRUSH_SVG("Starship/EditorViewport/restore", Icon16x16));
