Level Editor Menu Icons

diff --git a/Engine/Plugins/Animation/LiveLink/Source/LiveLinkEditor/Private/LiveLinkEditorModule.cpp b/Engine/Plugins/Animation/LiveLink/Source/LiveLinkEditor/Private/LiveLinkEditorModule.cpp
--- a/Engine/Plugins/Animation/LiveLink/Source/LiveLinkEditor/Private/LiveLinkEditorModule.cpp
+++ b/Engine/Plugins/Animation/LiveLink/Source/LiveLinkEditor/Private/LiveLinkEditorModule.cpp
@@ -61,6 +61,7 @@ namespace LiveLinkEditorModuleUtils
 }
 
 #define IMAGE_PLUGIN_BRUSH( RelativePath, ... ) FSlateImageBrush( LiveLinkEditorModuleUtils::InPluginContent( RelativePath, ".png" ), __VA_ARGS__ )
+#define IMAGE_PLUGIN_BRUSH_SVG( RelativePath, ... ) FSlateVectorImageBrush( LiveLinkEditorModuleUtils::InPluginContent( RelativePath, ".svg" ), __VA_ARGS__ )
 
 
 class FLiveLinkEditorModule : public IModuleInterface
@@ -94,15 +95,15 @@ public:
 		StyleSet->SetCoreContentRoot(FPaths::EngineContentDir() / TEXT("Slate"));
 
 		StyleSet->Set("LiveLinkClient.Common.Icon", new IMAGE_PLUGIN_BRUSH(TEXT("LiveLink_40x"), Icon40x40));
-		StyleSet->Set("LiveLinkClient.Common.Icon.Small", new IMAGE_PLUGIN_BRUSH(TEXT("LiveLink_16x"), Icon16x16));
-
-		StyleSet->Set("ClassIcon.LiveLinkPreset", new IMAGE_PLUGIN_BRUSH("LiveLink_16x", Icon16x16));
-		StyleSet->Set("ClassIcon.LiveLinkFrameInterpolationProcessor", new IMAGE_PLUGIN_BRUSH("LiveLink_16x", Icon16x16));
-		StyleSet->Set("ClassIcon.LiveLinkFramePreProcessor", new IMAGE_PLUGIN_BRUSH("LiveLink_16x", Icon16x16));
-		StyleSet->Set("ClassIcon.LiveLinkFrameTranslator", new IMAGE_PLUGIN_BRUSH("LiveLink_16x", Icon16x16));
-		StyleSet->Set("ClassIcon.LiveLinkPreset", new IMAGE_PLUGIN_BRUSH("LiveLink_16x", Icon16x16));
-		StyleSet->Set("ClassIcon.LiveLinkRole", new IMAGE_PLUGIN_BRUSH("LiveLink_16x", Icon16x16));
-		StyleSet->Set("ClassIcon.LiveLinkVirtualSubject", new IMAGE_PLUGIN_BRUSH("LiveLink_16x", Icon16x16));
+		StyleSet->Set("LiveLinkClient.Common.Icon.Small", new IMAGE_PLUGIN_BRUSH_SVG("Starship/LiveLink", Icon16x16));
+
+		StyleSet->Set("ClassIcon.LiveLinkPreset",                      new IMAGE_PLUGIN_BRUSH_SVG("Starship/LiveLink", Icon16x16));
+		StyleSet->Set("ClassIcon.LiveLinkFrameInterpolationProcessor", new IMAGE_PLUGIN_BRUSH_SVG("Starship/LiveLink", Icon16x16));
+		StyleSet->Set("ClassIcon.LiveLinkFramePreProcessor",           new IMAGE_PLUGIN_BRUSH_SVG("Starship/LiveLink", Icon16x16));
+		StyleSet->Set("ClassIcon.LiveLinkFrameTranslator",             new IMAGE_PLUGIN_BRUSH_SVG("Starship/LiveLink", Icon16x16));
+		StyleSet->Set("ClassIcon.LiveLinkPreset",                      new IMAGE_PLUGIN_BRUSH_SVG("Starship/LiveLink", Icon16x16));
+		StyleSet->Set("ClassIcon.LiveLinkRole",                        new IMAGE_PLUGIN_BRUSH_SVG("Starship/LiveLink", Icon16x16));
+		StyleSet->Set("ClassIcon.LiveLinkVirtualSubject",              new IMAGE_PLUGIN_BRUSH_SVG("Starship/LiveLink", Icon16x16));
 
 		StyleSet->Set("ClassThumbnail.LiveLinkPreset", new IMAGE_PLUGIN_BRUSH("LiveLink_40x", Icon40x40));
 
diff --git a/Engine/Source/Editor/EditorStyle/Private/StarshipStyle.cpp b/Engine/Source/Editor/EditorStyle/Private/StarshipStyle.cpp
--- a/Engine/Source/Editor/EditorStyle/Private/StarshipStyle.cpp
+++ b/Engine/Source/Editor/EditorStyle/Private/StarshipStyle.cpp
@@ -981,12 +981,12 @@ void FStarshipEditorStyle::FStyle::SetupGeneralStyles()
 
 	// Developer Tools Menu
 	{
-		Set("DeveloperTools.MenuIcon", new IMAGE_BRUSH( "Icons/icon_tab_DevTools_16x", Icon16x16 ) );
+		Set("DeveloperTools.MenuIcon", new IMAGE_BRUSH_SVG( "Starship/Menus/Duo/LevelEditor/DeveloperTools", Icon16x16 ) );
 	}
 
 	// Automation Tools Menu
 	{
-		Set("AutomationTools.MenuIcon", new IMAGE_BRUSH("Icons/icon_tab_Tools_16x", Icon16x16));
+		Set("AutomationTools.MenuIcon", new IMAGE_BRUSH_SVG("Starship/Menus/Duo/LevelEditor/AutomationTools", Icon16x16));
 	}
 
 	// Session Browser tab
@@ -1022,7 +1022,7 @@ void FStarshipEditorStyle::FStyle::SetupGeneralStyles()
 
 	// Launcher Window
 	{
-		Set("Launcher.TabIcon", new IMAGE_BRUSH( "Icons/icon_tab_SessionLauncher_16x", Icon16x16 ) );
+		Set("Launcher.TabIcon", new IMAGE_BRUSH_SVG( "Starship/Menus/Duo/LevelEditor/ProjectLauncher", Icon16x16 ) );
 		Set("Launcher.Tabs.Tools", new IMAGE_BRUSH( "/Icons/icon_tab_Tools_16x", Icon16x16 ) );
 	}
 
@@ -4175,24 +4175,30 @@ void FStarshipEditorStyle::FStyle::SetupLevelEditorStyle()
 
 #if WITH_EDITOR || (IS_PROGRAM && WITH_UNREAL_DEVELOPER_TOOLS)
 		{
-			Set( "LevelEditor.Tabs.Details", new IMAGE_BRUSH( "/Icons/icon_tab_SelectionDetails_16x", Icon16x16 ) );
-			Set( "LevelEditor.Tabs.Cinematics", new IMAGE_BRUSH("/Icons/icon_tab_Cinematics_16x", Icon16x16));
-			Set( "LevelEditor.Tabs.EditorModes", new IMAGE_BRUSH( "/Icons/icon_Editor_Modes_16x", Icon16x16 ) );
-			Set( "LevelEditor.Tabs.Modes", new IMAGE_BRUSH( "/Icons/icon_Editor_Modes_16x", Icon16x16 ) );
-			Set( "LevelEditor.Tabs.PlacementBrowser", new IMAGE_BRUSH("/Icons/icon_Editor_Modes_16x", Icon16x16));
-			Set( "LevelEditor.Tabs.Properties", new IMAGE_BRUSH( "/Icons/properties_16x", Icon16x16 ) );
-			Set( "LevelEditor.Tabs.Outliner", new IMAGE_BRUSH( "/Icons/icon_tab_SceneOutliner_16x", Icon16x16 ) );
-			Set( "LevelEditor.Tabs.ContentBrowser", new IMAGE_BRUSH( "/Icons/icon_tab_ContentBrowser_16x", Icon16x16 ) );
-			Set( "LevelEditor.Tabs.Levels", new IMAGE_BRUSH( "/Icons/icon_tab_Levels_16x", Icon16x16 ) );
-			Set( "LevelEditor.Tabs.WorldBrowser", new IMAGE_BRUSH( "/Icons/icon_tab_levels_16x", Icon16x16 ) );
-			Set( "LevelEditor.Tabs.WorldBrowserDetails", new IMAGE_BRUSH( "/Icons/icon_levels_detailsbutton_16x", Icon16x16 ) );
-			Set( "LevelEditor.Tabs.WorldBrowserComposition", new IMAGE_BRUSH( "/Icons/icon_levels_compositionbutton_16x", Icon16x16 ) );
-			Set( "LevelEditor.Tabs.Layers", new IMAGE_BRUSH( "/Icons/icon_tab_Layers_16x", Icon16x16 ) );
-			Set( "LevelEditor.Tabs.BuildAndSubmit", new IMAGE_BRUSH( "/Icons/icon_tab_BuildSubmit_16x", Icon16x16 ) );
-			Set( "LevelEditor.Tabs.StatsViewer", new IMAGE_BRUSH( "/Icons/icon_tab_Stats_16x", Icon16x16 ) );
-			Set( "LevelEditor.Tabs.Toolbar", new IMAGE_BRUSH( "/Icons/icon_tab_Toolbars_16x", Icon16x16 ) );
-			Set( "LevelEditor.Tabs.Viewports", new IMAGE_BRUSH( "/Icons/icon_tab_Viewports_16x", Icon16x16 ) );
-			Set( "LevelEditor.Tabs.HLOD", new IMAGE_BRUSH("/Icons/icon_tab_layers_16px", Icon16x16));
+			Set( "LevelEditor.Tabs.Details",                new IMAGE_BRUSH_SVG("Starship/Menus/Duo/LevelEditor/Details", Icon16x16) );///Icons/icon_tab_SelectionDetails_16x", Icon16x16 ) );
+			Set( "LevelEditor.Tabs.Cinematics",             new IMAGE_BRUSH_SVG("Starship/Menus/Duo/LevelEditor/Cinematics", Icon16x16) );///Icons/icon_tab_Cinematics_16x", Icon16x16));
+			Set( "LevelEditor.Tabs.EditorModes",            new IMAGE_BRUSH_SVG("Starship/Menus/Duo/LevelEditor/EditorModes", Icon16x16) );///Icons/icon_Editor_Modes_16x", Icon16x16 ) );
+			Set( "LevelEditor.Tabs.Modes",                  new IMAGE_BRUSH_SVG("Starship/Menus/Duo/LevelEditor/EditorModes", Icon16x16) );///Icons/icon_Editor_Modes_16x", Icon16x16 ) );
+			Set( "LevelEditor.Tabs.PlacementBrowser",       new IMAGE_BRUSH_SVG("Starship/Menus/Duo/LevelEditor/PlaceActors", Icon16x16) );///Icons/icon_Editor_Modes_16x", Icon16x16));
+			Set( "LevelEditor.Tabs.Properties",             new IMAGE_BRUSH_SVG("Starship/Menus/Duo/LevelEditor/Properties", Icon16x16) );///Icons/properties_16x", Icon16x16 ) );
+			Set( "LevelEditor.Tabs.Outliner",               new IMAGE_BRUSH_SVG("Starship/Menus/Duo/LevelEditor/WorldOutliner", Icon16x16) );///Icons/icon_tab_SceneOutliner_16x", Icon16x16 ) );
+			Set( "LevelEditor.Tabs.ContentBrowser",         new IMAGE_BRUSH_SVG("Starship/Menus/Duo/LevelEditor/ContentBrowser", Icon16x16) );///Icons/icon_tab_ContentBrowser_16x", Icon16x16 ) );
+			Set( "LevelEditor.Tabs.Levels",                 new IMAGE_BRUSH_SVG("Starship/Menus/Duo/LevelEditor/Levels", Icon16x16) );///Icons/icon_tab_Levels_16x", Icon16x16 ) );
+			Set( "LevelEditor.Tabs.WorldBrowser",           new IMAGE_BRUSH_SVG("Starship/Menus/Duo/LevelEditor/Levels", Icon16x16) );///Icons/icon_tab_levels_16x", Icon16x16 ) );
+			Set( "LevelEditor.Tabs.WorldBrowserDetails",    new IMAGE_BRUSH_SVG("Starship/Menus/Duo/LevelEditor/WorldSettings", Icon16x16) );///Icons/icon_levels_detailsbutton_16x", Icon16x16 ) );
+
+			Set( "LevelEditor.Tabs.WorldBrowserComposition",new IMAGE_BRUSH("Icons/icon_levels_compositionbutton_16x", Icon16x16 ) );
+
+			Set( "LevelEditor.Tabs.Layers",                 new IMAGE_BRUSH_SVG("Starship/Menus/Duo/LevelEditor/Layers", Icon16x16) );///Icons/icon_tab_Layers_16x", Icon16x16 ) );
+
+			Set( "LevelEditor.Tabs.BuildAndSubmit",         new IMAGE_BRUSH("Icons/icon_tab_BuildSubmit_16x", Icon16x16 ) );
+
+			Set( "LevelEditor.Tabs.StatsViewer",            new IMAGE_BRUSH_SVG("Starship/Menus/Duo/LevelEditor/Statistics", Icon16x16) );///Icons/icon_tab_Stats_16x", Icon16x16 ) );
+
+			Set( "LevelEditor.Tabs.Toolbar",                new IMAGE_BRUSH("Icons/icon_tab_Toolbars_16x", Icon16x16 ) );
+
+			Set( "LevelEditor.Tabs.Viewports",              new IMAGE_BRUSH_SVG("Starship/Menus/Duo/LevelEditor/Viewports", Icon16x16) );///Icons/icon_tab_Viewports_16x", Icon16x16 ) );
+			Set( "LevelEditor.Tabs.HLOD",                   new IMAGE_BRUSH_SVG("Starship/Menus/Duo/LevelEditor/HierarchicalLODOutliner", Icon16x16) );///Icons/icon_tab_layers_16px", Icon16x16));
 		}
 #endif // WITH_EDITOR || (IS_PROGRAM && WITH_UNREAL_DEVELOPER_TOOLS)
 
@@ -4286,7 +4292,9 @@ void FStarshipEditorStyle::FStyle::SetupLevelEditorStyle()
 
 		Set( "LevelEditor.WorldProperties", new IMAGE_BRUSH( "Icons/icon_worldscript_40x", Icon40x40 ) );
 		Set( "LevelEditor.WorldProperties.Small", new IMAGE_BRUSH( "Icons/icon_worldscript_40x", Icon20x20 ) );
-		Set( "LevelEditor.WorldProperties.Tab", new IMAGE_BRUSH( "Icons/icon_worldscript_40x", Icon16x16 ) );
+
+		Set( "LevelEditor.WorldProperties.Tab", new IMAGE_BRUSH_SVG( "Starship/Menus/Duo/LevelEditor/WorldSettings", Icon16x16 ) );
+
 		Set( "LevelEditor.OpenContentBrowser", new IMAGE_BRUSH_SVG( "Starship/MainToolBar/content", Icon20x20 ) );
 		Set( "LevelEditor.OpenMarketplace", new IMAGE_BRUSH_SVG( "Starship/MainToolBar/marketplace", Icon20x20) );
 		Set( "LevelEditor.OpenLevelBlueprint", new IMAGE_BRUSH_SVG( "Starship/MainToolBar/blueprints", Icon20x20) );
@@ -4446,8 +4454,9 @@ void FStarshipEditorStyle::FStyle::SetupLevelEditorStyle()
 		Set( "LevelViewport.EjectActorPilot.Small", new IMAGE_BRUSH( "Icons/icon_EjectActorPilot_16x", Icon16x16 ) );
 		Set( "LevelViewport.PilotSelectedActor", new IMAGE_BRUSH( "Icons/icon_PilotSelectedActor_16x", Icon16x16 ) );
 		Set( "LevelViewport.PilotSelectedActor.Small", new IMAGE_BRUSH( "Icons/icon_PilotSelectedActor_16x", Icon16x16 ) );
-		Set( "LevelViewport.ToggleActorPilotCameraView", new IMAGE_BRUSH( "Icons/icon_ToggleActorPilotCameraView_16x", Icon16x16 ) );
-		Set( "LevelViewport.ToggleActorPilotCameraView.Small", new IMAGE_BRUSH( "Icons/icon_ToggleActorPilotCameraView_16x", Icon16x16 ) );
+		
+		Set( "LevelViewport.ToggleActorPilotCameraView",       new IMAGE_BRUSH_SVG( "Starship/Menus/Duo/LevelEditor/CameraShake", Icon16x16 ) );
+		Set( "LevelViewport.ToggleActorPilotCameraView.Small", new IMAGE_BRUSH_SVG( "Starship/Menus/Duo/LevelEditor/CameraShake", Icon16x16 ) );
 
 		Set( "LevelViewport.ActorPilotText", FTextBlockStyle()
 			.SetFont( DEFAULT_FONT( "BoldCondensed", 12 ) )
@@ -4864,6 +4873,8 @@ void FStarshipEditorStyle::FStyle::SetupPersonaStyle()
 
 		Set( "BlueprintEditor.ActionMenu.ContextDescriptionFont",  DEFAULT_FONT("Regular", 12) );
 
+		Set( "BlueprintEditor.FindInBlueprints.MenuIcon", new IMAGE_BRUSH_SVG("Starship/Menus/Duo/LevelEditor/FindInBlueprint", Icon16x16) );
+
 		Set( "BlueprintEditor.FindInBlueprint", new IMAGE_BRUSH( "Icons/icon_Blueprint_Find_40px", Icon40x40 ) );
 		Set( "BlueprintEditor.FindInBlueprint.Small", new IMAGE_BRUSH( "Icons/icon_Blueprint_Find_40px", Icon20x20 ) );
 
@@ -5757,7 +5768,7 @@ void FStarshipEditorStyle::FStyle::SetupContentBrowserStyle()
 	// Content Browser
 	{
 		// Tab and menu icon
-		Set( "ContentBrowser.TabIcon", new IMAGE_BRUSH( "Icons/icon_tab_ContentBrowser_16x", Icon16x16 ) );
+		Set( "ContentBrowser.TabIcon", new IMAGE_BRUSH_SVG("Starship/Menus/Duo/LevelEditor/ContentBrowser", Icon16x16) );
 
 		// Sources View
 		Set( "ContentBrowser.SourceTitleFont", DEFAULT_FONT( "Regular", 12 ) );
diff --git a/Engine/Source/Editor/Kismet/Private/FindInBlueprintManager.cpp b/Engine/Source/Editor/Kismet/Private/FindInBlueprintManager.cpp
--- a/Engine/Source/Editor/Kismet/Private/FindInBlueprintManager.cpp
+++ b/Engine/Source/Editor/Kismet/Private/FindInBlueprintManager.cpp
@@ -3480,7 +3480,7 @@ void FFindInBlueprintSearchManager::EnableGlobalFindResults(bool bEnable)
 	if (bEnable)
 	{
 		// Register the spawners for all global Find Results tabs
-		const FSlateIcon GlobalFindResultsIcon(FEditorStyle::GetStyleSetName(), "Kismet.Tabs.FindResults");
+		const FSlateIcon GlobalFindResultsIcon(FEditorStyle::GetStyleSetName(), "BlueprintEditor.FindInBlueprints.MenuIcon");
 		GlobalFindResultsMenuItem = WorkspaceMenu::GetMenuStructure().GetToolsCategory()->AddGroup(
 			LOCTEXT("WorkspaceMenu_GlobalFindResultsCategory", "Find in Blueprints"),
 			LOCTEXT("GlobalFindResultsMenuTooltipText", "Find references to functions, events and variables in all Blueprints."),
