Details Panel Style adjustments

diff --git a/Engine/Source/Editor/EditorStyle/Private/StarshipStyle.cpp b/Engine/Source/Editor/EditorStyle/Private/StarshipStyle.cpp
--- a/Engine/Source/Editor/EditorStyle/Private/StarshipStyle.cpp
+++ b/Engine/Source/Editor/EditorStyle/Private/StarshipStyle.cpp
@@ -579,8 +579,8 @@ void FStarshipEditorStyle::FStyle::SetupGeneralStyles()
 	// Filtering/Searching feedback
 	{
 		const FLinearColor ActiveFilterColor = FLinearColor(1.0f,0.55f,0.0f,1.0f);
-		Set("Searching.SearchActiveTab", new BOX_BRUSH("Common/SearchPseudoTab", FVector2D(16,16), FMargin(0.49f), ActiveFilterColor));
-		Set("Searching.SearchActiveBorder", new BOX_BRUSH("Common/SearchActiveBorder", FVector2D(8,8), FMargin(0.49f), ActiveFilterColor));
+		Set("Searching.SearchActiveTab",    new FSlateNoResource());
+		Set("Searching.SearchActiveBorder", new FSlateRoundedBoxBrush(FLinearColor::Transparent, 0.0, FStyleColors::Primary, 1.f));
 	}
 
 
@@ -1117,9 +1117,9 @@ void FStarshipEditorStyle::FStyle::SetupGeneralStyles()
 	// Scene Outliner
 	{
 		Set( "SceneOutliner.NewFolderIcon", new CORE_IMAGE_BRUSH_SVG("Starship/Common/folder-plus", Icon16x16 ) );
-		Set( "SceneOutliner.FolderClosed",  new CORE_IMAGE_BRUSH_SVG("Starship/SceneOutliner/FolderClosed", Icon16x16, FStyleColors::AccentFolder ) );
-		Set( "SceneOutliner.FolderOpen",    new CORE_IMAGE_BRUSH_SVG("Starship/SceneOutliner/FolderOpen", Icon16x16, FStyleColors::AccentFolder ) );
-		Set( "SceneOutliner.World", 		new CORE_IMAGE_BRUSH_SVG("Starship/SceneOutliner/World", Icon16x16 ) );
+		Set( "SceneOutliner.FolderClosed",  new CORE_IMAGE_BRUSH_SVG("Starship/Common/folder-closed", Icon16x16, FStyleColors::AccentFolder ) );
+		Set( "SceneOutliner.FolderOpen",    new CORE_IMAGE_BRUSH_SVG("Starship/Common/folder-open", Icon16x16, FStyleColors::AccentFolder ) );
+		Set( "SceneOutliner.World", 		new CORE_IMAGE_BRUSH_SVG("Starship/Common/world", Icon16x16 ) );
 		Set( "SceneOutliner.ChangedItemHighlight", new FSlateRoundedBoxBrush( FStyleColors::White, 1.0) );
 
 		Set( "SceneOutliner.TableViewRow", FTableRowStyle(NormalTableRowStyle));
@@ -1878,10 +1878,10 @@ void FStarshipEditorStyle::FStyle::SetupLevelGeneralStyles()
 {
 // Levels General
 	{
-	Set("Level.VisibleIcon16x",             new CORE_IMAGE_BRUSH_SVG("Starship/SceneOutliner/Visible", Icon16x16));
-	Set("Level.VisibleHighlightIcon16x",    new CORE_IMAGE_BRUSH_SVG("Starship/SceneOutliner/Visible", Icon16x16));
-	Set("Level.NotVisibleIcon16x",          new CORE_IMAGE_BRUSH_SVG("Starship/SceneOutliner/Hidden",  Icon16x16));
-	Set("Level.NotVisibleHighlightIcon16x", new CORE_IMAGE_BRUSH_SVG("Starship/SceneOutliner/Hidden",  Icon16x16));
+	Set("Level.VisibleIcon16x",             new CORE_IMAGE_BRUSH_SVG("Starship/Common/visible", Icon16x16));
+	Set("Level.VisibleHighlightIcon16x",    new CORE_IMAGE_BRUSH_SVG("Starship/Common/visible", Icon16x16));
+	Set("Level.NotVisibleIcon16x",          new CORE_IMAGE_BRUSH_SVG("Starship/Common/hidden",  Icon16x16));
+	Set("Level.NotVisibleHighlightIcon16x", new CORE_IMAGE_BRUSH_SVG("Starship/Common/hidden",  Icon16x16));
 
 	Set("Level.LightingScenarioIcon16x", new IMAGE_BRUSH("Icons/icon_levels_LightingScenario_16px", Icon16x16));
 	Set("Level.LightingScenarioNotIcon16x", new IMAGE_BRUSH("Icons/icon_levels_LightingScenarioNot_16px", Icon16x16));
@@ -3014,8 +3014,8 @@ void FStarshipEditorStyle::FStyle::SetupPropertyEditorStyles()
 		Set( "PropertyWindow.FilterCancel", new IMAGE_BRUSH( "Old/FilterCancel", Icon16x16 ) );
 		Set( "PropertyWindow.Favorites_Enabled", new IMAGE_BRUSH( "Icons/Star_16x", Icon16x16 ) );
 		Set( "PropertyWindow.Favorites_Disabled", new IMAGE_BRUSH( "Icons/EmptyStar_16x", Icon16x16 ) );
-		Set( "PropertyWindow.Locked", new IMAGE_BRUSH( "Icons/padlock_locked_16x", Icon16x16 ) );
-		Set( "PropertyWindow.Unlocked", new IMAGE_BRUSH( "Icons/padlock_unlocked_16x", Icon16x16 ) );
+		Set( "PropertyWindow.Locked", new CORE_IMAGE_BRUSH_SVG( "Starship/Common/lock", Icon16x16 ) );
+		Set( "PropertyWindow.Unlocked", new CORE_IMAGE_BRUSH_SVG( "Starship/Common/lock-unlocked", Icon16x16 ) );
 		Set( "PropertyWindow.DiffersFromDefault", new IMAGE_BRUSH( "/PropertyView/DiffersFromDefault_8x8", FVector2D(8,8) ) ) ;
 		
 		Set( "PropertyWindow.NormalFont", DEFAULT_FONT("Regular", 8) );
@@ -3071,7 +3071,8 @@ void FStarshipEditorStyle::FStyle::SetupPropertyEditorStyles()
 		Set( "DetailsView.PulldownArrow.Up",          new CORE_IMAGE_BRUSH_SVG("Starship/Common/chevron-up",   Icon16x16, FStyleColors::Foreground)); 
 		Set( "DetailsView.PulldownArrow.Up.Hovered",  new CORE_IMAGE_BRUSH_SVG("Starship/Common/chevron-up",   Icon16x16, FStyleColors::ForegroundHover)); 
 
-		Set( "DetailsView.EditRawProperties",         new IMAGE_BRUSH("Icons/icon_PropertyMatrix_16px",  Icon16x16, FLinearColor::Black ) );
+		Set( "DetailsView.EditRawProperties",         new CORE_IMAGE_BRUSH_SVG("Starship/Common/layout-spreadsheet",  Icon16x16, FLinearColor::White) );
+		Set( "DetailsView.ViewOptions",         	  new CORE_IMAGE_BRUSH_SVG("Starship/Common/settings",  Icon16x16, FLinearColor::White) );
 		Set( "DetailsView.EditConfigProperties",      new IMAGE_BRUSH("Icons/icon_PropertyMatrix_16px",  Icon16x16, FLinearColor::White ) );
 
 		Set( "DetailsView.CollapsedCategory",         new FSlateColorBrush(FStyleColors::Background));
@@ -3097,10 +3098,10 @@ void FStarshipEditorStyle::FStyle::SetupPropertyEditorStyles()
 		Set( "DetailsView.AdvancedDropdownBorder",      new FSlateColorBrush(FStyleColors::Background));
 		Set( "DetailsView.AdvancedDropdownBorder.Open", new FSlateColorBrush(FStyleColors::Background));
 
-		Set( "DetailsView.CategoryFontStyle", FStyleFonts::Get().NormalBold);
+		Set( "DetailsView.CategoryFontStyle", FStyleFonts::Get().SmallBold);
 		Set( "DetailsView.CategoryTextStyle", FTextBlockStyle(NormalText)
 			.SetFont(GetFontStyle("DetailsView.CategoryFontStyle"))
-			.SetColorAndOpacity(FStyleColors::White)
+			.SetColorAndOpacity(FStyleColors::ForegroundHeader)
 		);
 
 		Set( "DetailsView.TreeView.TableRow", FTableRowStyle()
diff --git a/Engine/Source/Editor/Kismet/Private/SSCSEditor.cpp b/Engine/Source/Editor/Kismet/Private/SSCSEditor.cpp
--- a/Engine/Source/Editor/Kismet/Private/SSCSEditor.cpp
+++ b/Engine/Source/Editor/Kismet/Private/SSCSEditor.cpp
@@ -3800,13 +3800,18 @@ void SSCSEditor::Construct( const FArguments& InArgs )
 					+ SVerticalBox::Slot()
 						.AutoHeight()
 						.VAlign(VAlign_Top)
+						.HAlign(HAlign_Fill)
 					[
 						SAssignNew(ButtonBox, SHorizontalBox)
-				
+
+						+ SHorizontalBox::Slot()
+						.FillWidth(1)
+						[
+							SNew(SSpacer)
+						]
+
 						+ SHorizontalBox::Slot()
-						.Padding( 3.0f, 3.0f )
 						.AutoWidth()
-						.HAlign(HAlign_Left)
 						[
 							SNew(SComponentClassCombo)
 							.AddMetaData<FTagMetaData>(FTagMetaData(TEXT("Actor.AddComponent")))
@@ -3820,83 +3825,39 @@ void SSCSEditor::Construct( const FArguments& InArgs )
 						// horizontal slot (index) #1 => reserved for BP-editor search bar (see 'ButtonBox' usage below)
 
 						+ SHorizontalBox::Slot()
-						.FillWidth(1.0f)
-						.HAlign(HAlign_Right)
-						.Padding( 3.0f, 3.0f )
+						.AutoWidth()
 						[
 							SNew( SButton )
+							.ButtonStyle( FEditorStyle::Get(), "SimpleButton" )
 							.AddMetaData<FTagMetaData>(FTagMetaData(TEXT("Actor.ConvertToBlueprint")))
 							.Visibility( this, &SSCSEditor::GetPromoteToBlueprintButtonVisibility )
 							.OnClicked( this, &SSCSEditor::OnPromoteToBlueprintClicked )
-							.ButtonStyle(FEditorStyle::Get(), "FlatButton.Primary")
-							.ContentPadding(FMargin(10,0))
+							.ContentPadding(FMargin(5.f))
 							.ToolTip(IDocumentation::Get()->CreateToolTip(
 								LOCTEXT("PromoteToBluerprintTooltip","Converts this actor into a reusable Blueprint Class that can have script behavior" ),
 								NULL,
 								TEXT("Shared/LevelEditor"),
 								TEXT("ConvertToBlueprint")))
 							[
-								SNew(SHorizontalBox)
-								.Clipping(EWidgetClipping::ClipToBounds)
-						
-								+ SHorizontalBox::Slot()
-								.VAlign(VAlign_Center)
-								.Padding(3.f)
-								.AutoWidth()
-								[
-									SNew(STextBlock)
-									.TextStyle(FEditorStyle::Get(), "ContentBrowser.TopBar.Font")
-									.Font( FEditorStyle::Get().GetFontStyle( "FontAwesome.10" ) )
-									.Text( FEditorFontGlyphs::Cogs )
-								]
-
-								+ SHorizontalBox::Slot()
-								.VAlign(VAlign_Center)
-								.Padding(3.f)
-								.AutoWidth()
-								[
-									SNew(STextBlock)
-									.TextStyle(FEditorStyle::Get(), "ContentBrowser.TopBar.Font")
-									//.Text( LOCTEXT("PromoteToBlueprint", "Add Script") )
-									.Text(LOCTEXT("PromoteToBlueprint", "Blueprint/Add Script"))
-								]
+								SNew(SImage)
+								.ColorAndOpacity(FSlateColor::UseForeground())
+								.Image(FAppStyle::Get().GetBrush("Icons.Blueprints"))
 							]
 						]
 						+ SHorizontalBox::Slot()
-						.FillWidth(1.0f)
-						.Padding( 3.0f, 3.0f )
-						.HAlign(HAlign_Right)
-						.Padding(3.0f, 3.0f)
+						.AutoWidth()
 						[
 							SNew(SComboButton)
 							.AddMetaData<FTagMetaData>(FTagMetaData(TEXT("Actor.EditBlueprint")))
 							.Visibility(this, &SSCSEditor::GetEditBlueprintButtonVisibility)
-							.ContentPadding(FMargin(10, 0))
-							.ComboButtonStyle(FEditorStyle::Get(), "ToolbarComboButton")
-							.ButtonStyle(FEditorStyle::Get(), "FlatButton.Primary")
-							.ForegroundColor(FLinearColor::White)
+							.ComboButtonStyle(FEditorStyle::Get(), "SimpleComboButton")
+							.HasDownArrow(false)
+							.ContentPadding(FMargin(5.f))
 							.ButtonContent()
 							[
-								SNew( SHorizontalBox )
-								.Clipping(EWidgetClipping::ClipToBounds)
-
-								+ SHorizontalBox::Slot()
-								.AutoWidth()
-								.VAlign(VAlign_Center)
-								.Padding(3.f)
-								[
-									SNew(STextBlock)
-									.TextStyle(FEditorStyle::Get(), "ContentBrowser.TopBar.Font")
-									.Font(FEditorStyle::Get().GetFontStyle("FontAwesome.10"))
-									.Text(FEditorFontGlyphs::Cogs)
-								]
-						
-								+ SHorizontalBox::Slot()
-								[
-									SNew(STextBlock)
-									.TextStyle(FEditorStyle::Get(), "ContentBrowser.TopBar.Font")
-									.Text(LOCTEXT("EditBlueprint", "Edit Blueprint"))
-								]
+								SNew(SImage)
+								.ColorAndOpacity(FSlateColor::UseForeground())
+								.Image(FAppStyle::Get().GetBrush("Icons.Blueprints"))
 							]
 							.MenuContent()
 							[
diff --git a/Engine/Source/Editor/LevelEditor/Private/SActorDetails.cpp b/Engine/Source/Editor/LevelEditor/Private/SActorDetails.cpp
--- a/Engine/Source/Editor/LevelEditor/Private/SActorDetails.cpp
+++ b/Engine/Source/Editor/LevelEditor/Private/SActorDetails.cpp
@@ -58,7 +58,7 @@ public:
 		ChildSlot
 		[
 			SNew(SBorder)
-			.BorderImage(FEditorStyle::Get().GetBrush("ToolPanel.GroupBorder"))
+			.BorderImage(FAppStyle::Get().GetBrush("Brushes.Background"))
 			[
 				SNew(SHorizontalBox)
 				+ SHorizontalBox::Slot()
diff --git a/Engine/Source/Editor/LevelEditor/Private/SLevelEditor.cpp b/Engine/Source/Editor/LevelEditor/Private/SLevelEditor.cpp
--- a/Engine/Source/Editor/LevelEditor/Private/SLevelEditor.cpp
+++ b/Engine/Source/Editor/LevelEditor/Private/SLevelEditor.cpp
@@ -579,7 +579,8 @@ TSharedRef<SDockTab> SLevelEditor::SummonDetailsPanel( FName TabIdentifier )
 		.Label( Label )
 		.ToolTip( IDocumentation::Get()->CreateToolTip( Label, nullptr, "Shared/LevelEditor", "DetailsTab" ) )
 		[
-			SNew( SBox )
+			SNew( SBorder)
+			.BorderImage(FAppStyle::Get().GetBrush("Brushes.Background"))
 			.AddMetaData<FTutorialMetaData>(FTutorialMetaData(TEXT("ActorDetails"), TEXT("LevelEditorSelectionDetails")))
 			[
 				ActorDetails
diff --git a/Engine/Source/Editor/PropertyEditor/Private/SDetailNameArea.cpp b/Engine/Source/Editor/PropertyEditor/Private/SDetailNameArea.cpp
--- a/Engine/Source/Editor/PropertyEditor/Private/SDetailNameArea.cpp
+++ b/Engine/Source/Editor/PropertyEditor/Private/SDetailNameArea.cpp
@@ -159,11 +159,12 @@ TSharedRef< SWidget > SDetailNameArea::BuildObjectNameArea( const TArray< TWeakO
 				.FillWidth(1.0f)
 				[
 					SNew( SButton )
-					.ButtonStyle( FEditorStyle::Get(), "NoBorder" )
+					.ButtonStyle( FEditorStyle::Get(), "SimpleButton" )
 					.OnClicked(	OnLockButtonClicked )
 					.ToolTipText( LOCTEXT("LockSelectionButton_ToolTip", "Locks the current selection into the Details panel") )
 					[
 						SNew( SImage )
+						.ColorAndOpacity(FSlateColor::UseForeground())
 						.Image( this, &SDetailNameArea::OnGetLockButtonImageResource )
 					]
 				];
diff --git a/Engine/Source/Editor/PropertyEditor/Private/SDetailsView.cpp b/Engine/Source/Editor/PropertyEditor/Private/SDetailsView.cpp
--- a/Engine/Source/Editor/PropertyEditor/Private/SDetailsView.cpp
+++ b/Engine/Source/Editor/PropertyEditor/Private/SDetailsView.cpp
@@ -197,43 +197,34 @@ void SDetailsView::Construct(const FArguments& InArgs, const FDetailsViewArgs& I
 		SNew( SHorizontalBox )
 		.Visibility( this, &SDetailsView::GetFilterBoxVisibility )
 		+SHorizontalBox::Slot()
-		.FillWidth( 1 )
+		.Padding(6.f)
+		.FillWidth(1)
 		[
-			SNew(SOverlay)
-			+SOverlay::Slot()
-			.Padding(0.0f, 0.0f, 0.0f, 0.0f)
-			[
-				SNew(SImage)
-				.Image(FEditorStyle::GetBrush("Searching.SearchActiveTab"))
-				.Visibility_Lambda([&](){ return this->bHasActiveFilter ? EVisibility::Visible : EVisibility::Collapsed; })
-			]
-			+SOverlay::Slot()
-			.Padding(2.f, 2.0f, 4.f, 2.f)
-			.VAlign( VAlign_Center )
-			[
-				// Create the search box
-				SAssignNew(SearchBox, SSearchBox)
-				.HintText(LOCTEXT("SearchDetailsHint", "Search Details"))
-				.OnTextChanged(this, &SDetailsView::OnFilterTextChanged)
-				.OnTextCommitted(this, &SDetailsView::OnFilterTextCommitted)
-				.AddMetaData<FTagMetaData>(TEXT("Details.Search"))
-			]
+			// Create the search box
+			SAssignNew(SearchBox, SSearchBox)
+			.HintText(LOCTEXT("SearchDetailsHint", "Search Details"))
+			.OnTextChanged(this, &SDetailsView::OnFilterTextChanged)
+			.OnTextCommitted(this, &SDetailsView::OnFilterTextCommitted)
+			.AddMetaData<FTagMetaData>(TEXT("Details.Search"))
 		];
 
 	if (DetailsViewArgs.bShowPropertyMatrixButton)
 	{
 		FilterRow->AddSlot()
 			.Padding(4.0f, 0.0f, 0.0f, 0.0f)
+			.VAlign(VAlign_Center)
 			.AutoWidth()
 			[
 				// Create the property matrix button
 				SNew(SButton)
+				.ButtonStyle( FAppStyle::Get(), "SimpleButton" )
 				.OnClicked(this, &SDetailsView::OnOpenRawPropertyEditorClicked)
 				.IsEnabled(this, &SDetailsView::CanOpenRawPropertyEditor)
 				.ToolTipText(LOCTEXT("RawPropertyEditorButtonLabel", "Open Selection in Property Matrix"))
 				[
 					SNew(SImage)
-					.Image(FEditorStyle::GetBrush("DetailsView.EditRawProperties"))
+					.ColorAndOpacity(FSlateColor::UseForeground())
+					.Image(FAppStyle::Get().GetBrush("DetailsView.EditRawProperties"))
 				]
 			];
 	}
@@ -242,12 +233,14 @@ void SDetailsView::Construct(const FArguments& InArgs, const FDetailsViewArgs& I
 	{
 		FilterRow->AddSlot()
 			.HAlign(HAlign_Right)
+			.VAlign(VAlign_Center)
 			.AutoWidth()
 			[
 				SNew( SComboButton )
+				.HasDownArrow(false)
 				.ContentPadding(0)
 				.ForegroundColor( FSlateColor::UseForeground() )
-				.ButtonStyle( FEditorStyle::Get(), "ToggleButton" )
+				.ButtonStyle( FAppStyle::Get(), "SimpleButton" )
 				.AddMetaData<FTagMetaData>(FTagMetaData(TEXT("ViewOptions")))
 				.MenuContent()
 				[
@@ -256,7 +249,8 @@ void SDetailsView::Construct(const FArguments& InArgs, const FDetailsViewArgs& I
 				.ButtonContent()
 				[
 					SNew(SImage)
-					.Image( FEditorStyle::GetBrush("GenericViewButton") )
+					.ColorAndOpacity(FSlateColor::UseForeground())
+					.Image(FAppStyle::Get().GetBrush("DetailsView.ViewOptions"))
 				]
 			];
 	}
@@ -289,7 +283,11 @@ void SDetailsView::Construct(const FArguments& InArgs, const FDetailsViewArgs& I
 		VerticalBox->AddSlot()
 		.AutoHeight()
 		[
-			FilterRow.ToSharedRef()
+			SNew(SBorder)
+			.BorderImage(FAppStyle::Get().GetBrush("Brushes.Background"))
+			[
+				FilterRow.ToSharedRef()
+			]
 		];
 	}
 
diff --git a/Engine/Source/Editor/UnrealEd/Private/SComponentClassCombo.cpp b/Engine/Source/Editor/UnrealEd/Private/SComponentClassCombo.cpp
--- a/Engine/Source/Editor/UnrealEd/Private/SComponentClassCombo.cpp
+++ b/Engine/Source/Editor/UnrealEd/Private/SComponentClassCombo.cpp
@@ -64,28 +64,15 @@ void SComponentClassCombo::Construct(const FArguments& InArgs)
 
 	// Create the Construct arguments for the parent class (SComboButton)
 	SComboButton::FArguments Args;
-	Args.ButtonContent()
+	Args.ComboButtonStyle(FAppStyle::Get(), "SimpleComboButton")
+	.HasDownArrow(false)
+	.IsFocusable(true)
+	.OnComboBoxOpened(this, &SComponentClassCombo::ClearSelection)
+	.ButtonContent()
 	[
-		SNew(SHorizontalBox)
-		+SHorizontalBox::Slot()
-		.VAlign(VAlign_Center)
-		.AutoWidth()
-		.Padding(1.f,1.f)
-		[
-			SNew(STextBlock)
-			.TextStyle(FEditorStyle::Get(), "ContentBrowser.TopBar.Font")
-			.Font(FEditorStyle::Get().GetFontStyle("FontAwesome.10"))
-			.Text(FText::FromString(FString(TEXT("\xf067"))) /*fa-plus*/)
-		]
-		+ SHorizontalBox::Slot()
-		.VAlign(VAlign_Center)
-		.Padding(1.f)
-		[
-			SNew(STextBlock)
-			.Text(LOCTEXT("AddComponentButtonLabel", "Add Component"))
-			.TextStyle(FEditorStyle::Get(), "ContentBrowser.TopBar.Font")
-			.Visibility(InArgs._IncludeText.Get() ? EVisibility::Visible : EVisibility::Collapsed)
-		]
+		SNew(SImage)
+		.ColorAndOpacity(FSlateColor::UseForeground())
+		.Image(FAppStyle::Get().GetBrush("Icons.Plus"))
 	]
 	.MenuContent()
 	[
@@ -114,13 +101,7 @@ void SComponentClassCombo::Construct(const FArguments& InArgs)
 				]
 			]
 		]
-	]
-	.IsFocusable(true)
-	.ContentPadding(FMargin(5, 0))
-	.ComboButtonStyle(FEditorStyle::Get(), "ToolbarComboButton")
-	.ButtonStyle(FEditorStyle::Get(), "FlatButton.Success")
-	.ForegroundColor(FLinearColor::White)
-	.OnComboBoxOpened(this, &SComponentClassCombo::ClearSelection);
+	];
 
 	SComboButton::Construct(Args);
 
