Updated StarshipStyle with the latest changes from the old editor style

diff --git a/Engine/Source/Editor/EditorStyle/Private/StarshipStyle.cpp b/Engine/Source/Editor/EditorStyle/Private/StarshipStyle.cpp
--- a/Engine/Source/Editor/EditorStyle/Private/StarshipStyle.cpp
+++ b/Engine/Source/Editor/EditorStyle/Private/StarshipStyle.cpp
@@ -266,6 +266,11 @@ void FStarshipEditorStyle::FStyle::SetupGeneralStyles()
 		Set( "RichTextBlock.BoldHighlight", FTextBlockStyle(NormalText)
 			.SetFont( DEFAULT_FONT("Bold", FStarshipCoreStyle::RegularTextSize ))
 			.SetColorAndOpacity( FLinearColor( 1.0f, 1.0f, 1.0f ) ) );
+		Set("RichTextBlock.Italic", FTextBlockStyle(NormalText)
+			.SetFont(DEFAULT_FONT("Italic", FStarshipCoreStyle::RegularTextSize)));
+		Set("RichTextBlock.ItalicHighlight", FTextBlockStyle(NormalText)
+			.SetFont(DEFAULT_FONT("Italic", FStarshipCoreStyle::RegularTextSize))
+			.SetColorAndOpacity(FLinearColor(1.0f, 1.0f, 1.0f)));
 
 		Set( "TextBlock.HighlightShape",  new BOX_BRUSH( "Common/TextBlockHighlightShape", FMargin(3.f/8.f) ));
 		Set( "TextBlock.HighlighColor", FLinearColor( 0.02f, 0.3f, 0.0f ) );
@@ -2114,6 +2119,8 @@ void FStarshipEditorStyle::FStyle::SetupSequencerStyles()
 		Set("Sequencer.OpenDirectorBlueprint.Small", new IMAGE_BRUSH("Sequencer/Main_Icons/Icon_Sequencer_OpenDirectorBlueprint_24x", Icon24x24));
 		Set("Sequencer.OpenTaggedBindingManager", new IMAGE_BRUSH("Sequencer/Main_Icons/Icon_Sequencer_OpenTaggedBindingManager_16x", Icon48x48));
 		Set("Sequencer.OpenTaggedBindingManager.Small", new IMAGE_BRUSH("Sequencer/Main_Icons/Icon_Sequencer_OpenTaggedBindingManager_16x", Icon24x24));
+		Set("Sequencer.OpenNodeGroupsManager", new IMAGE_BRUSH("Sequencer/Main_Icons/Icon_Sequencer_OpenGroupManager_16x", Icon48x48));
+		Set("Sequencer.OpenNodeGroupsManager.Small", new IMAGE_BRUSH("Sequencer/Main_Icons/Icon_Sequencer_OpenGroupManager_16x", Icon24x24));
 		Set("Sequencer.BreadcrumbText", FTextBlockStyle(NormalText)
 			.SetFont(DEFAULT_FONT("Bold", 11))
 			.SetColorAndOpacity(FLinearColor(1.0f, 1.0f, 1.0f))
@@ -3834,6 +3841,17 @@ void FStarshipEditorStyle::FStyle::SetupGraphEditorStyles()
 			);
 		}
 
+		Set("ClassIcon.K2Node_CallFunction",	new IMAGE_BRUSH("Icons/icon_Blueprint_NewFunction_16x", Icon16x16));
+		Set("ClassIcon.K2Node_FunctionEntry",	new IMAGE_BRUSH("Icons/icon_Blueprint_NewFunction_16x", Icon16x16));
+		Set("ClassIcon.K2Node_CustomEvent",		new IMAGE_BRUSH("Icons/icon_Blueprint_Event_16x", Icon16x16));
+		Set("ClassIcon.K2Node_Event",			new IMAGE_BRUSH("Icons/icon_Blueprint_Event_16x", Icon16x16));
+		Set("ClassIcon.K2Node_Variable",		new IMAGE_BRUSH("Graph/Icons/Node", Icon16x16, FLinearColor::White));
+		Set("ClassIcon.K2Node_VariableGet",		new IMAGE_BRUSH("Graph/Icons/FIB_VarGet", Icon16x16, FLinearColor::White));
+		Set("ClassIcon.K2Node_VariableSet",		new IMAGE_BRUSH("Graph/Icons/FIB_VarSet", Icon16x16, FLinearColor::White));
+		Set("ClassIcon.K2Node_DynamicCast",		new IMAGE_BRUSH("Icons/icon_Blueprint_Cast_16x", Icon16x16));
+		Set("ClassIcon.EdGraphNode_Comment",	new IMAGE_BRUSH("Icons/icon_Blueprint_Comment_16x", Icon16x16));
+
+
 		Set( "GraphEditor.Default_16x", new IMAGE_BRUSH("Icons/icon_Blueprint_Node_16x", Icon16x16));
 		Set( "GraphEditor.EventGraph_16x", new IMAGE_BRUSH( "Icons/icon_Blueprint_EventGraph_16x", Icon16x16 ) );
 		Set( "GraphEditor.InterfaceFunction_16x", new IMAGE_BRUSH( "Icons/icon_Blueprint_Interfacefunction_16x", Icon16x16 ) );
@@ -6140,10 +6158,11 @@ void FStarshipEditorStyle::FStyle::SetupLandscapeEditorStyle()
 		else
 		{
 			// Tools
-			Set("LandscapeEditor.NewLandscape",       			new IMAGE_BRUSH("Icons/Landscape/Landscape_NewLandscape_x40", 		Icon20x20));
+			Set("LandscapeEditor.NewLandscape",					new IMAGE_BRUSH("Icons/Landscape/Landscape_NewLandscape_x40",		Icon20x20));
 			Set("LandscapeEditor.NewLandscape.Small",       	new IMAGE_BRUSH("Icons/Landscape/Landscape_NewLandscape_x40", 		Icon20x20));
-			// Set("LandscapeEditor.ResizeLandscape		",      new IMAGE_BRUSH("Icons/Landscape/Landscape_Copy_x40", 			Icon40x40));
-			// Set("LandscapeEditor.ResizeLandscape.Small",     new IMAGE_BRUSH("Icons/Landscape/Landscape_Copy_x40", 			Icon40x40));
+			Set("LandscapeEditor.ResizeLandscape",				new IMAGE_BRUSH("Icons/Landscape/Landscape_Resize_x40",				Icon20x20));
+			Set("LandscapeEditor.ResizeLandscape.Small",		new IMAGE_BRUSH("Icons/Landscape/Landscape_Resize_x40",				Icon20x20));
+
 
 			Set("LandscapeEditor.SculptTool",       			new IMAGE_BRUSH("Icons/Landscape/Landscape_Sculpt_x40",            Icon20x20));
 			Set("LandscapeEditor.SculptTool.Small",       		new IMAGE_BRUSH("Icons/Landscape/Landscape_Sculpt_x40",            Icon20x20));
