Changed Slate style colors to be FSlateColors and be user configurable.  The way that user style data is stored is completely temporary and is for testing only

diff --git a/Engine/Plugins/FX/Niagara/Source/NiagaraEditor/Private/NiagaraEditorStyle.cpp b/Engine/Plugins/FX/Niagara/Source/NiagaraEditor/Private/NiagaraEditorStyle.cpp
--- a/Engine/Plugins/FX/Niagara/Source/NiagaraEditor/Private/NiagaraEditorStyle.cpp
+++ b/Engine/Plugins/FX/Niagara/Source/NiagaraEditor/Private/NiagaraEditorStyle.cpp
@@ -10,6 +10,8 @@
 #include "Styling/CoreStyle.h"
 #include "Interfaces/IPluginManager.h"
 #include "Classes/EditorStyleSettings.h"
+#include "Styling/StarshipCoreStyle.h"
+#include "Styling/StyleColors.h"
 
 TSharedPtr< FSlateStyleSet > FNiagaraEditorStyle::NiagaraEditorStyleInstance = NULL;
 
@@ -203,10 +205,10 @@ TSharedRef< FSlateStyleSet > FNiagaraEditorStyle::Create()
 		.SetFont(DEFAULT_FONT("Regular", 8))
 		.SetColorAndOpacity(FLinearColor(1.0f, 1.0f, 1.0f, 0.5f)));
 
-	Style->Set("NiagaraEditor.Stack.HighlightedButtonBrush", new BOX_CORE_BRUSH("Common/ButtonHoverHint", FMargin(4 / 16.0f), GetDefault<UEditorStyleSettings>()->SelectionColor));
+	Style->Set("NiagaraEditor.Stack.HighlightedButtonBrush", new BOX_CORE_BRUSH("Common/ButtonHoverHint", FMargin(4 / 16.0f), FStyleColors::Primary));
 
 	// Parameter Map View
-	Style->Set("NiagaraEditor.Stack.DepressedHighlightedButtonBrush", new BOX_CORE_BRUSH("Common/ButtonHoverHint", FMargin(4 / 16.0f), GetDefault<UEditorStyleSettings>()->PressedSelectionColor));
+	Style->Set("NiagaraEditor.Stack.DepressedHighlightedButtonBrush", new BOX_CORE_BRUSH("Common/ButtonHoverHint", FMargin(4 / 16.0f), FStyleColors::PrimaryPress));
 	Style->Set("NiagaraEditor.Stack.ViewOptionsShadowColor", FLinearColor::Black);
 	Style->Set("NiagaraEditor.Stack.FlatButtonColor", FLinearColor(FColor(205, 205, 205)));
 
diff --git a/Engine/Source/Editor/EditorStyle/EditorStyle.Build.cs b/Engine/Source/Editor/EditorStyle/EditorStyle.Build.cs
--- a/Engine/Source/Editor/EditorStyle/EditorStyle.Build.cs
+++ b/Engine/Source/Editor/EditorStyle/EditorStyle.Build.cs
@@ -23,6 +23,7 @@ public class EditorStyle : ModuleRules
 		PrivateIncludePathModuleNames.AddRange(
 			new string[] {
 				"Settings",
+				"PropertyEditor"
 			}
 		);
 
diff --git a/Engine/Source/Editor/EditorStyle/Private/EditorStyleClasses.cpp b/Engine/Source/Editor/EditorStyle/Private/EditorStyleClasses.cpp
--- a/Engine/Source/Editor/EditorStyle/Private/EditorStyleClasses.cpp
+++ b/Engine/Source/Editor/EditorStyle/Private/EditorStyleClasses.cpp
@@ -19,8 +19,6 @@ UEditorStyleSettings::UEditorStyleSettings( const FObjectInitializer& ObjectInit
 	bEnableUserEditorLayoutManagement = true;
 
 	SelectionColor = FLinearColor(0.828f, 0.364f, 0.003f);
-	InactiveSelectionColor = FLinearColor(0.25f, 0.25f, 0.25f);
-	PressedSelectionColor = FLinearColor(0.701f, 0.225f, 0.003f);
 
 	EditorWindowBackgroundColor = FLinearColor::White;
 
@@ -62,18 +60,6 @@ FLinearColor UEditorStyleSettings::GetSubduedSelectionColor() const
 
 void UEditorStyleSettings::PostEditChangeProperty(struct FPropertyChangedEvent& PropertyChangedEvent)
 {
-	if (bResetEditorWindowBackgroundSettings)
-	{
-		// Reset the settings
-		bResetEditorWindowBackgroundSettings = false;
-
-		EditorWindowBackgroundColor = FLinearColor::White;
-
-		FSlateBrush DummyBrush;
-		EditorMainWindowBackgroundOverride = DummyBrush;
-		EditorChildWindowBackgroundOverride = DummyBrush;
-	}
-
 	Super::PostEditChangeProperty(PropertyChangedEvent);
 
 	const FName PropertyName = (PropertyChangedEvent.Property != nullptr) ? PropertyChangedEvent.Property->GetFName() : NAME_None;
diff --git a/Engine/Source/Editor/EditorStyle/Private/EditorStyleModule.cpp b/Engine/Source/Editor/EditorStyle/Private/EditorStyleModule.cpp
--- a/Engine/Source/Editor/EditorStyle/Private/EditorStyleModule.cpp
+++ b/Engine/Source/Editor/EditorStyle/Private/EditorStyleModule.cpp
@@ -4,6 +4,9 @@
 #include "Modules/ModuleManager.h"
 #include "Interfaces/IEditorStyleModule.h"
 #include "SlateEditorStyle.h"
+#include "StarshipStyle.h"
+#include "Styling/CoreStyle.h"
+#include "Styling/StyleColors.h"
 
 
 /**
@@ -18,20 +21,42 @@ public:
 
 	virtual void StartupModule( ) override
 	{
-		FSlateEditorStyle::Initialize();
-	}
+		if (FCoreStyle::IsStarshipStyle())
+		{
+			// This is necessary because the core style  loads the color table before ProcessNewlyLoadedUObjects is called which means none of the config properties are in the classs property link at that time.
+			UStyleColorTable::Get().ReloadConfig();
 
-	virtual void ShutdownModule( ) override
-	{
-		FSlateEditorStyle::Shutdown();
+			bUsingStarshipStyle = true;
+			FStarshipEditorStyle::Initialize();
+
+			// set the application style to be the editor style
+			FAppStyle::SetAppStyleSetName(FStarshipEditorStyle::GetStyleSetName());
+
+		}
+		else
+		{
+			FSlateEditorStyle::Initialize();
+
+			// set the application style to be the editor style
+			FAppStyle::SetAppStyleSetName(FEditorStyle::GetStyleSetName());
+		}
 	}
 
-	virtual TSharedRef<class FSlateStyleSet> CreateEditorStyleInstance( ) const override
+	virtual void ShutdownModule( ) override
 	{
-		return FSlateEditorStyle::Create(FSlateEditorStyle::Settings);
+		if (bUsingStarshipStyle)
+		{
+			FStarshipEditorStyle::Shutdown();
+		}
+		else
+		{
+			FSlateEditorStyle::Shutdown();
+		}
 	}
 
 	// End IModuleInterface interface
+private:
+	bool bUsingStarshipStyle = false;
 };
 
 
diff --git a/Engine/Source/Editor/EditorStyle/Private/EditorStyleSettingsCustomization.cpp b/Engine/Source/Editor/EditorStyle/Private/EditorStyleSettingsCustomization.cpp
new file mode 100644
--- /dev/null
+++ b/Engine/Source/Editor/EditorStyle/Private/EditorStyleSettingsCustomization.cpp
@@ -0,0 +1,22 @@
+// Copyright Epic Games, Inc. All Rights Reserved.
+
+#include "EditorStyleSettingsCustomization.h"
+#include "DetailCategoryBuilder.h"
+#include "Styling/StyleColors.h"
+#include "DetailLayoutBuilder.h"
+
+
+TSharedRef<IDetailCustomization> FEditorStyleSettingsCustomization::MakeInstance()
+{
+	return MakeShared<FEditorStyleSettingsCustomization>();
+}
+
+
+void FEditorStyleSettingsCustomization::CustomizeDetails(IDetailLayoutBuilder& DetailLayout)
+{
+	IDetailCategoryBuilder& ColorCategory = DetailLayout.EditCategory("Color");
+
+	TArray<UObject*> Objects = { &UStyleColorTable::Get() };
+
+	ColorCategory.AddExternalObjectProperty(Objects, "Colors");
+}
\ No newline at end of file
diff --git a/Engine/Source/Editor/EditorStyle/Private/EditorStyleSettingsCustomization.h b/Engine/Source/Editor/EditorStyle/Private/EditorStyleSettingsCustomization.h
new file mode 100644
--- /dev/null
+++ b/Engine/Source/Editor/EditorStyle/Private/EditorStyleSettingsCustomization.h
@@ -0,0 +1,18 @@
+// Copyright Epic Games, Inc. All Rights Reserved.
+
+#pragma once
+
+#include "IDetailCustomization.h"
+
+class IDetailLayoutBuilder;
+
+class FEditorStyleSettingsCustomization : public IDetailCustomization
+{
+public:
+	/** Makes a new instance of this detail layout class for a specific detail view requesting it */
+	static TSharedRef<IDetailCustomization> MakeInstance();
+
+	/** IDetailCustomization interface */
+	virtual void CustomizeDetails(IDetailLayoutBuilder& DetailLayout) override;
+};
+
diff --git a/Engine/Source/Editor/EditorStyle/Private/SlateEditorStyle.cpp b/Engine/Source/Editor/EditorStyle/Private/SlateEditorStyle.cpp
--- a/Engine/Source/Editor/EditorStyle/Private/SlateEditorStyle.cpp
+++ b/Engine/Source/Editor/EditorStyle/Private/SlateEditorStyle.cpp
@@ -33,7 +33,7 @@ BEGIN_SLATE_FUNCTION_BUILD_OPTIMIZATION
 TSharedPtr< FSlateEditorStyle::FStyle > FSlateEditorStyle::StyleInstance = NULL;
 TWeakObjectPtr< UEditorStyleSettings > FSlateEditorStyle::Settings = NULL;
 
-bool IncludeEditorSpecificStyles()
+static bool IncludeEditorSpecificStyles()
 {
 #if IS_PROGRAM
 	return true;
@@ -112,7 +112,7 @@ FSlateEditorStyle::FStyle::FStyle( const TWeakObjectPtr< UEditorStyleSettings >&
 
 }
 
-void SetColor( const TSharedRef< FLinearColor >& Source, const FLinearColor& Value )
+static void SetColor( const TSharedRef< FLinearColor >& Source, const FLinearColor& Value )
 {
 	Source->R = Value.R;
 	Source->G = Value.G;
@@ -133,10 +133,7 @@ void FSlateEditorStyle::FStyle::SyncSettings()
 	if ( Settings.IsValid() )
 	{
 		// Sync the colors used by FEditorStyle
-		SetColor( SelectorColor_LinearRef, Settings->KeyboardFocusColor );
 		SetColor( SelectionColor_LinearRef, Settings->SelectionColor );
-		SetColor( SelectionColor_Inactive_LinearRef, Settings->InactiveSelectionColor );
-		SetColor( SelectionColor_Pressed_LinearRef, Settings->PressedSelectionColor );
 
 		SetColor( LogColor_Background_LinearRef, Settings->LogBackgroundColor );
 		SetColor( LogColor_SelectionBackground_LinearRef, Settings->LogSelectionBackgroundColor );
@@ -150,10 +147,7 @@ void FSlateEditorStyle::FStyle::SyncSettings()
 		SetColor( SelectionColor_Subdued_LinearRef, SubduedSelectionColor );
 
 		// Also sync the colors used by FCoreStyle, as FEditorStyle isn't yet being used as an override everywhere
-		FCoreStyle::SetSelectorColor( Settings->KeyboardFocusColor );
 		FCoreStyle::SetSelectionColor( Settings->SelectionColor );
-		FCoreStyle::SetInactiveSelectionColor( Settings->InactiveSelectionColor );
-		FCoreStyle::SetPressedSelectionColor( Settings->PressedSelectionColor );
 
 		// Sync the window background settings
 		FSlateColor WindowBackgroundColor(FLinearColor::White);
@@ -225,6 +219,8 @@ void FSlateEditorStyle::FStyle::Initialize()
 	SetupMyBlueprintStyles();
 
 //	LogUnusedBrushResources();
+
+	SetParentStyleName( FCoreStyle::GetCoreStyle().GetStyleSetName() );
 }
 
 void FSlateEditorStyle::FStyle::SetupGeneralStyles()
@@ -425,7 +421,11 @@ void FSlateEditorStyle::FStyle::SetupGeneralStyles()
 		.SetHovered( BOX_BRUSH( "Common/Button_Hovered", FVector2D(32,32), 8.0f/32.0f ) )
 		.SetPressed( BOX_BRUSH( "Common/Button_Pressed", FVector2D(32,32), 8.0f/32.0f ) )
 		.SetNormalPadding( FMargin( 2,2,2,2 ) )
-		.SetPressedPadding( FMargin( 2,3,2,1 ) );
+		.SetPressedPadding( FMargin( 2,3,2,1 ) )
+		.SetNormalForeground(InvertedForeground)
+		.SetPressedForeground(InvertedForeground)
+		.SetHoveredForeground(InvertedForeground)
+		.SetDisabledForeground(InvertedForeground);
 	Set( "Button", Button );
 	Set( "Button.Disabled", new BOX_BRUSH( "Common/Button_Disabled", 8.0f/32.0f ) );
 	
@@ -1612,28 +1612,6 @@ void FSlateEditorStyle::FStyle::SetupGeneralStyles()
 	Set( "Editor.SearchBoxFont", DEFAULT_FONT( "Regular", 12) );
 #endif // WITH_EDITOR || (IS_PROGRAM && WITH_UNREAL_DEVELOPER_TOOLS)
 
-	// Slider and Volume Control
-	{
-		FSliderStyle SliderStyle = FSliderStyle()
-			.SetNormalBarImage(FSlateColorBrush(FColor::White))
-			.SetHoveredBarImage(FSlateColorBrush(FColor::White))
-			.SetDisabledBarImage(FSlateColorBrush(FLinearColor::Gray))
-			.SetNormalThumbImage( BOX_BRUSH( "Common/Button", 8.0f/32.0f ) )
-			.SetHoveredThumbImage(BOX_BRUSH("Common/Button", 8.0f / 32.0f))
-			.SetDisabledThumbImage( BOX_BRUSH( "Common/Button_Disabled", 8.0f/32.0f ) )
-			.SetBarThickness(2.0f);
-		Set( "Slider", SliderStyle );
-
-		Set( "VolumeControl", FVolumeControlStyle()
-			.SetSliderStyle( SliderStyle )
-			.SetHighVolumeImage( IMAGE_BRUSH( "Common/VolumeControl_High", Icon16x16 ) )
-			.SetMidVolumeImage( IMAGE_BRUSH( "Common/VolumeControl_Mid", Icon16x16 ) )
-			.SetLowVolumeImage( IMAGE_BRUSH( "Common/VolumeControl_Low", Icon16x16 ) )
-			.SetNoVolumeImage( IMAGE_BRUSH( "Common/VolumeControl_Off", Icon16x16 ) )
-			.SetMutedImage( IMAGE_BRUSH( "Common/VolumeControl_Muted", Icon16x16 ) )
-			);
-	}
-
 	// Console
 	{
 		Set( "DebugConsole.Background", new BOX_BRUSH( "Old/Menu_Background", FMargin(8.0f/64.0f) ) );
@@ -2056,102 +2034,6 @@ void FSlateEditorStyle::FStyle::SetupGeneralStyles()
 #endif // WITH_EDITOR
 
 #if WITH_EDITOR || (IS_PROGRAM && WITH_UNREAL_DEVELOPER_TOOLS)
-	// ToolBar
-	{
-		Set( "ToolBar.Background", new BOX_BRUSH( "Common/GroupBorder", FMargin(4.0f/16.0f) ) );
-		Set( "ToolBar.Icon", new IMAGE_BRUSH( "Icons/icon_tab_Toolbars_16x", Icon16x16 ) );
-		Set( "ToolBar.Expand", new IMAGE_BRUSH( "Icons/toolbar_expand_16x", Icon16x16) );
-		Set( "ToolBar.SubMenuIndicator", new IMAGE_BRUSH( "Common/SubmenuArrow", Icon8x8 ) );
-		Set( "ToolBar.SToolBarComboButtonBlock.Padding", FMargin(4.0f));
-		Set( "ToolBar.SToolBarButtonBlock.Padding", FMargin(4.0f));
-		Set( "ToolBar.SToolBarCheckComboButtonBlock.Padding", FMargin(4.0f));
-		Set( "ToolBar.SToolBarButtonBlock.CheckBox.Padding", FMargin(0.0f) );
-		Set( "ToolBar.SToolBarComboButtonBlock.ComboButton.Color", DefaultForeground );
-
-		Set( "ToolBar.Block.IndentedPadding", FMargin( 18.0f, 2.0f, 4.0f, 4.0f ) );
-		Set( "ToolBar.Block.Padding", FMargin( 2.0f, 2.0f, 4.0f, 4.0f ) );
-
-		Set( "ToolBar.Separator", new FSlateColorBrush( FLinearColor(FColor(47, 47, 47)) ) );
-		Set( "ToolBar.Separator.Padding", FMargin( 8.f, 0.f, 8.f, 0.f) );
-
-		Set( "ToolBar.Label", FTextBlockStyle(NormalText) .SetFont( DEFAULT_FONT( "Regular", 9 ) ) );
-		Set( "ToolBar.EditableText", FEditableTextBoxStyle(NormalEditableTextBoxStyle) .SetFont( DEFAULT_FONT( "Regular", 9 ) ) );
-		Set( "ToolBar.Keybinding", FTextBlockStyle(NormalText) .SetFont( DEFAULT_FONT( "Regular", 8 ) ) );
-
-		Set( "ToolBar.Heading", FTextBlockStyle(NormalText)
-			.SetFont( DEFAULT_FONT( "Regular", 8 ) )
-			.SetColorAndOpacity( FLinearColor( 0.4f, 0.4, 0.4f, 1.0f ) ) );
-
-		/* Create style for "ToolBar.CheckBox" widget ... */
-		const FCheckBoxStyle ToolBarCheckBoxStyle = FCheckBoxStyle()
-			.SetUncheckedImage( IMAGE_BRUSH( "Common/SmallCheckBox", Icon14x14 ) )
-			.SetCheckedImage( IMAGE_BRUSH( "Common/SmallCheckBox_Checked",  Icon14x14 ) )
-			.SetUncheckedHoveredImage( IMAGE_BRUSH( "Common/SmallCheckBox_Hovered", Icon14x14 ) )
-			.SetCheckedImage( IMAGE_BRUSH( "Common/SmallCheckBox_Checked_Hovered", Icon14x14 ) )
-			.SetUncheckedPressedImage(IMAGE_BRUSH( "Common/SmallCheckBox_Hovered", Icon14x14, FLinearColor( 0.5f, 0.5f, 0.5f ) ) )
-			.SetCheckedPressedImage( IMAGE_BRUSH( "Common/SmallCheckBox_Checked_Hovered", Icon14x14, FLinearColor( 0.5f, 0.5f, 0.5f ) ) )
-			.SetUndeterminedImage( IMAGE_BRUSH( "Common/CheckBox_Undetermined", Icon14x14 ) )
-			.SetUndeterminedHoveredImage( IMAGE_BRUSH( "Common/CheckBox_Undetermined_Hovered", Icon14x14 ) )
-			.SetUndeterminedPressedImage( IMAGE_BRUSH( "Common/CheckBox_Undetermined_Hovered", Icon14x14, FLinearColor( 0.5f, 0.5f, 0.5f ) ) );
-		/* ... and add new style */
-		Set( "ToolBar.CheckBox", ToolBarCheckBoxStyle );
-
-		/* Read-only checkbox that appears next to a toolbar item */
-		/* Set images for various SCheckBox states associated with read-only toolbar check box items... */
-		const FCheckBoxStyle ToolBarCheckStyle = FCheckBoxStyle()
-			.SetUncheckedImage(IMAGE_BRUSH("Icons/Empty_14x", Icon14x14))
-			.SetUncheckedHoveredImage(IMAGE_BRUSH("Icons/Empty_14x", Icon14x14))
-			.SetUncheckedPressedImage(IMAGE_BRUSH("Common/SmallCheckBox_Hovered", Icon14x14))
-			.SetCheckedImage(IMAGE_BRUSH("Common/SmallCheck", Icon14x14))
-			.SetCheckedHoveredImage(IMAGE_BRUSH("Common/SmallCheck", Icon14x14))
-			.SetCheckedPressedImage(IMAGE_BRUSH("Common/SmallCheck", Icon14x14))
-			.SetUndeterminedImage(IMAGE_BRUSH("Icons/Empty_14x", Icon14x14))
-			.SetUndeterminedHoveredImage(FSlateNoResource())
-			.SetUndeterminedPressedImage(FSlateNoResource());
-
-		/* ...and add the new style */
-		Set("Toolbar.Check", ToolBarCheckStyle);
-
-		// This radio button is actually just a check box with different images
-		/* Create style for "ToolBar.RadioButton" widget ... */
-		const FCheckBoxStyle ToolbarRadioButtonCheckBoxStyle = FCheckBoxStyle()
-				.SetUncheckedImage( IMAGE_BRUSH( "Common/RadioButton_Unselected_16x", Icon16x16 ) )
-				.SetCheckedImage( IMAGE_BRUSH( "Common/RadioButton_Selected_16x",  Icon16x16 ) )
-				.SetUncheckedHoveredImage( IMAGE_BRUSH( "Common/RadioButton_Unselected_16x", Icon16x16, SelectionColor ) )
-				.SetCheckedHoveredImage( IMAGE_BRUSH( "Common/RadioButton_Selected_16x", Icon16x16, SelectionColor ) )
-				.SetUncheckedPressedImage( IMAGE_BRUSH( "Common/RadioButton_Unselected_16x", Icon16x16, SelectionColor_Pressed ) )
-				.SetCheckedPressedImage( IMAGE_BRUSH( "Common/RadioButton_Selected_16x", Icon16x16, SelectionColor_Pressed ) );
-
-		/* ... and add new style */
-		Set( "ToolBar.RadioButton", ToolbarRadioButtonCheckBoxStyle );
-
-		/* Create style for "ToolBar.ToggleButton" widget ... */
-		const FCheckBoxStyle ToolBarToggleButtonCheckBoxStyle = FCheckBoxStyle()
-				.SetCheckBoxType( ESlateCheckBoxType::ToggleButton )
-				.SetUncheckedImage( FSlateNoResource() )
-				.SetUncheckedPressedImage( BOX_BRUSH( "Common/RoundedSelection_16x", 4.0f/16.0f, SelectionColor_Pressed ) )
-				.SetUncheckedHoveredImage( BOX_BRUSH( "Common/RoundedSelection_16x", 4.0f/16.0f, SelectionColor ) )
-				.SetCheckedImage( BOX_BRUSH( "Common/RoundedSelection_16x",  4.0f/16.0f, SelectionColor_Pressed ) )
-				.SetCheckedHoveredImage( BOX_BRUSH( "Common/RoundedSelection_16x",  4.0f/16.0f, SelectionColor_Pressed ) )
-				.SetCheckedPressedImage( BOX_BRUSH( "Common/RoundedSelection_16x",  4.0f/16.0f, SelectionColor ) );
-		/* ... and add new style */
-		Set( "ToolBar.ToggleButton", ToolBarToggleButtonCheckBoxStyle );
-
-		Set( "ToolBar.Button", FButtonStyle(Button)
-			.SetNormal ( FSlateNoResource() )
-			.SetPressed( BOX_BRUSH( "Common/RoundedSelection_16x", 4.0f/16.0f, SelectionColor_Pressed ) )
-			.SetHovered( BOX_BRUSH( "Common/RoundedSelection_16x", 4.0f/16.0f, SelectionColor ) )
-		);
-
-		Set( "ToolBar.Button.Normal", new FSlateNoResource() );
-		Set( "ToolBar.Button.Pressed", new BOX_BRUSH( "Common/RoundedSelection_16x", 4.0f/16.0f, SelectionColor_Pressed ) );
-		Set( "ToolBar.Button.Hovered", new BOX_BRUSH( "Common/RoundedSelection_16x", 4.0f/16.0f, SelectionColor ) );
-
-		Set( "ToolBar.Button.Checked", new BOX_BRUSH( "Common/RoundedSelection_16x",  4.0f/16.0f, SelectionColor_Pressed ) );
-		Set( "ToolBar.Button.Checked_Hovered", new BOX_BRUSH( "Common/RoundedSelection_16x",  4.0f/16.0f, SelectionColor_Pressed ) );
-		Set( "ToolBar.Button.Checked_Pressed", new BOX_BRUSH( "Common/RoundedSelection_16x",  4.0f/16.0f, SelectionColor ) );
-	}
-
 	// Mode ToolPalette 
 	{
 		Set( "PaletteToolBar.Background", 				new BOX_BRUSH( "Common/GroupBorder", FMargin(4.0f/16.0f) ));
@@ -2774,10 +2656,10 @@ void FSlateEditorStyle::FStyle::SetupSequencerStyles()
 	Set( "Sequencer.ToolBar.SToolBarButtonBlock.Padding", 			 FCoreStyle::Get().GetMargin("ToolBar.SToolBarButtonBlock.Padding"));
 	Set( "Sequencer.ToolBar.SToolBarCheckComboButtonBlock.Padding",  FCoreStyle::Get().GetMargin("ToolBar.SToolBarCheckComboButtonBlock.Padding"));
 	Set( "Sequencer.ToolBar.SToolBarButtonBlock.CheckBox.Padding", 	 FCoreStyle::Get().GetMargin("ToolBar.SToolBarButtonBlock.CheckBox.Padding"));
-	Set( "Sequencer.ToolBar.SToolBarComboButtonBlock.ComboButton.Color", GetColor("ToolBar.SToolBarComboButtonBlock.ComboButton.Color"));
+	Set( "Sequencer.ToolBar.SToolBarComboButtonBlock.ComboButton.Color", GetSlateColor("ToolBar.SToolBarComboButtonBlock.ComboButton.Color"));
 
 	// Used only for WidgetBlocks
-	Set( "Sequencer.ToolBar.Block.IndentedPadding", 	FCoreStyle::Get().GetMargin("Sequencer.ToolBar.Block.IndentedPadding") );
+	Set( "Sequencer.ToolBar.Block.IndentedPadding", 	FCoreStyle::Get().GetMargin("ToolBar.Block.IndentedPadding") );
 	Set( "Sequencer.ToolBar.Block.Padding", 			FMargin(0.0f, 4.0f));
 
 	Set( "Sequencer.ToolBar.Separator", 				new FSlateColorBrush( FLinearColor(FColor(48, 48, 48)) ) );
@@ -3077,115 +2959,6 @@ void FSlateEditorStyle::FStyle::SetupNotificationBarStyles()
 
 void FSlateEditorStyle::FStyle::SetupMenuBarStyles()
 {
-	// MenuBar
-	{
-		Set("Menu.Background", new BOX_BRUSH("Old/Menu_Background", FMargin(8.0f / 64.0f)));
-		Set("Menu.Icon", new IMAGE_BRUSH("Icons/icon_tab_toolbar_16px", Icon16x16));
-		Set("Menu.Expand", new IMAGE_BRUSH("Icons/toolbar_expand_16x", Icon16x16));
-		Set("Menu.SubMenuIndicator", new IMAGE_BRUSH("Common/SubmenuArrow", Icon8x8));
-		Set("Menu.SToolBarComboButtonBlock.Padding", FMargin(4.0f));
-		Set("Menu.SToolBarButtonBlock.Padding", FMargin(4.0f));
-		Set("Menu.SToolBarCheckComboButtonBlock.Padding", FMargin(4.0f));
-		Set("Menu.SToolBarButtonBlock.CheckBox.Padding", FMargin(0.0f));
-		Set("Menu.SToolBarComboButtonBlock.ComboButton.Color", DefaultForeground);
-
-		Set("Menu.Block.IndentedPadding", FMargin(18.0f, 2.0f, 4.0f, 4.0f));
-		Set("Menu.Block.Padding", FMargin(2.0f, 2.0f, 4.0f, 4.0f));
-
-		Set("Menu.Separator", new BOX_BRUSH("Old/Button", 4.0f / 32.0f));
-		Set("Menu.Separator.Padding", FMargin(0.5f));
-
-		Set("Menu.Label", FTextBlockStyle(NormalText).SetFont(DEFAULT_FONT("Regular", 9)));
-		Set("Menu.Label.Padding", FMargin(0.0f, 0.0f, 0.0f, 0.0f));
-		Set("Menu.Label.ContentPadding", FMargin(10.0f, 2.0f));
-		Set("Menu.EditableText", FEditableTextBoxStyle(NormalEditableTextBoxStyle).SetFont(DEFAULT_FONT("Regular", 9)));
-		Set("Menu.Keybinding", FTextBlockStyle(NormalText).SetFont(DEFAULT_FONT("Regular", 8)));
-
-		Set("Menu.Heading", FTextBlockStyle(NormalText)
-			.SetFont(DEFAULT_FONT("Regular", 8))
-			.SetColorAndOpacity(FLinearColor(0.4f, 0.4, 0.4f, 1.0f)));
-
-		/* Set images for various SCheckBox states associated with menu check box items... */
-		const FCheckBoxStyle BasicMenuCheckBoxStyle = FCheckBoxStyle()
-			.SetUncheckedImage(IMAGE_BRUSH("Common/SmallCheckBox", Icon14x14))
-			.SetUncheckedHoveredImage(IMAGE_BRUSH("Common/SmallCheckBox_Hovered", Icon14x14))
-			.SetUncheckedPressedImage(IMAGE_BRUSH("Common/SmallCheckBox_Hovered", Icon14x14, FLinearColor(0.5f, 0.5f, 0.5f)))
-			.SetCheckedImage(IMAGE_BRUSH("Common/SmallCheckBox_Checked", Icon14x14))
-			.SetCheckedHoveredImage(IMAGE_BRUSH("Common/SmallCheckBox_Checked_Hovered", Icon14x14))
-			.SetCheckedPressedImage(IMAGE_BRUSH("Common/SmallCheckBox_Checked_Hovered", Icon14x14, FLinearColor(0.5f, 0.5f, 0.5f)))
-			.SetUndeterminedImage(IMAGE_BRUSH("Common/CheckBox_Undetermined", Icon14x14))
-			.SetUndeterminedHoveredImage(IMAGE_BRUSH("Common/CheckBox_Undetermined_Hovered", Icon14x14))
-			.SetUndeterminedPressedImage(IMAGE_BRUSH("Common/CheckBox_Undetermined_Hovered", Icon14x14, FLinearColor(0.5f, 0.5f, 0.5f)));
-
-		/* ...and add the new style */
-		Set("Menu.CheckBox", BasicMenuCheckBoxStyle);
-
-		/* Read-only checkbox that appears next to a menu item */
-		/* Set images for various SCheckBox states associated with read-only menu check box items... */
-		const FCheckBoxStyle BasicMenuCheckStyle = FCheckBoxStyle()
-			.SetUncheckedImage(IMAGE_BRUSH("Icons/Empty_14x", Icon14x14))
-			.SetUncheckedHoveredImage(IMAGE_BRUSH("Icons/Empty_14x", Icon14x14))
-			.SetUncheckedPressedImage(IMAGE_BRUSH("Common/SmallCheckBox_Hovered", Icon14x14))
-			.SetCheckedImage(IMAGE_BRUSH("Common/SmallCheck", Icon14x14))
-			.SetCheckedHoveredImage(IMAGE_BRUSH("Common/SmallCheck", Icon14x14))
-			.SetCheckedPressedImage(IMAGE_BRUSH("Common/SmallCheck", Icon14x14))
-			.SetUndeterminedImage(IMAGE_BRUSH("Icons/Empty_14x", Icon14x14))
-			.SetUndeterminedHoveredImage(FSlateNoResource())
-			.SetUndeterminedPressedImage(FSlateNoResource());
-
-		/* ...and add the new style */
-		Set("Menu.Check", BasicMenuCheckStyle);
-
-		/* This radio button is actually just a check box with different images */
-		/* Set images for various Menu radio button (SCheckBox) states... */
-		const FCheckBoxStyle BasicMenuRadioButtonStyle = FCheckBoxStyle()
-			.SetUncheckedImage(IMAGE_BRUSH("Common/RadioButton_Unselected_16x", Icon16x16))
-			.SetUncheckedHoveredImage(IMAGE_BRUSH("Common/RadioButton_Unselected_16x", Icon16x16))
-			.SetUncheckedPressedImage(IMAGE_BRUSH("Common/RadioButton_Unselected_16x", Icon16x16))
-			.SetCheckedImage(IMAGE_BRUSH("Common/RadioButton_Selected_16x", Icon16x16))
-			.SetCheckedHoveredImage(IMAGE_BRUSH("Common/RadioButton_Selected_16x", Icon16x16, SelectionColor))
-			.SetCheckedPressedImage(IMAGE_BRUSH("Common/RadioButton_Unselected_16x", Icon16x16, SelectionColor_Pressed))
-			.SetUndeterminedImage(IMAGE_BRUSH("Common/RadioButton_Unselected_16x", Icon16x16))
-			.SetUndeterminedHoveredImage(IMAGE_BRUSH("Common/RadioButton_Unselected_16x", Icon16x16, SelectionColor))
-			.SetUndeterminedPressedImage(IMAGE_BRUSH("Common/RadioButton_Unselected_16x", Icon16x16, SelectionColor_Pressed));
-
-		/* ...and set new style */
-		Set("Menu.RadioButton", BasicMenuRadioButtonStyle);
-
-		/* Create style for "Menu.ToggleButton" widget ... */
-		const FCheckBoxStyle MenuToggleButtonCheckBoxStyle = FCheckBoxStyle()
-			.SetCheckBoxType(ESlateCheckBoxType::ToggleButton)
-			.SetUncheckedImage(FSlateNoResource())
-			.SetUncheckedPressedImage(BOX_BRUSH("Common/RoundedSelection_16x", 4.0f / 16.0f, SelectionColor_Pressed))
-			.SetUncheckedHoveredImage(BOX_BRUSH("Common/RoundedSelection_16x", 4.0f / 16.0f, SelectionColor))
-			.SetCheckedImage(BOX_BRUSH("Common/RoundedSelection_16x", 4.0f / 16.0f, SelectionColor_Pressed))
-			.SetCheckedHoveredImage(BOX_BRUSH("Common/RoundedSelection_16x", 4.0f / 16.0f, SelectionColor_Pressed))
-			.SetCheckedPressedImage(BOX_BRUSH("Common/RoundedSelection_16x", 4.0f / 16.0f, SelectionColor));
-		/* ... and add new style */
-		Set("Menu.ToggleButton", MenuToggleButtonCheckBoxStyle);
-
-		const FButtonStyle NoBorder = FButtonStyle()
-			.SetNormal(FSlateNoResource())
-			.SetHovered(FSlateNoResource())
-			.SetPressed(FSlateNoResource())
-			.SetNormalPadding(FMargin(0, 0, 0, 1))
-			.SetPressedPadding(FMargin(0, 1, 0, 0));
-
-		Set("Menu.Button", FButtonStyle(NoBorder)
-			.SetNormal(FSlateNoResource())
-			.SetPressed(BOX_BRUSH("Common/RoundedSelection_16x", 4.0f / 16.0f, SelectionColor_Pressed))
-			.SetHovered(BOX_BRUSH("Common/RoundedSelection_16x", 4.0f / 16.0f, SelectionColor))
-			.SetNormalPadding(FMargin(0, 1))
-			.SetPressedPadding(FMargin(0, 2, 0, 0))
-		);
-
-		Set("Menu.Button.Checked", new BOX_BRUSH("Common/RoundedSelection_16x", 4.0f / 16.0f, SelectionColor_Pressed));
-		Set("Menu.Button.Checked_Hovered", new BOX_BRUSH("Common/RoundedSelection_16x", 4.0f / 16.0f, SelectionColor_Pressed));
-		Set("Menu.Button.Checked_Pressed", new BOX_BRUSH("Common/RoundedSelection_16x", 4.0f / 16.0f, SelectionColor));
-
-		/* The style of a menu bar button when it has a sub menu open */
-		Set("Menu.Button.SubMenuOpen", new BORDER_BRUSH("Common/Selection", FMargin(4.f / 16.f), FLinearColor(0.10f, 0.10f, 0.10f)));
-	}
 }
 
 void FSlateEditorStyle::FStyle::SetupGeneralIcons()
@@ -3296,6 +3069,7 @@ void FSlateEditorStyle::FStyle::SetupProjectBadgeStyle()
 	);
 
 	Set("SProjectBadge.BadgeShape", new BOX_BRUSH("ProjectBadge/Badge", Icon16x16, FMargin(6.0f/16)));
+	Set("SProjectBadge.BadgePadding", FMargin(10.0f, 2.5f));
 }
 
 void FSlateEditorStyle::FStyle::SetupDockingStyles()
@@ -5668,6 +5442,11 @@ void FSlateEditorStyle::FStyle::SetupLevelEditorStyle()
 		}
 	}
 
+	// Level Editor Toolbar
+	{
+		Set("LevelEditorToolBar.Background", new FSlateNoResource());
+	}
+
 	// Scalability (Performance Warning)
 	{
 		Set( "Scalability.ScalabilitySettings", new IMAGE_BRUSH("Scalability/ScalabilitySettings", FVector2D(473.0f, 266.0f) ) );
@@ -6149,6 +5928,29 @@ void FSlateEditorStyle::FStyle::SetupPersonaStyle()
 
 	// Play in editor / play in world
 	{
+		const FButtonStyle PlayToolbarButton = FButtonStyle(FCoreStyle::Get().GetWidgetStyle<FButtonStyle>("ToolBar.Button"))
+			.SetPressed(FSlateNoResource())
+			.SetHovered(FSlateNoResource())
+			.SetNormalForeground(FSlateColor::UseForeground())
+			.SetPressedForeground(FSlateColor::UseForeground())
+			.SetHoveredForeground(FSlateColor::UseForeground());
+
+		Set("PlayToolbar.Button", PlayToolbarButton);
+
+
+		FComboButtonStyle PlayToolbarComboButton = FComboButtonStyle(FCoreStyle::Get().GetWidgetStyle<FComboButtonStyle>("ComboButton"));
+		PlayToolbarComboButton.ButtonStyle = PlayToolbarButton;
+
+		Set("PlayToolbar.SimpleComboButton", PlayToolbarComboButton);
+
+
+		Set("PlayToolbar.Label", FTextBlockStyle(NormalText).SetFont(DEFAULT_FONT("Regular", 10)));
+		Set("PlayToolbar.Label.Padding", FMargin(5, 0, 0, 0));
+		Set("PlayToolbar.Label.DropShadowSize", FVector2D::ZeroVector);
+		Set("PlayToolbar.SToolBarButtonBlock.Padding", FMargin(5.0f, 0.f));
+		Set("PlayToolbar.IconSize", Icon40x40);
+
+
 		Set( "PlayWorld.Simulate", new IMAGE_BRUSH( "Icons/icon_simulate_40x", Icon40x40 ) );
 		Set( "PlayWorld.Simulate.Small", new IMAGE_BRUSH( "Icons/icon_simulate_40x", Icon20x20 ) );
 
diff --git a/Engine/Source/Editor/EditorStyle/Private/StarshipStyle.cpp b/Engine/Source/Editor/EditorStyle/Private/StarshipStyle.cpp
--- a/Engine/Source/Editor/EditorStyle/Private/StarshipStyle.cpp
+++ b/Engine/Source/Editor/EditorStyle/Private/StarshipStyle.cpp
@@ -13,6 +13,7 @@
 	#include "PlatformInfo.h"
 #endif
 #include "Styling/ToolBarStyle.h"
+#include "Styling/StyleColors.h"
 
 
 #define DEFAULT_FONT(...) FStarshipCoreStyle::GetDefaultFontStyle(__VA_ARGS__)
@@ -128,10 +129,11 @@ void FStarshipEditorStyle::FStyle::SyncSettings()
 		SetColor(SelectionColor_Subdued_LinearRef, SubduedSelectionColor);
 
 		// Also sync the colors used by FStarshipCoreStyle, as FEditorStyle isn't yet being used as an override everywhere
-		FStarshipCoreStyle::SetSelectorColor(Settings->KeyboardFocusColor);
+	/*	FStarshipCoreStyle::SetSelectorColor(Settings->KeyboardFocusColor);
 		FStarshipCoreStyle::SetSelectionColor(Settings->SelectionColor);
 		FStarshipCoreStyle::SetInactiveSelectionColor(Settings->InactiveSelectionColor);
 		FStarshipCoreStyle::SetPressedSelectionColor(Settings->PressedSelectionColor);
+*/
 
 		// Sync the window background settings
 		FWindowStyle& WindowStyle = const_cast<FWindowStyle&>(FStarshipCoreStyle::GetCoreStyle().GetWidgetStyle<FWindowStyle>("Window"));
diff --git a/Engine/Source/Editor/EditorStyle/Private/StarshipStyle.h b/Engine/Source/Editor/EditorStyle/Private/StarshipStyle.h
--- a/Engine/Source/Editor/EditorStyle/Private/StarshipStyle.h
+++ b/Engine/Source/Editor/EditorStyle/Private/StarshipStyle.h
@@ -12,6 +12,12 @@
 #include "Classes/EditorStyleSettings.h"
 #include "ISettingsModule.h"
 
+#if WITH_EDITOR
+#include "EditorStyleSettingsCustomization.h"
+#include "PropertyEditorModule.h"
+#endif
+
+
 struct FPropertyChangedEvent;
 
 /**
@@ -38,6 +44,10 @@ public:
 				Settings
 			);
 		}
+
+
+		FPropertyEditorModule& PropertyEditorModule = FModuleManager::Get().GetModuleChecked<FPropertyEditorModule>("PropertyEditor");
+		PropertyEditorModule.RegisterCustomClassLayout("EditorStyleSettings", FOnGetDetailCustomizationInstance::CreateStatic(&FEditorStyleSettingsCustomization::MakeInstance));
 #endif
 
 		StyleInstance = Create(Settings);
@@ -46,6 +56,7 @@ public:
 
 	static void Shutdown()
 	{
+#if WITH_EDITOR
 		ISettingsModule* SettingsModule = FModuleManager::GetModulePtr<ISettingsModule>("Settings");
 
 		if (SettingsModule != nullptr)
@@ -53,6 +64,14 @@ public:
 			SettingsModule->UnregisterSettings("Editor", "General", "Appearance");
 		}
 
+
+
+		FPropertyEditorModule* PropertyEditorModule = FModuleManager::Get().GetModulePtr<FPropertyEditorModule>("PropertyEditor");
+		if (PropertyEditorModule)
+		{
+			PropertyEditorModule->UnregisterCustomClassLayout("EditorStyleSettings");
+		}
+#endif
 		//FEditorStyle::SetStyle()
 		StyleInstance.Reset();
 	}
diff --git a/Engine/Source/Editor/EditorStyle/Public/Classes/EditorStyleSettings.h b/Engine/Source/Editor/EditorStyle/Public/Classes/EditorStyleSettings.h
--- a/Engine/Source/Editor/EditorStyle/Public/Classes/EditorStyleSettings.h
+++ b/Engine/Source/Editor/EditorStyle/Public/Classes/EditorStyleSettings.h
@@ -74,20 +74,11 @@ public:
 	UPROPERTY(EditAnywhere, config, Category=Colors, meta=(DisplayName="Selection Color"))
 	FLinearColor SelectionColor;
 
-	/** The color used to represent a pressed item */
-	UPROPERTY(EditAnywhere, config, Category=Colors, meta=(DisplayName="Pressed Selection Color"))
-	FLinearColor PressedSelectionColor;
-
-	/** The color used to represent selected items that are currently inactive */
-	UPROPERTY(EditAnywhere, config, Category=Colors, meta=(DisplayName="Inactive Selection Color"))
-	FLinearColor InactiveSelectionColor;
-
-	/** The color used to represent keyboard input selection focus */
-	UPROPERTY(EditAnywhere, config, Category=Colors, meta=(DisplayName="Keyboard Focus Color"), AdvancedDisplay)
-	FLinearColor KeyboardFocusColor;
+	UPROPERTY(config)
+	bool bEnableEditorWindowBackgroundColor;
 
 	/** The color used to tint the editor window backgrounds */
-	UPROPERTY(EditAnywhere, config, Category=Colors)
+	UPROPERTY(EditAnywhere, config, Category=Colors, meta=(EditCondition="bEnableEditorWindowBackgroundColor"))
 	FLinearColor EditorWindowBackgroundColor;
 
 	/** The override for the background of the main window (if not modified, the defaults will be used) */
@@ -98,10 +89,6 @@ public:
 	UPROPERTY(EditAnywhere, config, Category=Colors)
 	FSlateBrush EditorChildWindowBackgroundOverride;
 
-	/** Check to reset the window background settings to editor defaults */
-	UPROPERTY(EditAnywhere, config, Category=Colors)
-	bool bResetEditorWindowBackgroundSettings;
-
 	/** Whether to use small toolbar icons without labels or not. */
 	UPROPERTY(EditAnywhere, config, Category=UserInterface)
 	uint32 bUseSmallToolBarIcons:1;
diff --git a/Engine/Source/Editor/PropertyEditor/Public/DetailLayoutBuilder.h b/Engine/Source/Editor/PropertyEditor/Public/DetailLayoutBuilder.h
--- a/Engine/Source/Editor/PropertyEditor/Public/DetailLayoutBuilder.h
+++ b/Engine/Source/Editor/PropertyEditor/Public/DetailLayoutBuilder.h
@@ -5,7 +5,6 @@
 #include "CoreMinimal.h"
 #include "Fonts/SlateFontInfo.h"
 #include "EditorStyleSet.h"
-#include "AssetThumbnail.h"
 #include "PropertyHandle.h"
 #include "IDetailPropertyRow.h"
 
diff --git a/Engine/Source/Editor/PropertyEditor/Public/PropertyEditorModule.h b/Engine/Source/Editor/PropertyEditor/Public/PropertyEditorModule.h
--- a/Engine/Source/Editor/PropertyEditor/Public/PropertyEditorModule.h
+++ b/Engine/Source/Editor/PropertyEditor/Public/PropertyEditorModule.h
@@ -6,7 +6,6 @@
 #include "Widgets/SWindow.h"
 #include "Modules/ModuleInterface.h"
 #include "UObject/StructOnScope.h"
-#include "Toolkits/IToolkitHost.h"
 #include "IDetailsView.h"
 #include "PropertyEditorDelegates.h"
 #include "IPropertyTypeCustomization.h"
@@ -64,6 +63,13 @@ class SWindow;
 class IPropertyTableCellPresenter;
 class IPropertyTypeCustomization;
 class IDetailsView;
+class IToolkitHost;
+class FAssetEditorToolkit;
+
+namespace EToolkitMode
+{
+	enum Type;
+};
 
 /**
  * Base class for adding an extra data to identify a custom property type
diff --git a/Engine/Source/Runtime/Slate/Private/Framework/Application/SlateApplication.cpp b/Engine/Source/Runtime/Slate/Private/Framework/Application/SlateApplication.cpp
--- a/Engine/Source/Runtime/Slate/Private/Framework/Application/SlateApplication.cpp
+++ b/Engine/Source/Runtime/Slate/Private/Framework/Application/SlateApplication.cpp
@@ -51,6 +51,7 @@
 #include "Windows/WindowsHWrapper.h"
 #endif
 #include "Debugging/SlateDebugging.h"
+#include "Styling/StyleColors.h"
 
 
 CSV_DECLARE_CATEGORY_MODULE_EXTERN(CORE_API, Basic);
@@ -6561,6 +6562,7 @@ void FSlateApplication::InitializeCoreStyle()
 {
 	if (FCoreStyle::IsStarshipStyle())
 	{
+		UStyleColorTable::Get();
 		FStarshipCoreStyle::ResetToDefault();
 		FAppStyle::SetAppStyleSet(FStarshipCoreStyle::GetCoreStyle());
 	}
diff --git a/Engine/Source/Runtime/SlateCore/Private/Rendering/DrawElements.cpp b/Engine/Source/Runtime/SlateCore/Private/Rendering/DrawElements.cpp
--- a/Engine/Source/Runtime/SlateCore/Private/Rendering/DrawElements.cpp
+++ b/Engine/Source/Runtime/SlateCore/Private/Rendering/DrawElements.cpp
@@ -97,7 +97,7 @@ static bool ShouldCull(const FSlateWindowElementList& ElementList, const FPaintG
 
 static bool ShouldCull(const FSlateWindowElementList& ElementList, const FPaintGeometry& PaintGeometry, const FSlateBrush* InBrush, const FLinearColor& InTint)
 {
-	if (InTint.A == 0 || ShouldCull(ElementList, PaintGeometry, InBrush))
+	if ((InTint.A == 0 && InBrush->OutlineSettings.Color.GetSpecifiedColor().A == 0) || ShouldCull(ElementList, PaintGeometry, InBrush))
 	{
 		return true;
 	}
@@ -210,17 +210,37 @@ FSlateDrawElement& FSlateDrawElement::MakeBoxInternal(
 	const FLinearColor& InTint
 )
 {
-	EElementType ElementType = (InBrush->DrawAs == ESlateBrushDrawType::Border) ? EElementType::ET_Border : EElementType::ET_Box;
+	EElementType ElementType = (InBrush->DrawAs == ESlateBrushDrawType::Border) ? EElementType::ET_Border : (InBrush->DrawAs == ESlateBrushDrawType::RoundedBox) ? EElementType::ET_RoundedBox : EElementType::ET_Box;
+
+	// Cast to Rounded Rect to get the internal paramerters 
+	// New payload type - inherit from BoxPayload 
 
 	FSlateDrawElement& Element = ElementList.AddUninitialized();
 
 	const FMargin& Margin = InBrush->GetMargin();
-	FSlateBoxPayload& BoxPayload = ElementList.CreatePayload<FSlateBoxPayload>(Element);
+
+	FSlateBoxPayload* BoxPayload;
+	if ( ElementType == EElementType::ET_RoundedBox )
+	{
+		FSlateRoundedBoxPayload* RBoxPayload = &ElementList.CreatePayload<FSlateRoundedBoxPayload>(Element);
+		float Radius = InBrush->OutlineSettings.Radius;
+		if (InBrush->OutlineSettings.RoundingType == ESlateBrushRoundingType::HalfHeightRadius)
+		{
+			Radius = PaintGeometry.GetLocalSize().Y/2.0;
+		}
+		RBoxPayload->SetRadius(Radius);
+		RBoxPayload->SetOutline(InBrush->OutlineSettings.Color.GetSpecifiedColor(), InBrush->OutlineSettings.Width * PaintGeometry.DrawScale);
+		BoxPayload = RBoxPayload;
+	}
+	else 
+	{
+		BoxPayload = &ElementList.CreatePayload<FSlateBoxPayload>(Element);
+	}
 
 	Element.Init(ElementList, ElementType, InLayer, PaintGeometry, InDrawEffects);
 
-	BoxPayload.SetTint(InTint);
-	BoxPayload.SetBrush(InBrush);
+	BoxPayload->SetTint(InTint);
+	BoxPayload->SetBrush(InBrush, PaintGeometry.GetLocalSize(), PaintGeometry.DrawScale);
 
 	return Element;
 }
diff --git a/Engine/Source/Runtime/SlateCore/Private/Styling/SlateColor.cpp b/Engine/Source/Runtime/SlateCore/Private/Styling/SlateColor.cpp
--- a/Engine/Source/Runtime/SlateCore/Private/Styling/SlateColor.cpp
+++ b/Engine/Source/Runtime/SlateCore/Private/Styling/SlateColor.cpp
@@ -2,6 +2,8 @@
 
 #include "Styling/SlateColor.h"
 #include "UObject/PropertyTag.h"
+#include "Styling/StyleColors.h"
+#include "SlateGlobals.h"
 
 bool FSlateColor::SerializeFromMismatchedTag(const FPropertyTag& Tag, FStructuredArchive::FSlot Slot)
 {
@@ -27,3 +29,8 @@ bool FSlateColor::SerializeFromMismatchedTag(const FPropertyTag& Tag, FStructure
 
 	return false;
 }
+
+const FLinearColor& FSlateColor::GetColorFromTable() const
+{
+	return UStyleColorTable::Get().GetColor(ColorTableId);
+}
diff --git a/Engine/Source/Runtime/SlateCore/Private/Styling/StarshipCoreStyle.cpp b/Engine/Source/Runtime/SlateCore/Private/Styling/StarshipCoreStyle.cpp
--- a/Engine/Source/Runtime/SlateCore/Private/Styling/StarshipCoreStyle.cpp
+++ b/Engine/Source/Runtime/SlateCore/Private/Styling/StarshipCoreStyle.cpp
@@ -13,6 +13,7 @@
 #include "Fonts/LegacySlateFontInfoCache.h"
 #include "Styling/SlateStyleMacros.h"
 #include "Styling/ToolBarStyle.h"
+#include "Styling/StyleColors.h"
 
 // This is to fix the issue that SlateStyleMacros like IMAGE_BRUSH look for RootToContentDir but Style->RootToContentDir is how the core style is set up
 #define RootToContentDir Style->RootToContentDir
@@ -23,50 +24,6 @@
 TSharedPtr< ISlateStyle > FStarshipCoreStyle::Instance = nullptr;
 
 
-// HEX Colors from sRGB Color space
-#define COLOR( HexValue ) FLinearColor::FromSRGBColor(FColor::FromHex(HexValue))
-
-const FSlateColor FStyleColors::Black				= COLOR("#000000FF");
-const FSlateColor FStyleColors::Title				= COLOR("#151515FF");
-const FSlateColor FStyleColors::Foldout				= COLOR("0F0F0FFF");
-const FSlateColor FStyleColors::Input				= COLOR("#1A1A1AFF");
-const FSlateColor FStyleColors::Background			= COLOR("#242424FF");
-const FSlateColor FStyleColors::Header				= COLOR("#2F2F2FFF");
-const FSlateColor FStyleColors::Dropdown			= COLOR("#383838FF");
-const FSlateColor FStyleColors::Hover				= COLOR("#575757FF");
-const FSlateColor FStyleColors::Hover2				= COLOR("#808080FF");
-
-const FSlateColor FStyleColors::White				= COLOR("#FFFFFFFF");
-const FSlateColor FStyleColors::White25				= COLOR("#FFFFFF40");
-const FSlateColor FStyleColors::Highlight			= COLOR("#0078D7FF");
-
-const FSlateColor FStyleColors::Primary				= COLOR("#26BBFFFF");
-const FSlateColor FStyleColors::PrimaryHover		= COLOR("#6FD2FFFF");
-const FSlateColor FStyleColors::PrimaryPress		= COLOR("#1989BCFF");
-
-const FSlateColor FStyleColors::Foreground			= COLOR("#A6A6A6FF");
-const FSlateColor FStyleColors::ForegroundHover		= COLOR("#FFFFFFFF");
-const FSlateColor FStyleColors::ForegroundInverted	= FStyleColors::Input;
-const FSlateColor FStyleColors::ForegroundHeader	= COLOR("#C8C8C8FF");
-
-const FSlateColor FStyleColors::Select				= FStyleColors::Primary;
-const FSlateColor FStyleColors::SelectInactive		= COLOR("#99B3BFFF");
-const FSlateColor FStyleColors::SelectParent		= COLOR("#2C323AFF");
-const FSlateColor FStyleColors::SelectHover			= FStyleColors::Background;
-
-const FSlateColor FStyleColors::AccentBlue			= COLOR("#26BBFFFF");
-const FSlateColor FStyleColors::AccentPurple		= COLOR("#A139BFFF");
-const FSlateColor FStyleColors::AccentPink			= COLOR("#FF729CFF");
-const FSlateColor FStyleColors::AccentRed			= COLOR("#FF4040FF");
-const FSlateColor FStyleColors::AccentOrange		= COLOR("#FE9B07FF");
-const FSlateColor FStyleColors::AccentYellow		= COLOR("#FFDC1AFF");
-const FSlateColor FStyleColors::AccentGreen			= COLOR("#8BC24AFF");
-const FSlateColor FStyleColors::AccentBrown			= COLOR("#804D39FF");
-const FSlateColor FStyleColors::AccentBlack			= COLOR("#242424FF");
-const FSlateColor FStyleColors::AccentGray			= COLOR("#808080FF");
-const FSlateColor FStyleColors::AccentWhite			= COLOR("#FFFFFFFF");
-const FSlateColor FStyleColors::AccentFolder		= COLOR("#B68F55FF");
-
 #define FONT(...) FSlateFontInfo(FLegacySlateFontInfoCache::Get().GetDefaultFont(), __VA_ARGS__)
 
 TUniquePtr<FStyleFonts> FStyleFonts::Instance = nullptr;
diff --git a/Engine/Source/Runtime/SlateCore/Private/Styling/StyleColors.cpp b/Engine/Source/Runtime/SlateCore/Private/Styling/StyleColors.cpp
new file mode 100644
--- /dev/null
+++ b/Engine/Source/Runtime/SlateCore/Private/Styling/StyleColors.cpp
@@ -0,0 +1,98 @@
+// Copyright Epic Games, Inc. All Rights Reserved.
+
+#include "Styling/StyleColors.h"
+
+
+const FSlateColor FStyleColors::Black = EStyleColor::Black;
+const FSlateColor FStyleColors::Title = EStyleColor::Title;
+const FSlateColor FStyleColors::Foldout = EStyleColor::Foldout;
+const FSlateColor FStyleColors::Input = EStyleColor::Input;
+const FSlateColor FStyleColors::Background = EStyleColor::Background;
+const FSlateColor FStyleColors::Header = EStyleColor::Header;
+const FSlateColor FStyleColors::Dropdown = EStyleColor::Dropdown;
+const FSlateColor FStyleColors::Hover = EStyleColor::Hover;
+const FSlateColor FStyleColors::Hover2 = EStyleColor::Hover2;
+
+const FSlateColor FStyleColors::White = EStyleColor::White;
+const FSlateColor FStyleColors::White25 = EStyleColor::White25;
+const FSlateColor FStyleColors::Highlight = EStyleColor::Highlight;
+
+const FSlateColor FStyleColors::Primary = EStyleColor::Primary;
+const FSlateColor FStyleColors::PrimaryHover = EStyleColor::PrimaryHover;
+const FSlateColor FStyleColors::PrimaryPress = EStyleColor::PrimaryPress;
+
+const FSlateColor FStyleColors::Foreground = EStyleColor::Foreground;
+const FSlateColor FStyleColors::ForegroundHover = EStyleColor::ForegroundHover;
+const FSlateColor FStyleColors::ForegroundInverted = EStyleColor::ForegroundInverted;
+const FSlateColor FStyleColors::ForegroundHeader = EStyleColor::ForegroundHeader;
+
+const FSlateColor FStyleColors::Select = EStyleColor::Select;
+const FSlateColor FStyleColors::SelectInactive = EStyleColor::SelectInactive;
+const FSlateColor FStyleColors::SelectParent = EStyleColor::SelectParent;
+const FSlateColor FStyleColors::SelectHover = EStyleColor::SelectHover;
+// if select ==  primary shouldnt we have a select pressed which is the same as primary press?
+
+const FSlateColor FStyleColors::AccentBlue = EStyleColor::AccentBlue;
+const FSlateColor FStyleColors::AccentPurple = EStyleColor::AccentPurple;
+const FSlateColor FStyleColors::AccentPink = EStyleColor::AccentPink;
+const FSlateColor FStyleColors::AccentRed = EStyleColor::AccentRed;
+const FSlateColor FStyleColors::AccentOrange = EStyleColor::AccentOrange;
+const FSlateColor FStyleColors::AccentYellow = EStyleColor::AccentYellow;
+const FSlateColor FStyleColors::AccentGreen = EStyleColor::AccentGreen;
+const FSlateColor FStyleColors::AccentBrown = EStyleColor::AccentBrown;
+const FSlateColor FStyleColors::AccentBlack = EStyleColor::AccentBlack;
+const FSlateColor FStyleColors::AccentGray = EStyleColor::AccentGray;
+const FSlateColor FStyleColors::AccentWhite = EStyleColor::AccentWhite;
+const FSlateColor FStyleColors::AccentFolder = EStyleColor::AccentFolder;
+
+
+UStyleColorTable::UStyleColorTable()
+{
+	InitalizeDefaults();
+}
+
+void UStyleColorTable::InitalizeDefaults()
+{
+	SetColor(EStyleColor::Black, COLOR("#000000FF"));
+	SetColor(EStyleColor::Title, COLOR("#151515FF"));
+	SetColor(EStyleColor::Foldout, COLOR("0F0F0FFF"));
+	SetColor(EStyleColor::Input, COLOR("#1A1A1AFF"));
+	SetColor(EStyleColor::Background, COLOR("#242424FF"));
+	SetColor(EStyleColor::Header, COLOR("#2F2F2FFF"));
+	SetColor(EStyleColor::Dropdown, COLOR("#383838FF"));
+	SetColor(EStyleColor::Hover, COLOR("#575757FF"));
+	SetColor(EStyleColor::Hover2, COLOR("#808080FF"));
+
+	SetColor(EStyleColor::White, COLOR("#FFFFFFFF"));
+	SetColor(EStyleColor::White25, COLOR("#FFFFFF40"));
+	SetColor(EStyleColor::Highlight, COLOR("#0078D7FF"));
+
+	SetColor(EStyleColor::Primary, COLOR("#26BBFFFF"));
+	SetColor(EStyleColor::PrimaryHover, COLOR("#6FD2FFFF"));
+	SetColor(EStyleColor::PrimaryPress, COLOR("#1989BCFF"));
+
+	SetColor(EStyleColor::Foreground, COLOR("#A6A6A6FF"));
+	SetColor(EStyleColor::ForegroundHover, COLOR("#FFFFFFFF"));
+	SetColor(EStyleColor::ForegroundInverted, GetColor(EStyleColor::Input));
+	SetColor(EStyleColor::ForegroundHeader, COLOR("#C8C8C8FF"));
+
+	SetColor(EStyleColor::Select, GetColor(EStyleColor::Primary));
+	SetColor(EStyleColor::SelectInactive, COLOR("#99B3BFFF"));
+	SetColor(EStyleColor::SelectParent, COLOR("#2C323AFF"));
+	SetColor(EStyleColor::SelectHover, GetColor(EStyleColor::Background));
+	// if select ==  primary shouldnt we have a select pressed which is the same as primary press?
+
+	SetColor(EStyleColor::AccentBlue, COLOR("#26BBFFFF"));
+	SetColor(EStyleColor::AccentPurple, COLOR("#A139BFFF"));
+	SetColor(EStyleColor::AccentPink, COLOR("#FF729CFF"));
+	SetColor(EStyleColor::AccentRed, COLOR("#FF4040FF"));
+	SetColor(EStyleColor::AccentOrange, COLOR("#FE9B07FF"));
+	SetColor(EStyleColor::AccentYellow, COLOR("#FFDC1AFF"));
+	SetColor(EStyleColor::AccentGreen, COLOR("#8BC24AFF"));
+	SetColor(EStyleColor::AccentBrown, COLOR("#804D39FF"));
+	SetColor(EStyleColor::AccentBlack, COLOR("#242424FF"));
+	SetColor(EStyleColor::AccentGray, COLOR("#808080FF"));
+	SetColor(EStyleColor::AccentWhite, COLOR("#FFFFFFFF"));
+	SetColor(EStyleColor::AccentFolder, COLOR("#B68F55FF"));
+}
+
diff --git a/Engine/Source/Runtime/SlateCore/Public/Brushes/SlateRoundedBoxBrush.h b/Engine/Source/Runtime/SlateCore/Public/Brushes/SlateRoundedBoxBrush.h
--- a/Engine/Source/Runtime/SlateCore/Public/Brushes/SlateRoundedBoxBrush.h
+++ b/Engine/Source/Runtime/SlateCore/Public/Brushes/SlateRoundedBoxBrush.h
@@ -198,7 +198,7 @@ struct SLATECORE_API FSlateRoundedBoxBrush
 					  ESlateBrushTileType::NoTile, 
 					  ESlateBrushImageType::NoImage, 
 					  InImageSize,
-					  InColor.GetSpecifiedColor(),
+					  InColor,
 					  nullptr,
 					  false
 					 )
@@ -219,7 +219,7 @@ struct SLATECORE_API FSlateRoundedBoxBrush
 					  ESlateBrushTileType::NoTile, 
 					  ESlateBrushImageType::NoImage, 
 					  InImageSize,
-					  InColor.GetSpecifiedColor(),
+					  InColor,
 					  nullptr,
 					  false
 					 )
@@ -241,12 +241,12 @@ struct SLATECORE_API FSlateRoundedBoxBrush
 					  ESlateBrushTileType::NoTile, 
 					  ESlateBrushImageType::NoImage, 
 					  InImageSize,
-					  InColor.GetSpecifiedColor(),
+					  InColor,
 					  nullptr,
 					  false
 					  )
 	{ 
-		OutlineSettings = FSlateBrushOutlineSettings(InOutlineColor.GetSpecifiedColor(), InOutlineWidth);
+		OutlineSettings = FSlateBrushOutlineSettings(InOutlineColor, InOutlineWidth);
 	}
 
 	/** 
@@ -264,12 +264,12 @@ struct SLATECORE_API FSlateRoundedBoxBrush
 					  ESlateBrushTileType::NoTile, 
 					  ESlateBrushImageType::NoImage, 
 					  InImageSize,
-					  InColor.GetSpecifiedColor(),
+					  InColor,
 					  nullptr,
 					  false
 					  )
 	{ 
-		OutlineSettings = FSlateBrushOutlineSettings(InRadius, InOutlineColor.GetSpecifiedColor(), InOutlineWidth);
+		OutlineSettings = FSlateBrushOutlineSettings(InRadius, InOutlineColor, InOutlineWidth);
 	}
 
 };
diff --git a/Engine/Source/Runtime/SlateCore/Public/Styling/SlateBrush.h b/Engine/Source/Runtime/SlateCore/Public/Styling/SlateBrush.h
--- a/Engine/Source/Runtime/SlateCore/Public/Styling/SlateBrush.h
+++ b/Engine/Source/Runtime/SlateCore/Public/Styling/SlateBrush.h
@@ -27,7 +27,10 @@ namespace ESlateBrushDrawType
 		Border,
 
 		/** Draw an image; margin is ignored */
-		Image
+		Image,
+
+		/** Draw a solid rectangle with an outline and corner radius */
+		RoundedBox
 	};
 }
 
@@ -94,9 +97,84 @@ namespace ESlateBrushImageType
 
 		/** The image is a special texture in linear space (usually a rendering resource such as a lookup table). */
 		Linear,
+
+		/** The image is vector graphics and will be rendered and cached in full color using size/scale requested by slate */
+		Vector,
 	};
 }
 
+
+/**
+ * Enumerates rounding options
+ */
+UENUM()
+namespace ESlateBrushRoundingType
+{
+	enum Type
+	{
+		/** Use the specified Radius **/
+		FixedRadius, 
+
+		/** The rounding radius should be half the height such that it always looks perfectly round **/
+		HalfHeightRadius,
+	};
+}
+
+
+/**
+ * Possible options for rounded box brush image
+ */
+USTRUCT(BlueprintType)
+struct SLATECORE_API FSlateBrushOutlineSettings
+{
+	GENERATED_USTRUCT_BODY()
+
+	FSlateBrushOutlineSettings()
+		: Radius(0.0)
+		, Color(FLinearColor::Transparent)
+		, Width(0.0)
+		, RoundingType(ESlateBrushRoundingType::HalfHeightRadius)
+	{}
+
+	FSlateBrushOutlineSettings(float InRadius)
+		: Radius(InRadius)
+		, Color(FLinearColor::Transparent)
+		, Width(0.0)
+		, RoundingType(ESlateBrushRoundingType::FixedRadius)
+	{}
+
+	FSlateBrushOutlineSettings(const FSlateColor& InColor, float InWidth)
+		: Radius(0.0)
+		, Color(InColor)
+		, Width(InWidth)
+		, RoundingType(ESlateBrushRoundingType::HalfHeightRadius)
+	{}
+
+	FSlateBrushOutlineSettings(float InRadius, const FSlateColor& InColor, float InWidth)
+		: Radius(InRadius)
+		, Color(InColor)
+		, Width(InWidth)
+		, RoundingType(ESlateBrushRoundingType::FixedRadius)
+	{}
+
+	/** Radius in Slate Units applied to the outline. */
+	UPROPERTY(EditAnywhere, BlueprintReadWrite, Category=Brush)
+	float Radius;
+
+	/** Tinting applied to the border outline. */
+	UPROPERTY(EditAnywhere, BlueprintReadWrite, Category=Brush, meta=(DisplayName="Outline", sRGB="true"))
+	FSlateColor Color;
+
+	/** Line width in Slate Units applied to the border outline. */
+	UPROPERTY(EditAnywhere, BlueprintReadWrite, Category=Brush)
+	float Width;
+
+	/** The Rounding Type **/
+	UPROPERTY(EditAnywhere, BlueprintReadWrite, Category=Brush)
+	TEnumAsByte<enum ESlateBrushRoundingType::Type > RoundingType;
+
+};
+
 namespace SlateBrushDefs
 {
 	static const float DefaultImageSize = 32.0f;
@@ -110,6 +188,7 @@ struct SLATECORE_API FSlateBrush
 {
 	GENERATED_USTRUCT_BODY()
 
+	friend class FSlateShaderResourceManager;
 public:
 	/** Size of the resource in Slate Units */
 	UPROPERTY(EditAnywhere, BlueprintReadWrite, Category=Brush)
@@ -129,6 +208,10 @@ public:
 	UPROPERTY(EditAnywhere, BlueprintReadWrite, Category=Brush, meta=( DisplayName="Tint", sRGB="true" ))
 	FSlateColor TintColor;
 
+	/** How to draw the outline.  Currently only used for RoundedBox type brushes. */
+	UPROPERTY(EditAnywhere, BlueprintReadWrite, Category=Brush)
+	FSlateBrushOutlineSettings OutlineSettings;
+
 public:
 
 	/**
@@ -304,12 +387,21 @@ public:
 	 */
 	static const FString UTextureIdentifier( );
 	
+	const FSlateResourceHandle& GetRenderingResource(FVector2D LocalSize, float DrawScale) const
+	{
+		UpdateRenderingResource(LocalSize, DrawScale);
+
+		return ResourceHandle;
+	}
+
 	const FSlateResourceHandle& GetRenderingResource() const
 	{
-		if (!ResourceHandle.IsValid())
+		if (ImageType == ESlateBrushImageType::Vector)
 		{
-			UpdateRenderingResource();
+			UE_LOG(LogSlate, Warning, TEXT("FSlateBrush::GetRenderingResource should be called with a size and scale for vector brushes"));
 		}
+	
+		UpdateRenderingResource(GetImageSize(), 1.0f);
 
 		return ResourceHandle;
 	}
@@ -322,7 +414,7 @@ public:
 #endif
 
 private:
-	void UpdateRenderingResource() const;
+	void UpdateRenderingResource(FVector2D LocalSize, float DrawScale) const;
 	bool CanRenderResourceObject(UObject* InResourceObject) const;
 
 private:
@@ -385,12 +477,14 @@ protected:
 	 * @param InTiling        Tile horizontally/vertically or both? (only in image mode)
 	 * @param InImageType	  The type of image
 	 * @param InTint		  Tint to apply to the element.
+	 * @param InOutlineSettings Optional Outline Border Settings for RoundedBox mode
 	 */
-	 FORCENOINLINE FSlateBrush( ESlateBrushDrawType::Type InDrawType, const FName InResourceName, const FMargin& InMargin, ESlateBrushTileType::Type InTiling, ESlateBrushImageType::Type InImageType, const FVector2D& InImageSize, const FLinearColor& InTint = FLinearColor::White, UObject* InObjectResource = nullptr, bool bInDynamicallyLoaded = false );
+	 FORCENOINLINE FSlateBrush( ESlateBrushDrawType::Type InDrawType, const FName InResourceName, const FMargin& InMargin, ESlateBrushTileType::Type InTiling, ESlateBrushImageType::Type InImageType, const FVector2D& InImageSize, const FLinearColor& InTint = FLinearColor::White, UObject* InObjectResource = nullptr, bool bInDynamicallyLoaded = false);
+
+	 FORCENOINLINE FSlateBrush( ESlateBrushDrawType::Type InDrawType, const FName InResourceName, const FMargin& InMargin, ESlateBrushTileType::Type InTiling, ESlateBrushImageType::Type InImageType, const FVector2D& InImageSize, const TSharedRef< FLinearColor >& InTint, UObject* InObjectResource = nullptr, bool bInDynamicallyLoaded = false);
 
-	 FORCENOINLINE FSlateBrush( ESlateBrushDrawType::Type InDrawType, const FName InResourceName, const FMargin& InMargin, ESlateBrushTileType::Type InTiling, ESlateBrushImageType::Type InImageType, const FVector2D& InImageSize, const TSharedRef< FLinearColor >& InTint, UObject* InObjectResource = nullptr, bool bInDynamicallyLoaded = false );
+	 FORCENOINLINE FSlateBrush( ESlateBrushDrawType::Type InDrawType, const FName InResourceName, const FMargin& InMargin, ESlateBrushTileType::Type InTiling, ESlateBrushImageType::Type InImageType, const FVector2D& InImageSize, const FSlateColor& InTint, UObject* InObjectResource = nullptr, bool bInDynamicallyLoaded = false);
 
-	 FORCENOINLINE FSlateBrush( ESlateBrushDrawType::Type InDrawType, const FName InResourceName, const FMargin& InMargin, ESlateBrushTileType::Type InTiling, ESlateBrushImageType::Type InImageType, const FVector2D& InImageSize, const FSlateColor& InTint, UObject* InObjectResource = nullptr, bool bInDynamicallyLoaded = false );
 };
 
 /** Provides a means to hold onto the source of a slate brush. */
diff --git a/Engine/Source/Runtime/SlateCore/Public/Styling/SlateColor.h b/Engine/Source/Runtime/SlateCore/Public/Styling/SlateColor.h
--- a/Engine/Source/Runtime/SlateCore/Public/Styling/SlateColor.h
+++ b/Engine/Source/Runtime/SlateCore/Public/Styling/SlateColor.h
@@ -14,24 +14,25 @@
  * Should we use the specified color? If not, then which color from the style should we use.
  */
 UENUM(BlueprintType)
-namespace ESlateColorStylingMode
+enum class ESlateColorStylingMode : uint8
 {
-	enum Type
-	{
-		/** Color value is stored in this Slate color. */
-		UseColor_Specified UMETA(DisplayName="Specified Color"),
+	/** Color value is stored in this Slate color. */
+	UseColor_Specified UMETA(DisplayName = "Specified Color"),
 
-		/** Color value is stored in the linked color. */
-		UseColor_Specified_Link UMETA(Hidden),
+	/** Color value is stored in the linked color table. */
+	UseColor_ColorTable UMETA(Hidden),
 
-		/** Use the widget's foreground color. */
-		UseColor_Foreground UMETA(DisplayName="Foreground Color"),
+	/** Use the widget's foreground color. */
+	UseColor_Foreground UMETA(DisplayName = "Foreground Color"),
 
-		/** Use the widget's subdued color. */
-		UseColor_Foreground_Subdued UMETA(Hidden)
-	};
-}
+	/** Use the widget's subdued color. */
+	UseColor_Foreground_Subdued UMETA(Hidden),
+
+	/** Use the foreground color defined in a widget specific style. */
+	UseColor_UseStyle UMETA(Hidden),
+};
 
+enum class EStyleColor : uint8;
 
 /**
  * A Slate color can be a directly specified value, or the color can be pulled from a WidgetStyle.
@@ -69,9 +70,14 @@ public:
 	 * @param InColor The color value to assign.
 	 */
 	FSlateColor( const TSharedRef< FLinearColor >& InColor )
-		: SpecifiedColor(FLinearColor(1.0f, 0.0f, 1.0f))
-		, ColorUseRule(ESlateColorStylingMode::UseColor_Specified_Link)
-		, LinkedSpecifiedColor(InColor)
+		: SpecifiedColor(*InColor)
+		, ColorUseRule(ESlateColorStylingMode::UseColor_Specified)
+	{ }
+
+	FSlateColor(EStyleColor InColorTableId)
+		: SpecifiedColor()
+		, ColorUseRule(ESlateColorStylingMode::UseColor_ColorTable)
+		, ColorTableId(InColorTableId)
 	{ }
 
 	~FSlateColor()
@@ -99,8 +105,8 @@ public:
 			return SpecifiedColor;
 			break;
 
-		case ESlateColorStylingMode::UseColor_Specified_Link:
-			return *LinkedSpecifiedColor;
+		case ESlateColorStylingMode::UseColor_ColorTable:
+			return GetColorFromTable();
 			break;
 
 		case ESlateColorStylingMode::UseColor_Foreground_Subdued:
@@ -117,9 +123,9 @@ public:
 	 */
 	FLinearColor GetSpecifiedColor( ) const
 	{
-		if (LinkedSpecifiedColor.IsValid())
+		if (ColorUseRule == ESlateColorStylingMode::UseColor_ColorTable)
 		{
-			return *LinkedSpecifiedColor;
+			return GetColorFromTable();
 		}
 
 		return SpecifiedColor;
@@ -133,7 +139,7 @@ public:
 	 */
 	bool IsColorSpecified( ) const
 	{
-		return (ColorUseRule == ESlateColorStylingMode::UseColor_Specified) || (ColorUseRule == ESlateColorStylingMode::UseColor_Specified_Link);
+		return (ColorUseRule == ESlateColorStylingMode::UseColor_Specified) || (ColorUseRule == ESlateColorStylingMode::UseColor_ColorTable);
 	}
 
 	/**
@@ -143,14 +149,9 @@ public:
 	 */
 	void Unlink()
 	{
-		if (ColorUseRule == ESlateColorStylingMode::UseColor_Specified_Link)
+		if (ColorUseRule == ESlateColorStylingMode::UseColor_ColorTable)
 		{
 			ColorUseRule = ESlateColorStylingMode::UseColor_Specified;
-			if (ensure(LinkedSpecifiedColor.IsValid()))
-			{
-				SpecifiedColor = *LinkedSpecifiedColor;
-				LinkedSpecifiedColor.Reset();
-			}
 		}
 	}
 
@@ -165,7 +166,7 @@ public:
 	{
 		return SpecifiedColor == Other.SpecifiedColor
 			&& ColorUseRule == Other.ColorUseRule
-			&& (ColorUseRule != ESlateColorStylingMode::UseColor_Specified_Link || LinkedSpecifiedColor == Other.LinkedSpecifiedColor);
+			&& (ColorUseRule != ESlateColorStylingMode::UseColor_ColorTable || ColorTableId == Other.ColorTableId);
 	}
 
 	/**
@@ -194,18 +195,26 @@ public:
 		return FSlateColor( ESlateColorStylingMode::UseColor_Foreground_Subdued );
 	}
 
+	/** @returns an FSlateColor that is the subdued version of the widget's foreground. */
+	static FSlateColor UseStyle()
+	{
+		return FSlateColor(ESlateColorStylingMode::UseColor_UseStyle);
+	}
+
+
 	/** Used to upgrade an FColor or FLinearColor property to an FSlateColor property */
 	bool SerializeFromMismatchedTag(const struct FPropertyTag& Tag, FStructuredArchive::FSlot Slot);
 
 protected:
 	
 	// Private constructor to prevent construction of invalid FSlateColors
-	FSlateColor( ESlateColorStylingMode::Type InColorUseRule )
+	FSlateColor(ESlateColorStylingMode InColorUseRule)
 		: SpecifiedColor(FLinearColor(1.0f, 0.0f, 1.0f))
 		, ColorUseRule(InColorUseRule)
-		, LinkedSpecifiedColor()
+		, ColorTableId()
 	{ }
 
+	const FLinearColor& GetColorFromTable() const;
 protected:
 
 	// The current specified color; only meaningful when ColorToUse == UseColor_Specified.
@@ -214,12 +223,11 @@ protected:
 
 	// The rule for which color to pick.
 	UPROPERTY(EditAnywhere, BlueprintReadWrite, Category=Color)
-	TEnumAsByte<ESlateColorStylingMode::Type> ColorUseRule;
+	ESlateColorStylingMode ColorUseRule;
 
 private:
-
-	// A shared pointer to the linked color value, if any.
-	TSharedPtr<FLinearColor> LinkedSpecifiedColor;
+	// Id to a color in a color table to be used with ESlateColorStylingMode::UseColor_ColorTable
+	EStyleColor ColorTableId;
 };
 
 template<>
diff --git a/Engine/Source/Runtime/SlateCore/Public/Styling/StarshipCoreStyle.h b/Engine/Source/Runtime/SlateCore/Public/Styling/StarshipCoreStyle.h
--- a/Engine/Source/Runtime/SlateCore/Public/Styling/StarshipCoreStyle.h
+++ b/Engine/Source/Runtime/SlateCore/Public/Styling/StarshipCoreStyle.h
@@ -75,53 +75,6 @@ private:
 
 };
 
-
-struct SLATECORE_API FStyleColors
-{
-	static const FSlateColor Black;
-	static const FSlateColor Title;
-	static const FSlateColor Foldout;
-	static const FSlateColor Input;
-	static const FSlateColor Background;
-	static const FSlateColor Header;
-	static const FSlateColor Dropdown;
-	static const FSlateColor Hover;
-	static const FSlateColor Hover2;
-	static const FSlateColor White;
-	static const FSlateColor White25;
-	static const FSlateColor Highlight;
-
-	static const FSlateColor Primary;
-	static const FSlateColor PrimaryHover;
-	static const FSlateColor PrimaryPress;
-
-
-	static const FSlateColor Foreground;
-	static const FSlateColor ForegroundHover;
-	static const FSlateColor ForegroundInverted;
-	static const FSlateColor ForegroundHeader;
-
-	static const FSlateColor Select;
-	static const FSlateColor SelectInactive;
-	static const FSlateColor SelectParent;
-	static const FSlateColor SelectHover;
-
-	static const FSlateColor AccentBlue;
-	static const FSlateColor AccentPurple;
-	static const FSlateColor AccentPink;
-	static const FSlateColor AccentRed;
-	static const FSlateColor AccentOrange;
-	static const FSlateColor AccentYellow;
-	static const FSlateColor AccentGreen;
-	static const FSlateColor AccentBrown;
-	static const FSlateColor AccentBlack;
-	static const FSlateColor AccentGray;
-	static const FSlateColor AccentWhite;
-	static const FSlateColor AccentFolder;
-
-
-};
-
 struct SLATECORE_API FStyleFonts
 {
 
diff --git a/Engine/Source/Runtime/SlateCore/Public/Styling/StyleColors.h b/Engine/Source/Runtime/SlateCore/Public/Styling/StyleColors.h
new file mode 100644
--- /dev/null
+++ b/Engine/Source/Runtime/SlateCore/Public/Styling/StyleColors.h
@@ -0,0 +1,139 @@
+// Copyright Epic Games, Inc. All Rights Reserved.
+
+#pragma once
+
+#include "CoreMinimal.h"
+#include "SlateColor.h"
+#include "StyleColors.generated.h"
+
+
+// HEX Colors from sRGB Color space
+#define COLOR( HexValue ) FLinearColor::FromSRGBColor(FColor::FromHex(HexValue))
+
+UENUM()
+enum class EStyleColor : uint8
+{
+	Black,
+	Title,
+	Foldout,
+	Input,
+	Background,
+	Header,
+	Dropdown,
+	Hover,
+	Hover2,
+	White,
+	White25,
+	Highlight,
+	Primary,
+	PrimaryHover,
+	PrimaryPress,
+	Foreground,
+	ForegroundHover,
+	ForegroundInverted,
+	ForegroundHeader,
+	Select,
+	SelectInactive,
+	SelectParent,
+	SelectHover,
+	AccentBlue,
+	AccentPurple,
+	AccentPink,
+	AccentRed,
+	AccentOrange,
+	AccentYellow,
+	AccentGreen,
+	AccentBrown,
+	AccentBlack,
+	AccentGray,
+	AccentWhite,
+	AccentFolder,
+
+	MAX
+};
+
+UCLASS(Config=EditorSettings)
+class SLATECORE_API UStyleColorTable : public UObject 
+{
+	GENERATED_BODY()
+public:
+	static UStyleColorTable& Get()
+	{
+		return *GetMutableDefault<UStyleColorTable>();
+	}
+
+	const FLinearColor& GetColor(EStyleColor Color)
+	{
+		return Colors[static_cast<int32>(Color)];
+	}
+
+	void SetColor(EStyleColor InColorId, FLinearColor InColor)
+	{
+		Colors[static_cast<int32>(InColorId)] = InColor;
+	}
+
+	UStyleColorTable();
+
+#if WITH_EDITOR
+	virtual void PostEditChangeProperty(struct FPropertyChangedEvent& PropertyChangedEvent) override
+	{
+		Super::PostEditChangeProperty(PropertyChangedEvent);
+
+		if (PropertyChangedEvent.ChangeType != EPropertyChangeType::Interactive)
+		{
+			SaveConfig();
+		}
+	}
+#endif
+
+	void InitalizeDefaults();
+
+private:
+	UPROPERTY(EditAnywhere, Config, Category=Colors)
+	FLinearColor Colors[EStyleColor::MAX];
+};
+
+
+struct SLATECORE_API FStyleColors
+{
+	static const FSlateColor Black;
+	static const FSlateColor Title;
+	static const FSlateColor Foldout;
+	static const FSlateColor Input;
+	static const FSlateColor Background;
+	static const FSlateColor Header;
+	static const FSlateColor Dropdown;
+	static const FSlateColor Hover;
+	static const FSlateColor Hover2;
+	static const FSlateColor White;
+	static const FSlateColor White25;
+	static const FSlateColor Highlight;
+
+	static const FSlateColor Primary;
+	static const FSlateColor PrimaryHover;
+	static const FSlateColor PrimaryPress;
+
+
+	static const FSlateColor Foreground;
+	static const FSlateColor ForegroundHover;
+	static const FSlateColor ForegroundInverted;
+	static const FSlateColor ForegroundHeader;
+
+	static const FSlateColor Select;
+	static const FSlateColor SelectInactive;
+	static const FSlateColor SelectParent;
+	static const FSlateColor SelectHover;
+
+	static const FSlateColor AccentBlue;
+	static const FSlateColor AccentPurple;
+	static const FSlateColor AccentPink;
+	static const FSlateColor AccentRed;
+	static const FSlateColor AccentOrange;
+	static const FSlateColor AccentYellow;
+	static const FSlateColor AccentGreen;
+	static const FSlateColor AccentBrown;
+	static const FSlateColor AccentBlack;
+	static const FSlateColor AccentGray;
+	static const FSlateColor AccentWhite;
+	static const FSlateColor AccentFolder;
+};
