More CIS fixing

diff --git a/Engine/Source/Editor/EditorStyle/Private/StarshipStyle.cpp b/Engine/Source/Editor/EditorStyle/Private/StarshipStyle.cpp
--- a/Engine/Source/Editor/EditorStyle/Private/StarshipStyle.cpp
+++ b/Engine/Source/Editor/EditorStyle/Private/StarshipStyle.cpp
@@ -14,11 +14,63 @@
 #endif
 #include "Styling/ToolBarStyle.h"
 #include "Styling/StyleColors.h"
+#include "Framework/Application/SlateApplication.h"
 
 
 #define DEFAULT_FONT(...) FStarshipCoreStyle::GetDefaultFontStyle(__VA_ARGS__)
 #define ICON_FONT(...) FSlateFontInfo(RootToContentDir("Fonts/FontAwesome", TEXT(".ttf")), __VA_ARGS__)
 
+void FStarshipEditorStyle::Initialize()
+{
+	Settings = NULL;
+
+	// The core style must be initialized before the editor style
+	FSlateApplication::InitializeCoreStyle();
+
+#if WITH_EDITOR
+	Settings = GetMutableDefault<UEditorStyleSettings>();
+	ISettingsModule* SettingsModule = FModuleManager::GetModulePtr<ISettingsModule>("Settings");
+
+	if (SettingsModule != nullptr)
+	{
+		SettingsModule->RegisterSettings("Editor", "General", "Appearance",
+			NSLOCTEXT("EditorStyle", "Appearance_UserSettingsName", "Appearance"),
+			NSLOCTEXT("EditorStyle", "Appearance_UserSettingsDescription", "Customize the look of the editor."),
+			Settings
+		);
+	}
+
+
+	FPropertyEditorModule& PropertyEditorModule = FModuleManager::Get().GetModuleChecked<FPropertyEditorModule>("PropertyEditor");
+	PropertyEditorModule.RegisterCustomClassLayout("EditorStyleSettings", FOnGetDetailCustomizationInstance::CreateStatic(&FEditorStyleSettingsCustomization::MakeInstance));
+#endif
+
+	StyleInstance = Create(Settings);
+	SetStyle(StyleInstance.ToSharedRef());
+}
+
+void FStarshipEditorStyle::Shutdown()
+{
+#if WITH_EDITOR
+	ISettingsModule* SettingsModule = FModuleManager::GetModulePtr<ISettingsModule>("Settings");
+
+	if (SettingsModule != nullptr)
+	{
+		SettingsModule->UnregisterSettings("Editor", "General", "Appearance");
+	}
+
+
+
+	FPropertyEditorModule* PropertyEditorModule = FModuleManager::Get().GetModulePtr<FPropertyEditorModule>("PropertyEditor");
+	if (PropertyEditorModule)
+	{
+		PropertyEditorModule->UnregisterCustomClassLayout("EditorStyleSettings");
+	}
+#endif
+	//FEditorStyle::SetStyle()
+	StyleInstance.Reset();
+}
+
 BEGIN_SLATE_FUNCTION_BUILD_OPTIMIZATION
 
 /* FStarshipEditorStyle static initialization
diff --git a/Engine/Source/Editor/EditorStyle/Private/StarshipStyle.h b/Engine/Source/Editor/EditorStyle/Private/StarshipStyle.h
--- a/Engine/Source/Editor/EditorStyle/Private/StarshipStyle.h
+++ b/Engine/Source/Editor/EditorStyle/Private/StarshipStyle.h
@@ -28,56 +28,8 @@ class FStarshipEditorStyle
 {
 public:
 
-	static void Initialize()
-	{
-		Settings = NULL;
-
-		// The core style must be initialized before the editor style
-		FSlateApplication::InitializeCoreStyle();
-
-#if WITH_EDITOR
-		Settings = GetMutableDefault<UEditorStyleSettings>();
-		ISettingsModule* SettingsModule = FModuleManager::GetModulePtr<ISettingsModule>("Settings");
-
-		if (SettingsModule != nullptr)
-		{
-			SettingsModule->RegisterSettings("Editor", "General", "Appearance",
-				NSLOCTEXT("EditorStyle", "Appearance_UserSettingsName", "Appearance"),
-				NSLOCTEXT("EditorStyle", "Appearance_UserSettingsDescription", "Customize the look of the editor."),
-				Settings
-			);
-		}
-
-
-		FPropertyEditorModule& PropertyEditorModule = FModuleManager::Get().GetModuleChecked<FPropertyEditorModule>("PropertyEditor");
-		PropertyEditorModule.RegisterCustomClassLayout("EditorStyleSettings", FOnGetDetailCustomizationInstance::CreateStatic(&FEditorStyleSettingsCustomization::MakeInstance));
-#endif
-
-		StyleInstance = Create(Settings);
-		SetStyle(StyleInstance.ToSharedRef());
-	}
-
-	static void Shutdown()
-	{
-#if WITH_EDITOR
-		ISettingsModule* SettingsModule = FModuleManager::GetModulePtr<ISettingsModule>("Settings");
-
-		if (SettingsModule != nullptr)
-		{
-			SettingsModule->UnregisterSettings("Editor", "General", "Appearance");
-		}
-
-
-
-		FPropertyEditorModule* PropertyEditorModule = FModuleManager::Get().GetModulePtr<FPropertyEditorModule>("PropertyEditor");
-		if (PropertyEditorModule)
-		{
-			PropertyEditorModule->UnregisterCustomClassLayout("EditorStyleSettings");
-		}
-#endif
-		//FEditorStyle::SetStyle()
-		StyleInstance.Reset();
-	}
+	static void Initialize();
+	static void Shutdown();
 
 	static void SyncCustomizations()
 	{
diff --git a/Engine/Source/Editor/PropertyEditor/Private/SDetailsViewBase.h b/Engine/Source/Editor/PropertyEditor/Private/SDetailsViewBase.h
--- a/Engine/Source/Editor/PropertyEditor/Private/SDetailsViewBase.h
+++ b/Engine/Source/Editor/PropertyEditor/Private/SDetailsViewBase.h
@@ -30,6 +30,7 @@ class IDetailKeyframeHandler;
 class IDetailPropertyExtensionHandler;
 class SDetailNameArea;
 class IPropertyGenerationUtilities;
+class SSearchBox;
 struct FDetailsViewObjectRoot;
 
 
diff --git a/Engine/Source/Editor/UnrealEd/Private/PlayLevelViaLauncher.cpp b/Engine/Source/Editor/UnrealEd/Private/PlayLevelViaLauncher.cpp
--- a/Engine/Source/Editor/UnrealEd/Private/PlayLevelViaLauncher.cpp
+++ b/Engine/Source/Editor/UnrealEd/Private/PlayLevelViaLauncher.cpp
@@ -20,6 +20,7 @@
 #include "TargetReceipt.h"
 #include "DesktopPlatformModule.h"
 #include "PlatformInfo.h"
+#include "Framework/Docking/TabManager.h"
 
 #define LOCTEXT_NAMESPACE "PlayLevel"
 
diff --git a/Engine/Source/Editor/UnrealEd/Public/SAssetEditorViewport.h b/Engine/Source/Editor/UnrealEd/Public/SAssetEditorViewport.h
--- a/Engine/Source/Editor/UnrealEd/Public/SAssetEditorViewport.h
+++ b/Engine/Source/Editor/UnrealEd/Public/SAssetEditorViewport.h
@@ -8,6 +8,7 @@
 
 class FEditorViewportClient;
 class FAssetEditorViewportLayout;
+class FMenuBuilder;
 
 class UNREALED_API SAssetEditorViewport : public SEditorViewport
 {
