Updated File Menu Icons

diff --git a/Engine/Source/Editor/EditorStyle/Private/StarshipStyle.cpp b/Engine/Source/Editor/EditorStyle/Private/StarshipStyle.cpp
--- a/Engine/Source/Editor/EditorStyle/Private/StarshipStyle.cpp
+++ b/Engine/Source/Editor/EditorStyle/Private/StarshipStyle.cpp
@@ -1293,7 +1293,7 @@ void FStarshipEditorStyle::FStyle::SetupGeneralStyles()
 
 	// Open any asset dialog
 	{
-		Set( "SystemWideCommands.SummonOpenAssetDialog", new IMAGE_BRUSH( "Icons/icon_asset_open_16px", Icon16x16 ) );
+		Set( "SystemWideCommands.SummonOpenAssetDialog", new IMAGE_BRUSH_SVG( "Starship/Menus/File/OpenAsset", Icon16x16 ) );
 	
 		Set( "GlobalAssetPicker.Background", new BOX_BRUSH( "Old/Menu_Background", FMargin(8.0f/64.0f) ) );
 		Set( "GlobalAssetPicker.OutermostMargin", FMargin(4, 4, 4, 4) );
@@ -1311,17 +1311,22 @@ void FStarshipEditorStyle::FStyle::SetupGeneralStyles()
 	{
 		Set( "MainFrame.AutoSaveImage", new IMAGE_BRUSH( "Icons/icon_Autosave", Icon24x24 ) );
 
-		Set( "MainFrame.SaveAll", new IMAGE_BRUSH( "Icons/icon_file_saveall_16px", Icon16x16 ) );
-		Set( "MainFrame.ChoosePackagesToSave", new IMAGE_BRUSH( "Icons/icon_file_choosepackages_16px", Icon16x16 ) );
-		Set( "MainFrame.NewProject", new IMAGE_BRUSH( "Icons/icon_file_ProjectNew_16x", Icon16x16 ) );
-		Set( "MainFrame.OpenProject", new IMAGE_BRUSH( "Icons/icon_file_ProjectOpen_16x", Icon16x16 ) );
-		Set( "MainFrame.AddCodeToProject", new IMAGE_BRUSH( "Icons/icon_file_ProjectAddCode_16x", Icon16x16 ) );
-		Set( "MainFrame.Exit", new IMAGE_BRUSH( "Icons/icon_file_exit_16px", Icon16x16 ) );
-		Set( "MainFrame.CookContent", new IMAGE_BRUSH( "Icons/icon_package_16x", Icon16x16 ) );
-		Set( "MainFrame.PackageProject", new IMAGE_BRUSH( "Icons/icon_package_16x", Icon16x16 ) );
-		Set( "MainFrame.RecentProjects", new IMAGE_BRUSH( "Icons/icon_file_ProjectsRecent_16px", Icon16x16 ) );
-		Set( "MainFrame.RecentLevels", new IMAGE_BRUSH( "Icons/icon_file_LevelsRecent_16px", Icon16x16 ) );
-		Set( "MainFrame.FavoriteLevels", new IMAGE_BRUSH( "Old/Favorites_Enabled", Icon16x16 ) );
+		Set( "MainFrame.SaveAll",              new IMAGE_BRUSH_SVG( "Starship/Menus/File/SaveAll", Icon16x16 ) );
+		Set( "MainFrame.ChoosePackagesToSave", new IMAGE_BRUSH_SVG( "Starship/Menus/File/icon_file_choosepackages_16px", Icon16x16 ) );
+		Set( "MainFrame.NewProject",           new IMAGE_BRUSH_SVG( "Starship/Menus/File/ProjectNew", Icon16x16 ) );
+		Set( "MainFrame.OpenProject",          new IMAGE_BRUSH_SVG( "Starship/Menus/File/ProjectOpen", Icon16x16 ) );
+		Set( "MainFrame.AddCodeToProject",     new IMAGE_BRUSH_SVG( "Starship/Menus/File/ProjectC++", Icon16x16 ) );
+		Set( "MainFrame.Exit",                 new IMAGE_BRUSH_SVG( "Starship/Menus/File/Exit", Icon16x16 ) );
+		Set( "MainFrame.CookContent",          new IMAGE_BRUSH_SVG( "Starship/Menus/File/CookContent", Icon16x16 ) );
+		Set( "MainFrame.OpenVisualStudio",     new IMAGE_BRUSH_SVG( "Starship/Menus/File/VisualStudio", Icon16x16 ) );
+		Set( "MainFrame.RefreshVisualStudio",  new IMAGE_BRUSH_SVG( "Starship/Menus/File/RefreshVisualStudio", Icon16x16 ) );
+		Set( "MainFrame.PackageProject",       new IMAGE_BRUSH_SVG( "Starship/Menus/File/ProjectPackage", Icon16x16 ) );
+		Set( "MainFrame.RecentProjects",       new IMAGE_BRUSH_SVG( "Starship/Menus/File/ProjectsRecent", Icon16x16 ) );
+		Set( "MainFrame.RecentLevels",         new IMAGE_BRUSH_SVG( "Starship/Menus/File/LevelRecent", Icon16x16 ) );
+		Set( "MainFrame.FavoriteLevels",       new IMAGE_BRUSH_SVG( "Starship/Menus/File/LevelFavorite", Icon16x16 ) );
+
+		Set( "MainFrame.ChooseFilesToSave",       new IMAGE_BRUSH_SVG( "Starship/Menus/File/SaveChoose", Icon16x16 ) );
+		Set( "MainFrame.ConnectToSourceControl",  new IMAGE_BRUSH_SVG( "Starship/Menus/File/SourceControl", Icon16x16 ) );
 
 		Set( "MainFrame.DebugTools.SmallFont", DEFAULT_FONT( "Regular", 8 ) );
 		Set( "MainFrame.DebugTools.NormalFont", DEFAULT_FONT( "Regular", 9 ) );
@@ -4276,11 +4281,15 @@ void FStarshipEditorStyle::FStyle::SetupLevelEditorStyle()
 #endif // WITH_EDITOR || (IS_PROGRAM && WITH_UNREAL_DEVELOPER_TOOLS)
 
 #if WITH_EDITOR
-		Set( "LevelEditor.NewLevel", new IMAGE_BRUSH( "Icons/icon_file_new_16px", Icon16x16 ) );
-		Set( "LevelEditor.OpenLevel", new IMAGE_BRUSH( "Icons/icon_file_open_16px", Icon16x16 ) );
-		Set( "LevelEditor.Save", new IMAGE_BRUSH( "Icons/icon_file_save_16px", Icon16x16 ) );
-		Set( "LevelEditor.SaveAs", new IMAGE_BRUSH( "Icons/icon_file_saveas_16px", Icon16x16 ) );
-		Set( "LevelEditor.SaveAllLevels", new IMAGE_BRUSH( "Icons/icon_file_savelevels_16px", Icon16x16 ) );
+		Set( "LevelEditor.NewLevel",      new IMAGE_BRUSH_SVG( "Starship/Menus/File/LevelNew", Icon16x16 ) );
+		Set( "LevelEditor.OpenLevel",     new IMAGE_BRUSH_SVG( "Starship/Menus/File/LevelOpen", Icon16x16 ) );
+		Set( "LevelEditor.Save",          new IMAGE_BRUSH_SVG( "Starship/Menus/File/SaveCurrent", Icon16x16 ) );
+		Set( "LevelEditor.SaveAs",        new IMAGE_BRUSH_SVG( "Starship/Menus/File/SaveCurrentAs", Icon16x16 ) );
+		Set( "LevelEditor.SaveAllLevels", new IMAGE_BRUSH_SVG( "Starship/Menus/File/LevelSaveAll", Icon16x16 ) );
+
+		Set( "LevelEditor.ImportScene",    new IMAGE_BRUSH_SVG( "Starship/Menus/File/LevelImportInto", Icon16x16 ) );
+		Set( "LevelEditor.ExportAll",      new IMAGE_BRUSH_SVG( "Starship/Menus/File/ExportAll", Icon16x16 ) );
+		Set( "LevelEditor.ExportSelected", new IMAGE_BRUSH_SVG( "Starship/Menus/File/ExportSelected", Icon16x16 ) );
 
 		Set( "LevelEditor.Build", new IMAGE_BRUSH( "Icons/icon_build_40x", Icon40x40 ) );
 		Set( "LevelEditor.Build.Small", new IMAGE_BRUSH( "Icons/icon_build_40x", Icon20x20 ) );
diff --git a/Engine/Source/Editor/MainFrame/Private/Menus/MainMenu.cpp b/Engine/Source/Editor/MainFrame/Private/Menus/MainMenu.cpp
--- a/Engine/Source/Editor/MainFrame/Private/Menus/MainMenu.cpp
+++ b/Engine/Source/Editor/MainFrame/Private/Menus/MainMenu.cpp
@@ -437,6 +437,8 @@ void FMainMenu::RegisterFileProjectMenu()
 	Section.AddMenuEntry( FMainFrameCommands::Get().OpenProject );
 
 	FText ShortIDEName = FSourceCodeNavigation::GetSelectedSourceCodeIDE();
+	FSlateIcon OpenIDEIcon = FSourceCodeNavigation::GetOpenSourceCodeIDEIcon();
+	FSlateIcon RefreshIDEIcon = FSourceCodeNavigation::GetRefreshSourceCodeIDEIcon(); 
 
 	Section.AddMenuEntry( FMainFrameCommands::Get().AddCodeToProject,
 		TAttribute<FText>(),
@@ -465,27 +467,30 @@ void FMainMenu::RegisterFileProjectMenu()
 	);
 	*/
 
-	Section.AddDynamicEntry("CodeProject", FNewToolMenuSectionDelegate::CreateLambda([ShortIDEName](FToolMenuSection& InSection)
+	Section.AddDynamicEntry("CodeProject", FNewToolMenuSectionDelegate::CreateLambda([ShortIDEName, RefreshIDEIcon](FToolMenuSection& InSection)
 	{
 		if (FSourceCodeNavigation::DoesModuleSolutionExist())
 		{
 			InSection.AddMenuEntry( FMainFrameCommands::Get().RefreshCodeProject,
 				FText::Format(LOCTEXT("RefreshCodeProjectLabel", "Refresh {0} Project"), ShortIDEName),
-				FText::Format(LOCTEXT("RefreshCodeProjectTooltip", "Refreshes your C++ code project in {0}."), ShortIDEName)
+				FText::Format(LOCTEXT("RefreshCodeProjectTooltip", "Refreshes your C++ code project in {0}."), ShortIDEName),
+				RefreshIDEIcon	
 			);
 		}
 		else
 		{
 			InSection.AddMenuEntry( FMainFrameCommands::Get().RefreshCodeProject,
 				FText::Format(LOCTEXT("GenerateCodeProjectLabel", "Generate {0} Project"), ShortIDEName),
-				FText::Format(LOCTEXT("GenerateCodeProjectTooltip", "Generates your C++ code project in {0}."), ShortIDEName)
+				FText::Format(LOCTEXT("GenerateCodeProjectTooltip", "Generates your C++ code project in {0}."), ShortIDEName),
+				RefreshIDEIcon	
 			);
 		}
 	}));
 
 	Section.AddMenuEntry( FMainFrameCommands::Get().OpenIDE,
 		FText::Format(LOCTEXT("OpenIDELabel", "Open {0}"), ShortIDEName),
-		FText::Format(LOCTEXT("OpenIDETooltip", "Opens your C++ code in {0}."), ShortIDEName)
+		FText::Format(LOCTEXT("OpenIDETooltip", "Opens your C++ code in {0}."), ShortIDEName),
+		OpenIDEIcon
 	);
 
 	Section.AddDynamicEntry("CookContentForPlatform", FNewToolMenuSectionDelegate::CreateLambda([](FToolMenuSection& InSection)
@@ -507,7 +512,7 @@ void FMainMenu::RegisterFileProjectMenu()
 				"CookContentForPlatform",
 				FText::Format(LOCTEXT("CookContentForPlatform", "Cook Content for {0}"), CookedPlatformText),
 				FText::Format(LOCTEXT("CookContentForPlatformTooltip", "Cook your game content for debugging on the {0} platform"), CookedPlatformText),
-				FSlateIcon(),
+				FSlateIcon(FEditorStyle::GetStyleSetName(), "MainFrame.CookContent"),
 				Action
 			);
 		}
diff --git a/Engine/Source/Editor/UnrealEd/Private/SourceCodeNavigation.cpp b/Engine/Source/Editor/UnrealEd/Private/SourceCodeNavigation.cpp
--- a/Engine/Source/Editor/UnrealEd/Private/SourceCodeNavigation.cpp
+++ b/Engine/Source/Editor/UnrealEd/Private/SourceCodeNavigation.cpp
@@ -1760,6 +1760,26 @@ FText FSourceCodeNavigation::GetSuggestedSourceCodeIDE(bool bShortIDEName)
 #endif
 }
 
+
+FSlateIcon FSourceCodeNavigation::GetOpenSourceCodeIDEIcon()
+{
+#if PLATFORM_WINDOWS
+		return FSlateIcon(FAppStyle::GetAppStyleSetName(), "MainFrame.OpenVisualStudio");
+#else
+	return FSlateIcon();
+#endif
+}
+
+
+FSlateIcon FSourceCodeNavigation::GetRefreshSourceCodeIDEIcon()
+{
+#if PLATFORM_WINDOWS
+		return FSlateIcon(FAppStyle::GetAppStyleSetName(), "MainFrame.RefreshVisualStudio");
+#else
+	return FSlateIcon();
+#endif
+}
+
 FString FSourceCodeNavigation::GetSuggestedSourceCodeIDEDownloadURL()
 {
 	FString SourceCodeIDEURL;
diff --git a/Engine/Source/Editor/UnrealEd/Public/SourceCodeNavigation.h b/Engine/Source/Editor/UnrealEd/Public/SourceCodeNavigation.h
--- a/Engine/Source/Editor/UnrealEd/Public/SourceCodeNavigation.h
+++ b/Engine/Source/Editor/UnrealEd/Public/SourceCodeNavigation.h
@@ -3,6 +3,7 @@
 #pragma once
 
 #include "CoreMinimal.h"
+#include "Textures/SlateIcon.h"
 #include "UObject/WeakObjectPtr.h"
 
 class AActor;
@@ -317,6 +318,12 @@ public:
 	/** Returns the name of the selected IDE */
 	UNREALED_API static FText GetSelectedSourceCodeIDE();
 
+	/** Returns the IDE specific representative icon to open source code*/
+	UNREALED_API static FSlateIcon GetOpenSourceCodeIDEIcon();
+
+	/** Returns the IDE specific representative icon to refresh source code*/
+	UNREALED_API static FSlateIcon GetRefreshSourceCodeIDEIcon();
+
 	/** Returns the url to the location where the suggested IDE can be downloaded */
 	UNREALED_API static FString GetSuggestedSourceCodeIDEDownloadURL();
 
