Fix content browser drawer icon

diff --git a/Engine/Source/Editor/EditorStyle/Private/StarshipStyle.cpp b/Engine/Source/Editor/EditorStyle/Private/StarshipStyle.cpp
--- a/Engine/Source/Editor/EditorStyle/Private/StarshipStyle.cpp
+++ b/Engine/Source/Editor/EditorStyle/Private/StarshipStyle.cpp
@@ -7732,7 +7732,6 @@ void FStarshipEditorStyle::FStyle::SetupStatusBarStyle()
 	Set("StatusBar.StatusBarButton", StatusBarButton);
 	Set("StatusBar.StatusBarComboButton", StatusBarComboButton);
 
-	Set("StatusBar.ContentBrowserIcon", new CORE_IMAGE_BRUSH_SVG("Starship/Common/tile", Icon16x16));
 	Set("StatusBar.ContentBrowserUp", new CORE_IMAGE_BRUSH_SVG("Starship/Common/chevron-up", Icon16x16));
 	Set("StatusBar.ContentBrowserDown", new CORE_IMAGE_BRUSH_SVG("Starship/Common/chevron-down", Icon16x16));
 	Set("StatusBar.ContentBrowserShadow", new BOX_BRUSH("Starship/StatusBar/drawer-shadow", FMargin(10/64.0f, 20/64.f, 10/64.f, 0), FLinearColor(0,0,0,1)));
diff --git a/Engine/Source/Editor/StatusBar/Private/SStatusBar.cpp b/Engine/Source/Editor/StatusBar/Private/SStatusBar.cpp
--- a/Engine/Source/Editor/StatusBar/Private/SStatusBar.cpp
+++ b/Engine/Source/Editor/StatusBar/Private/SStatusBar.cpp
@@ -577,7 +577,7 @@ TSharedRef<SWidget> SStatusBar::MakeContentBrowserWidget()
 			[
 				SNew(SImage)
 				.ColorAndOpacity(FSlateColor::UseForeground())
-				.Image(FAppStyle::Get().GetBrush("StatusBar.ContentBrowserIcon"))
+				.Image(FAppStyle::Get().GetBrush("ContentBrowser.TabIcon"))
 			]
 			+ SHorizontalBox::Slot()
 			.VAlign(VAlign_Center)
