Updated Window Menu Icons

diff --git a/Engine/Plugins/Compositing/Composure/Source/ComposureLayersEditor/Private/CompElementEditorModule.cpp b/Engine/Plugins/Compositing/Composure/Source/ComposureLayersEditor/Private/CompElementEditorModule.cpp
--- a/Engine/Plugins/Compositing/Composure/Source/ComposureLayersEditor/Private/CompElementEditorModule.cpp
+++ b/Engine/Plugins/Compositing/Composure/Source/ComposureLayersEditor/Private/CompElementEditorModule.cpp
@@ -275,7 +275,7 @@ void FCompElementEditorModule::RegisterEditorTab()
 		FLevelEditorModule& LevelEditorModule = FModuleManager::GetModuleChecked<FLevelEditorModule>(CompElementEditor_Impl::LevelEditorModuleName);
 		TSharedPtr<FTabManager> LevelEditorTabManager = LevelEditorModule.GetLevelEditorTabManager();
 
-		const FSlateIcon LayersIcon(FEditorStyle::GetStyleSetName(), "LevelEditor.Tabs.Layers");
+		const FSlateIcon LayersIcon(FEditorStyle::GetStyleSetName(), "LevelEditor.Tabs.ComposureCompositing");
 		LevelEditorTabManager->RegisterTabSpawner(CompElementEditor_Impl::ComposureLayersTabName, FOnSpawnTab::CreateStatic(&FCompElementEditorModule::SpawnComposureLayersTab))
 			.SetDisplayName(NSLOCTEXT("LevelEditorTabs", "LevelEditorComposureLayerBrowser", "Composure Compositing"))
 			.SetTooltipText(NSLOCTEXT("LevelEditorTabs", "LevelEditorComposureLayerBrowserTooltipText", "Open the Composure compositing tab."))
diff --git a/Engine/Plugins/Importers/USDImporter/Source/USDStageEditor/Private/USDStageEditorModule.cpp b/Engine/Plugins/Importers/USDImporter/Source/USDStageEditor/Private/USDStageEditorModule.cpp
--- a/Engine/Plugins/Importers/USDImporter/Source/USDStageEditor/Private/USDStageEditorModule.cpp
+++ b/Engine/Plugins/Importers/USDImporter/Source/USDStageEditor/Private/USDStageEditorModule.cpp
@@ -47,7 +47,7 @@ public:
 				FLevelEditorModule& LevelEditorModule = FModuleManager::GetModuleChecked< FLevelEditorModule >( "LevelEditor" );
 				TSharedPtr< FTabManager > LevelEditorTabManager = LevelEditorModule.GetLevelEditorTabManager();
 
-				const FSlateIcon LayersIcon(FEditorStyle::GetStyleSetName(), "LevelEditor.Tabs.Layers");
+				const FSlateIcon LayersIcon(FEditorStyle::GetStyleSetName(), "LevelEditor.Tabs.USDStage");
 
 				LevelEditorTabManager->RegisterTabSpawner( TEXT("USDStage"), FOnSpawnTab::CreateStatic( &FUsdStageEditorModule::SpawnUsdStageTab ) )
 					.SetDisplayName( LOCTEXT( "USDStage", "USD Stage" ) )
diff --git a/Engine/Plugins/Tests/FunctionalTestingEditor/Source/Private/FunctionalTestingEditorModule.cpp b/Engine/Plugins/Tests/FunctionalTestingEditor/Source/Private/FunctionalTestingEditorModule.cpp
--- a/Engine/Plugins/Tests/FunctionalTestingEditor/Source/Private/FunctionalTestingEditorModule.cpp
+++ b/Engine/Plugins/Tests/FunctionalTestingEditor/Source/Private/FunctionalTestingEditorModule.cpp
@@ -150,7 +150,7 @@ class FFunctionalTestingEditorModule : public IFunctionalTestingEditorModule
 		MenuBuilder.AddMenuEntry(
 			LOCTEXT("AutomationLabel", "Test Automation"),
 			LOCTEXT("Tooltip", "Launch the Testing Automation Frontend."),
-			FSlateIcon(FEditorStyle::GetStyleSetName(), "AutomationTools.MenuIcon"),
+			FSlateIcon(FAppStyle::GetAppStyleSetName(), "AutomationTools.TestAutomation"),
 			FUIAction(FExecuteAction::CreateStatic(&FFunctionalTestingEditorModule::OnShowAutomationFrontend)));
 		MenuBuilder.EndSection();
 
@@ -160,7 +160,10 @@ class FFunctionalTestingEditorModule : public IFunctionalTestingEditorModule
 			MenuBuilder.AddSubMenu(
 				LOCTEXT("AutomationTools", "Automation Tools"),
 				LOCTEXT("AutomationToolsToolTip", "Assorted tools to help generate data for some of the automation tests."),
-				FNewMenuDelegate::CreateRaw(this, &FFunctionalTestingEditorModule::PopulateAutomationTools)
+				FNewMenuDelegate::CreateRaw(this, &FFunctionalTestingEditorModule::PopulateAutomationTools),
+				false,
+				FSlateIcon(FAppStyle::GetAppStyleSetName(), "AutomationTools.MenuIcon")
+
 			);
 		}
 	}
diff --git a/Engine/Source/Editor/EditorStyle/Private/StarshipStyle.cpp b/Engine/Source/Editor/EditorStyle/Private/StarshipStyle.cpp
--- a/Engine/Source/Editor/EditorStyle/Private/StarshipStyle.cpp
+++ b/Engine/Source/Editor/EditorStyle/Private/StarshipStyle.cpp
@@ -1040,12 +1040,13 @@ void FStarshipEditorStyle::FStyle::SetupGeneralStyles()
 
 	// Developer Tools Menu
 	{
-		Set("DeveloperTools.MenuIcon", new IMAGE_BRUSH_SVG( "Starship/Menus/Duo/LevelEditor/DeveloperTools", Icon16x16 ) );
+		Set("DeveloperTools.MenuIcon", new IMAGE_BRUSH_SVG( "Starship/Menus/LevelEditor/DeveloperTools", Icon16x16 ) );
 	}
 
 	// Automation Tools Menu
 	{
-		Set("AutomationTools.MenuIcon", new IMAGE_BRUSH_SVG("Starship/Menus/Duo/LevelEditor/AutomationTools", Icon16x16));
+		Set("AutomationTools.MenuIcon", new IMAGE_BRUSH_SVG("Starship/Menus/LevelEditor/AutomationTools", Icon16x16));
+		Set("AutomationTools.TestAutomation", new IMAGE_BRUSH_SVG("Starship/Menus/LevelEditor/TestAutomation", Icon16x16));
 	}
 
 	// Session Browser tab
@@ -1081,7 +1082,7 @@ void FStarshipEditorStyle::FStyle::SetupGeneralStyles()
 
 	// Launcher Window
 	{
-		Set("Launcher.TabIcon", new IMAGE_BRUSH_SVG( "Starship/Menus/Duo/LevelEditor/ProjectLauncher", Icon16x16 ) );
+		Set("Launcher.TabIcon", new IMAGE_BRUSH_SVG( "Starship/Menus/LevelEditor/ProjectLauncher", Icon16x16 ) );
 		Set("Launcher.Tabs.Tools", new IMAGE_BRUSH( "/Icons/icon_tab_Tools_16x", Icon16x16 ) );
 	}
 
@@ -4182,6 +4183,11 @@ void FStarshipEditorStyle::FStyle::SetupLevelEditorStyle()
 		Set("LevelEditor.Tutorials", new IMAGE_BRUSH("Icons/Help/icon_Help_tutorials_16x", Icon16x16));
 		Set("LevelEditor.BrowseViewportControls", new IMAGE_BRUSH("Icons/Help/icon_Help_Documentation_16x", Icon16x16));
 
+		Set("MainFrame.ToggleFullscreen", new IMAGE_BRUSH_SVG("Starship/Menus/LevelEditor/EnableFullscreen", Icon16x16));
+		Set("MainFrame.LoadLayout", new IMAGE_BRUSH_SVG("Starship/Menus/LevelEditor/LayoutLoad", Icon16x16));
+		Set("MainFrame.SaveLayout", new IMAGE_BRUSH_SVG("Starship/Menus/LevelEditor/LayoutSave", Icon16x16));
+		Set("MainFrame.RemoveLayout", new IMAGE_BRUSH_SVG("Starship/Menus/LevelEditor/LayoutRemove", Icon16x16));
+
 		Set("MainFrame.VisitAskAQuestionPage", new IMAGE_BRUSH("Icons/Help/icon_Help_ask_16x", Icon16x16));
 		Set("MainFrame.VisitWiki", new IMAGE_BRUSH("Icons/Help/icon_Help_Documentation_16x", Icon16x16));
 		Set("MainFrame.VisitForums", new IMAGE_BRUSH("Icons/Help/icon_Help_Documentation_16x", Icon16x16));
@@ -4253,30 +4259,32 @@ void FStarshipEditorStyle::FStyle::SetupLevelEditorStyle()
 
 #if WITH_EDITOR || (IS_PROGRAM && WITH_UNREAL_DEVELOPER_TOOLS)
 		{
-			Set( "LevelEditor.Tabs.Details",                new IMAGE_BRUSH_SVG("Starship/Menus/Duo/LevelEditor/Details", Icon16x16) );///Icons/icon_tab_SelectionDetails_16x", Icon16x16 ) );
-			Set( "LevelEditor.Tabs.Cinematics",             new IMAGE_BRUSH_SVG("Starship/Menus/Duo/LevelEditor/Cinematics", Icon16x16) );///Icons/icon_tab_Cinematics_16x", Icon16x16));
-			Set( "LevelEditor.Tabs.EditorModes",            new IMAGE_BRUSH_SVG("Starship/Menus/Duo/LevelEditor/EditorModes", Icon16x16) );///Icons/icon_Editor_Modes_16x", Icon16x16 ) );
-			Set( "LevelEditor.Tabs.Modes",                  new IMAGE_BRUSH_SVG("Starship/Menus/Duo/LevelEditor/EditorModes", Icon16x16) );///Icons/icon_Editor_Modes_16x", Icon16x16 ) );
-			Set( "LevelEditor.Tabs.PlacementBrowser",       new IMAGE_BRUSH_SVG("Starship/Menus/Duo/LevelEditor/PlaceActors", Icon16x16) );///Icons/icon_Editor_Modes_16x", Icon16x16));
-			Set( "LevelEditor.Tabs.Properties",             new IMAGE_BRUSH_SVG("Starship/Menus/Duo/LevelEditor/Properties", Icon16x16) );///Icons/properties_16x", Icon16x16 ) );
-			Set( "LevelEditor.Tabs.Outliner",               new IMAGE_BRUSH_SVG("Starship/Menus/Duo/LevelEditor/WorldOutliner", Icon16x16) );///Icons/icon_tab_SceneOutliner_16x", Icon16x16 ) );
-			Set( "LevelEditor.Tabs.ContentBrowser",         new IMAGE_BRUSH_SVG("Starship/Menus/Duo/LevelEditor/ContentBrowser", Icon16x16) );///Icons/icon_tab_ContentBrowser_16x", Icon16x16 ) );
-			Set( "LevelEditor.Tabs.Levels",                 new IMAGE_BRUSH_SVG("Starship/Menus/Duo/LevelEditor/Levels", Icon16x16) );///Icons/icon_tab_Levels_16x", Icon16x16 ) );
-			Set( "LevelEditor.Tabs.WorldBrowser",           new IMAGE_BRUSH_SVG("Starship/Menus/Duo/LevelEditor/Levels", Icon16x16) );///Icons/icon_tab_levels_16x", Icon16x16 ) );
-			Set( "LevelEditor.Tabs.WorldBrowserDetails",    new IMAGE_BRUSH_SVG("Starship/Menus/Duo/LevelEditor/WorldSettings", Icon16x16) );///Icons/icon_levels_detailsbutton_16x", Icon16x16 ) );
+			Set( "LevelEditor.Tabs.Details",                new IMAGE_BRUSH_SVG("Starship/Menus/LevelEditor/Details", Icon16x16) );
+			Set( "LevelEditor.Tabs.Cinematics",             new IMAGE_BRUSH_SVG("Starship/Menus/LevelEditor/Cinematics", Icon16x16) );
+			Set( "LevelEditor.Tabs.EditorModes",            new IMAGE_BRUSH_SVG("Starship/Menus/LevelEditor/EditorModes", Icon16x16) );
+			Set( "LevelEditor.Tabs.Modes",                  new IMAGE_BRUSH_SVG("Starship/Menus/LevelEditor/EditorModes", Icon16x16) );
+			Set( "LevelEditor.Tabs.PlacementBrowser",       new IMAGE_BRUSH_SVG("Starship/Menus/LevelEditor/PlaceActors", Icon16x16) );
+			Set( "LevelEditor.Tabs.Properties",             new IMAGE_BRUSH_SVG("Starship/Menus/LevelEditor/Properties", Icon16x16) );
+			Set( "LevelEditor.Tabs.Outliner",               new IMAGE_BRUSH_SVG("Starship/Menus/LevelEditor/WorldOutliner", Icon16x16) );
+			Set( "LevelEditor.Tabs.ContentBrowser",         new IMAGE_BRUSH_SVG("Starship/Menus/LevelEditor/ContentBrowser", Icon16x16) );
+			Set( "LevelEditor.Tabs.Levels",                 new IMAGE_BRUSH_SVG("Starship/Menus/LevelEditor/Levels", Icon16x16) );
+			Set( "LevelEditor.Tabs.WorldBrowser",           new IMAGE_BRUSH_SVG("Starship/Menus/LevelEditor/Levels", Icon16x16) );
+			Set( "LevelEditor.Tabs.WorldBrowserDetails",    new IMAGE_BRUSH_SVG("Starship/Menus/LevelEditor/WorldSettings", Icon16x16) );
 
 			Set( "LevelEditor.Tabs.WorldBrowserComposition",new IMAGE_BRUSH("Icons/icon_levels_compositionbutton_16x", Icon16x16 ) );
 
-			Set( "LevelEditor.Tabs.Layers",                 new IMAGE_BRUSH_SVG("Starship/Menus/Duo/LevelEditor/Layers", Icon16x16) );///Icons/icon_tab_Layers_16x", Icon16x16 ) );
+			Set( "LevelEditor.Tabs.Layers",                 new IMAGE_BRUSH_SVG("Starship/Menus/LevelEditor/Layers", Icon16x16) );
+			Set( "LevelEditor.Tabs.ComposureCompositing",   new IMAGE_BRUSH_SVG("Starship/Menus/LevelEditor/ComposureCompositing", Icon16x16) );
+			Set( "LevelEditor.Tabs.USDStage",   			new IMAGE_BRUSH_SVG("Starship/Menus/LevelEditor/USDStage", Icon16x16) );
 
 			Set( "LevelEditor.Tabs.BuildAndSubmit",         new IMAGE_BRUSH("Icons/icon_tab_BuildSubmit_16x", Icon16x16 ) );
 
-			Set( "LevelEditor.Tabs.StatsViewer",            new IMAGE_BRUSH_SVG("Starship/Menus/Duo/LevelEditor/Statistics", Icon16x16) );///Icons/icon_tab_Stats_16x", Icon16x16 ) );
+			Set( "LevelEditor.Tabs.StatsViewer",            new IMAGE_BRUSH_SVG("Starship/Menus/LevelEditor/Statistics", Icon16x16) );
 
 			Set( "LevelEditor.Tabs.Toolbar",                new IMAGE_BRUSH("Icons/icon_tab_Toolbars_16x", Icon16x16 ) );
 
-			Set( "LevelEditor.Tabs.Viewports",              new IMAGE_BRUSH_SVG("Starship/Menus/Duo/LevelEditor/Viewports", Icon16x16) );///Icons/icon_tab_Viewports_16x", Icon16x16 ) );
-			Set( "LevelEditor.Tabs.HLOD",                   new IMAGE_BRUSH_SVG("Starship/Menus/Duo/LevelEditor/HierarchicalLODOutliner", Icon16x16) );///Icons/icon_tab_layers_16px", Icon16x16));
+			Set( "LevelEditor.Tabs.Viewports",              new IMAGE_BRUSH_SVG("Starship/Menus/LevelEditor/Viewports", Icon16x16) );
+			Set( "LevelEditor.Tabs.HLOD",                   new IMAGE_BRUSH_SVG("Starship/Menus/LevelEditor/HierarchicalLODOutliner", Icon16x16) );
 		}
 #endif // WITH_EDITOR || (IS_PROGRAM && WITH_UNREAL_DEVELOPER_TOOLS)
 
@@ -4375,7 +4383,7 @@ void FStarshipEditorStyle::FStyle::SetupLevelEditorStyle()
 		Set( "LevelEditor.WorldProperties", new IMAGE_BRUSH( "Icons/icon_worldscript_40x", Icon40x40 ) );
 		Set( "LevelEditor.WorldProperties.Small", new IMAGE_BRUSH( "Icons/icon_worldscript_40x", Icon20x20 ) );
 
-		Set( "LevelEditor.WorldProperties.Tab", new IMAGE_BRUSH_SVG( "Starship/Menus/Duo/LevelEditor/WorldSettings", Icon16x16 ) );
+		Set( "LevelEditor.WorldProperties.Tab", new IMAGE_BRUSH_SVG( "Starship/Menus/LevelEditor/WorldSettings", Icon16x16 ) );
 
 		Set( "LevelEditor.OpenContentBrowser", new IMAGE_BRUSH_SVG( "Starship/MainToolBar/content", Icon20x20 ) );
 		Set( "LevelEditor.OpenMarketplace", new IMAGE_BRUSH_SVG( "Starship/MainToolBar/marketplace", Icon20x20) );
@@ -4537,8 +4545,8 @@ void FStarshipEditorStyle::FStyle::SetupLevelEditorStyle()
 		Set( "LevelViewport.PilotSelectedActor", new IMAGE_BRUSH( "Icons/icon_PilotSelectedActor_16x", Icon16x16 ) );
 		Set( "LevelViewport.PilotSelectedActor.Small", new IMAGE_BRUSH( "Icons/icon_PilotSelectedActor_16x", Icon16x16 ) );
 		
-		Set( "LevelViewport.ToggleActorPilotCameraView",       new IMAGE_BRUSH_SVG( "Starship/Menus/Duo/LevelEditor/CameraShake", Icon16x16 ) );
-		Set( "LevelViewport.ToggleActorPilotCameraView.Small", new IMAGE_BRUSH_SVG( "Starship/Menus/Duo/LevelEditor/CameraShake", Icon16x16 ) );
+		Set( "LevelViewport.ToggleActorPilotCameraView",       new IMAGE_BRUSH_SVG( "Starship/Menus/LevelEditor/CameraShake", Icon16x16 ) );
+		Set( "LevelViewport.ToggleActorPilotCameraView.Small", new IMAGE_BRUSH_SVG( "Starship/Menus/LevelEditor/CameraShake", Icon16x16 ) );
 
 		Set( "LevelViewport.ActorPilotText", FTextBlockStyle()
 			.SetFont( DEFAULT_FONT( "BoldCondensed", 12 ) )
@@ -4955,7 +4963,7 @@ void FStarshipEditorStyle::FStyle::SetupPersonaStyle()
 
 		Set( "BlueprintEditor.ActionMenu.ContextDescriptionFont",  DEFAULT_FONT("Regular", 12) );
 
-		Set( "BlueprintEditor.FindInBlueprints.MenuIcon", new IMAGE_BRUSH_SVG("Starship/Menus/Duo/LevelEditor/FindInBlueprint", Icon16x16) );
+		Set( "BlueprintEditor.FindInBlueprints.MenuIcon", new IMAGE_BRUSH_SVG("Starship/Menus/LevelEditor/FindInBlueprints", Icon16x16) );
 
 		Set( "BlueprintEditor.FindInBlueprint", new IMAGE_BRUSH( "Icons/icon_Blueprint_Find_40px", Icon40x40 ) );
 		Set( "BlueprintEditor.FindInBlueprint.Small", new IMAGE_BRUSH( "Icons/icon_Blueprint_Find_40px", Icon20x20 ) );
@@ -5850,7 +5858,7 @@ void FStarshipEditorStyle::FStyle::SetupContentBrowserStyle()
 	// Content Browser
 	{
 		// Tab and menu icon
-		Set( "ContentBrowser.TabIcon", new IMAGE_BRUSH_SVG("Starship/Menus/Duo/LevelEditor/ContentBrowser", Icon16x16) );
+		Set( "ContentBrowser.TabIcon", new IMAGE_BRUSH_SVG("Starship/Menus/LevelEditor/ContentBrowser", Icon16x16) );
 
 		// Sources View
 		Set( "ContentBrowser.SourceTitleFont", DEFAULT_FONT( "Regular", 12 ) );
@@ -7644,6 +7652,8 @@ void FStarshipEditorStyle::FStyle::SetupTranslationEditorStyles()
 
 void FStarshipEditorStyle::FStyle::SetupLocalizationDashboardStyles()
 {
+	Set("LocalizationDashboard.MenuIcon", new IMAGE_BRUSH_SVG("Starship/Menus/LevelEditor/LocalizationDashboard", Icon16x16));
+
 	Set("LocalizationDashboard.GatherTextAllTargets", new IMAGE_BRUSH("Icons/Icon_Localisation_Gather_All_40x", Icon40x40));
 	Set("LocalizationDashboard.ImportTextAllTargetsAllCultures", new IMAGE_BRUSH("Icons/Icon_Localisation_Import_All_40x", Icon40x40));
 	Set("LocalizationDashboard.ExportTextAllTargetsAllCultures", new IMAGE_BRUSH("Icons/Icon_Localisation_Export_All_40x", Icon40x40));
diff --git a/Engine/Source/Editor/MainFrame/Private/Menus/MainMenu.cpp b/Engine/Source/Editor/MainFrame/Private/Menus/MainMenu.cpp
--- a/Engine/Source/Editor/MainFrame/Private/Menus/MainMenu.cpp
+++ b/Engine/Source/Editor/MainFrame/Private/Menus/MainMenu.cpp
@@ -229,7 +229,7 @@ void FMainMenu::RegisterWindowMenu()
 						"LocalizationDashboard",
 						LOCTEXT("LocalizationDashboardLabel", "Localization Dashboard"),
 						LOCTEXT("LocalizationDashboardToolTip", "Open the Localization Dashboard for this Project."),
-						FSlateIcon(),
+						FSlateIcon(FAppStyle::GetAppStyleSetName(), "LocalizationDashboard.MenuIcon"), 
 						FUIAction(FExecuteAction::CreateStatic(&FMainMenu::OpenLocalizationDashboard))
 						);
 				}
@@ -257,7 +257,9 @@ void FMainMenu::RegisterWindowMenu()
 			"LoadLayout",
 			NSLOCTEXT("LayoutMenu", "LayoutLoadHeader", "Load Layout"),
 			NSLOCTEXT("LayoutMenu", "LoadLayoutsSubMenu_ToolTip", "Load a layout configuration from disk"),
-			FNewToolMenuDelegate::CreateStatic(&FLayoutsMenuLoad::MakeLoadLayoutsMenu)
+			FNewToolMenuDelegate::CreateStatic(&FLayoutsMenuLoad::MakeLoadLayoutsMenu),
+			false,
+			FSlateIcon(FAppStyle::GetAppStyleSetName(), "MainFrame.LoadLayout") 
 		));
 		// Save and Remove Layout
 		// Opposite to "Load Layout", Save and Remove are dynamic, i.e., they can be enabled/removed depending on the value of
@@ -271,14 +273,18 @@ void FMainMenu::RegisterWindowMenu()
 					"OverrideLayout",
 					NSLOCTEXT("LayoutMenu", "OverrideLayoutsSubMenu", "Save Layout"),
 					NSLOCTEXT("LayoutMenu", "OverrideLayoutsSubMenu_ToolTip", "Save your current layout configuration on disk"),
-					FNewToolMenuDelegate::CreateStatic(&FLayoutsMenuSave::MakeSaveLayoutsMenu)
+					FNewToolMenuDelegate::CreateStatic(&FLayoutsMenuSave::MakeSaveLayoutsMenu),
+					false,
+					FSlateIcon(FAppStyle::GetAppStyleSetName(), "MainFrame.SaveLayout") 
 				));
 				// Remove Layout
 				InSection.AddEntry(FToolMenuEntry::InitSubMenu(
 					"RemoveLayout",
 					NSLOCTEXT("LayoutMenu", "RemoveLayoutsSubMenu", "Remove Layout"),
 					NSLOCTEXT("LayoutMenu", "RemoveLayoutsSubMenu_ToolTip", "Remove a layout configuration from disk"),
-					FNewToolMenuDelegate::CreateStatic(&FLayoutsMenuRemove::MakeRemoveLayoutsMenu)
+					FNewToolMenuDelegate::CreateStatic(&FLayoutsMenuRemove::MakeRemoveLayoutsMenu),
+					false,
+					FSlateIcon(FAppStyle::GetAppStyleSetName(), "MainFrame.RemoveLayout")
 				));
 			}
 		}));
