Help Menu Icons

diff --git a/Engine/Source/Editor/EditorStyle/Private/StarshipStyle.cpp b/Engine/Source/Editor/EditorStyle/Private/StarshipStyle.cpp
--- a/Engine/Source/Editor/EditorStyle/Private/StarshipStyle.cpp
+++ b/Engine/Source/Editor/EditorStyle/Private/StarshipStyle.cpp
@@ -4178,24 +4178,28 @@ void FStarshipEditorStyle::FStyle::SetupLevelEditorStyle()
 	// Level editor tool bar icons
 #if WITH_EDITOR
 	{
-		Set("LevelEditor.BrowseDocumentation", new IMAGE_BRUSH("Icons/Help/icon_Help_Documentation_16x", Icon16x16));
-		Set("LevelEditor.BrowseAPIReference", new IMAGE_BRUSH("Icons/Help/icon_Help_api-1_16x", Icon16x16));
-		Set("LevelEditor.Tutorials", new IMAGE_BRUSH("Icons/Help/icon_Help_tutorials_16x", Icon16x16));
-		Set("LevelEditor.BrowseViewportControls", new IMAGE_BRUSH("Icons/Help/icon_Help_Documentation_16x", Icon16x16));
-
-		Set("MainFrame.ToggleFullscreen", new IMAGE_BRUSH_SVG("Starship/Menus/LevelEditor/EnableFullscreen", Icon16x16));
-		Set("MainFrame.LoadLayout", new IMAGE_BRUSH_SVG("Starship/Menus/LevelEditor/LayoutLoad", Icon16x16));
-		Set("MainFrame.SaveLayout", new IMAGE_BRUSH_SVG("Starship/Menus/LevelEditor/LayoutSave", Icon16x16));
-		Set("MainFrame.RemoveLayout", new IMAGE_BRUSH_SVG("Starship/Menus/LevelEditor/LayoutRemove", Icon16x16));
-
-		Set("MainFrame.VisitAskAQuestionPage", new IMAGE_BRUSH("Icons/Help/icon_Help_ask_16x", Icon16x16));
-		Set("MainFrame.VisitWiki", new IMAGE_BRUSH("Icons/Help/icon_Help_Documentation_16x", Icon16x16));
-		Set("MainFrame.VisitForums", new IMAGE_BRUSH("Icons/Help/icon_Help_Documentation_16x", Icon16x16));
-		Set("MainFrame.VisitSearchForAnswersPage", new IMAGE_BRUSH("Icons/Help/icon_Help_search_16x", Icon16x16));
-		Set("MainFrame.VisitSupportWebSite", new IMAGE_BRUSH("Icons/Help/icon_Help_support_16x", Icon16x16));
-		Set("MainFrame.VisitEpicGamesDotCom", new IMAGE_BRUSH("Icons/Help/icon_Help_epic_16x", Icon16x16));
-		Set("MainFrame.AboutUnrealEd", new IMAGE_BRUSH("Icons/Help/icon_Help_unreal_16x", Icon16x16));
-		Set("MainFrame.CreditsUnrealEd", new IMAGE_BRUSH("Icons/Help/icon_Help_credits_16x", Icon16x16));
+		Set("LevelEditor.BrowseDocumentation",     new IMAGE_BRUSH_SVG("Starship/Menus/Help/Documentation", Icon16x16));
+		Set("LevelEditor.BrowseAPIReference",      new IMAGE_BRUSH_SVG("Starship/Menus/Help/Documentation", Icon16x16));
+		Set("LevelEditor.BrowseCVars",             new IMAGE_BRUSH_SVG("Starship/Menus/Help/Console", Icon16x16));
+		Set("LevelEditor.Tutorials",               new IMAGE_BRUSH_SVG("Starship/Menus/Help/Tutorials", Icon16x16));
+		Set("LevelEditor.BrowseViewportControls",  new IMAGE_BRUSH_SVG("Starship/Menus/Help/ViewportControls", Icon16x16));
+
+		Set("MainFrame.ToggleFullscreen",          new IMAGE_BRUSH_SVG("Starship/Menus/LevelEditor/EnableFullscreen", Icon16x16));
+		Set("MainFrame.LoadLayout",                new IMAGE_BRUSH_SVG("Starship/Menus/LevelEditor/LayoutLoad", Icon16x16));
+		Set("MainFrame.SaveLayout",                new IMAGE_BRUSH_SVG("Starship/Menus/LevelEditor/LayoutSave", Icon16x16));
+		Set("MainFrame.RemoveLayout",              new IMAGE_BRUSH_SVG("Starship/Menus/LevelEditor/LayoutRemove", Icon16x16));
+
+		Set("MainFrame.OpenIssueTracker",          new IMAGE_BRUSH_SVG("Starship/Menus/Help/IssueTracker", Icon16x16));
+		Set("MainFrame.ReportABug",                new IMAGE_BRUSH_SVG("Starship/Menus/Help/Bug", Icon16x16));
+
+		Set("MainFrame.VisitAskAQuestionPage",     new IMAGE_BRUSH_SVG("Starship/Menus/Bug", Icon16x16));
+		Set("MainFrame.VisitWiki",                 new IMAGE_BRUSH_SVG("Starship/Menus/Help/WebBrowser", Icon16x16));
+		Set("MainFrame.VisitForums",               new IMAGE_BRUSH_SVG("Starship/Menus/Help/WebBrowser", Icon16x16));
+		Set("MainFrame.VisitSearchForAnswersPage", new IMAGE_BRUSH_SVG("Starship/Menus/Help/WebBrowser", Icon16x16));
+		Set("MainFrame.VisitSupportWebSite",       new IMAGE_BRUSH_SVG("Starship/Menus/Help/Support", Icon16x16));
+		Set("MainFrame.VisitEpicGamesDotCom",      new IMAGE_BRUSH("Icons/Help/icon_Help_epic_16x", Icon16x16));
+		Set("MainFrame.AboutUnrealEd",             new IMAGE_BRUSH("Icons/Help/icon_Help_unreal_16x", Icon16x16));
+		Set("MainFrame.CreditsUnrealEd",           new IMAGE_BRUSH_SVG("Starship/Menus/Help/Credits", Icon16x16));
 
 		Set( "EditorViewport.TranslateMode", new IMAGE_BRUSH_SVG( "Starship/EditorViewport/translate", Icon16x16 ) );
 		Set( "EditorViewport.RotateMode", new IMAGE_BRUSH_SVG("Starship/EditorViewport/rotate", Icon16x16 ) );
