Fix stress test unity CIS

diff --git a/Engine/Source/Editor/EditorStyle/Private/SlateEditorStyle.cpp b/Engine/Source/Editor/EditorStyle/Private/SlateEditorStyle.cpp
--- a/Engine/Source/Editor/EditorStyle/Private/SlateEditorStyle.cpp
+++ b/Engine/Source/Editor/EditorStyle/Private/SlateEditorStyle.cpp
@@ -33,14 +33,6 @@ BEGIN_SLATE_FUNCTION_BUILD_OPTIMIZATION
 TSharedPtr< FSlateEditorStyle::FStyle > FSlateEditorStyle::StyleInstance = NULL;
 TWeakObjectPtr< UEditorStyleSettings > FSlateEditorStyle::Settings = NULL;
 
-static bool IncludeEditorSpecificStyles()
-{
-#if IS_PROGRAM
-	return true;
-#else
-	return GIsEditor;
-#endif
-}
 
 /* FSlateEditorStyle interface
  *****************************************************************************/
@@ -174,6 +166,15 @@ void FSlateEditorStyle::FStyle::SyncSettings()
 	}
 }
 
+bool FSlateEditorStyle::FStyle::IncludeEditorSpecificStyles()
+{
+#if IS_PROGRAM
+	return true;
+#else
+	return GIsEditor;
+#endif
+}
+
 void FSlateEditorStyle::FStyle::Initialize()
 {
 	SyncSettings();
diff --git a/Engine/Source/Editor/EditorStyle/Private/SlateEditorStyle.h b/Engine/Source/Editor/EditorStyle/Private/SlateEditorStyle.h
--- a/Engine/Source/Editor/EditorStyle/Private/SlateEditorStyle.h
+++ b/Engine/Source/Editor/EditorStyle/Private/SlateEditorStyle.h
@@ -167,6 +167,9 @@ public:
 		FButtonStyle HoverHintOnly;
 
 		TWeakObjectPtr< UEditorStyleSettings > Settings;
+
+		static bool IncludeEditorSpecificStyles();
+
 	};
 
 	static TSharedRef< class FSlateEditorStyle::FStyle > Create( const TWeakObjectPtr< UEditorStyleSettings >& InCustomization )
diff --git a/Engine/Source/Editor/EditorStyle/Private/StarshipStyle.cpp b/Engine/Source/Editor/EditorStyle/Private/StarshipStyle.cpp
--- a/Engine/Source/Editor/EditorStyle/Private/StarshipStyle.cpp
+++ b/Engine/Source/Editor/EditorStyle/Private/StarshipStyle.cpp
@@ -59,8 +59,6 @@ void FStarshipEditorStyle::Shutdown()
 		SettingsModule->UnregisterSettings("Editor", "General", "Appearance");
 	}
 
-
-
 	FPropertyEditorModule* PropertyEditorModule = FModuleManager::Get().GetModulePtr<FPropertyEditorModule>("PropertyEditor");
 	if (PropertyEditorModule)
 	{
@@ -79,7 +77,15 @@ BEGIN_SLATE_FUNCTION_BUILD_OPTIMIZATION
 TSharedPtr< FStarshipEditorStyle::FStyle > FStarshipEditorStyle::StyleInstance = NULL;
 TWeakObjectPtr< UEditorStyleSettings > FStarshipEditorStyle::Settings = NULL;
 
-static bool IncludeEditorSpecificStyles()
+void FStarshipEditorStyle::FStyle::SetColor(const TSharedRef< FLinearColor >& Source, const FLinearColor& Value)
+{
+	Source->R = Value.R;
+	Source->G = Value.G;
+	Source->B = Value.B;
+	Source->A = Value.A;
+}
+
+bool FStarshipEditorStyle::FStyle::IncludeEditorSpecificStyles()
 {
 #if IS_PROGRAM
 	return true;
@@ -148,13 +154,6 @@ FStarshipEditorStyle::FStyle::FStyle( const TWeakObjectPtr< UEditorStyleSettings
 
 }
 
-static void SetColor( const TSharedRef< FLinearColor >& Source, const FLinearColor& Value )
-{
-	Source->R = Value.R;
-	Source->G = Value.G;
-	Source->B = Value.B;
-	Source->A = Value.A;
-}
 
 void FStarshipEditorStyle::FStyle::SettingsChanged(UObject* ChangedObject, FPropertyChangedEvent& PropertyChangedEvent)
 {
diff --git a/Engine/Source/Editor/EditorStyle/Private/StarshipStyle.h b/Engine/Source/Editor/EditorStyle/Private/StarshipStyle.h
--- a/Engine/Source/Editor/EditorStyle/Private/StarshipStyle.h
+++ b/Engine/Source/Editor/EditorStyle/Private/StarshipStyle.h
@@ -76,6 +76,9 @@ public:
 		void SyncSettings();
 		void SyncParentStyles();
 
+		static void SetColor(const TSharedRef<FLinearColor>& Source, const FLinearColor& Value);
+		static bool IncludeEditorSpecificStyles();
+
 		const FVector2D Icon7x16;
 		const FVector2D Icon8x4;
 		const FVector2D Icon16x4;
