EditorViewport  Updated Perspective and Lighting Menu Icons

diff --git a/Engine/Source/Editor/EditorStyle/Private/StarshipStyle.cpp b/Engine/Source/Editor/EditorStyle/Private/StarshipStyle.cpp
--- a/Engine/Source/Editor/EditorStyle/Private/StarshipStyle.cpp
+++ b/Engine/Source/Editor/EditorStyle/Private/StarshipStyle.cpp
@@ -1000,7 +1000,8 @@ void FStarshipEditorStyle::FStyle::SetupGeneralStyles()
 			.SetButtonStyle(DebugConsoleButton);
 
 		Set("DebugConsole.ComboButton", DebugConsoleComboButton);
-		Set("DebugConsole.Icon", new CORE_IMAGE_BRUSH_SVG("Starship/Common/layout-header-body", Icon16x16));
+
+		Set("DebugConsole.Icon", new IMAGE_BRUSH_SVG("Starship/Menus/Help/Console", Icon16x16));
 
 		Set("OutputLog.OpenSourceLocation", new IMAGE_BRUSH("Icons/icon_Asset_Open_Source_Location_16x", Icon16x16));
 		Set("OutputLog.OpenInExternalEditor", new IMAGE_BRUSH("Icons/icon_Asset_Open_In_External_Editor_16x", Icon16x16));
@@ -4223,14 +4224,17 @@ void FStarshipEditorStyle::FStyle::SetupLevelEditorStyle()
 
 		Set( "EditorViewport.CamSpeedSetting", new IMAGE_BRUSH_SVG( "Starship/EditorViewport/camera", Icon16x16) );
 		
-		Set( "EditorViewport.LitMode", new IMAGE_BRUSH( "Icons/icon_ViewMode_Lit_16px", Icon16x16 ) );
-		Set( "EditorViewport.UnlitMode", new IMAGE_BRUSH( "Icons/icon_ViewMode_Unlit_16px", Icon16x16 ) );
-		Set( "EditorViewport.WireframeMode", new IMAGE_BRUSH( "Icons/icon_ViewMode_BrushWireframe_16px", Icon16x16 ) );
-		Set( "EditorViewport.DetailLightingMode", new IMAGE_BRUSH( "Icons/icon_ViewMode_DetailLighting_16px", Icon16x16 ) );
-		Set( "EditorViewport.LightingOnlyMode", new IMAGE_BRUSH( "Icons/icon_ViewMode_LightingOnly_16px", Icon16x16 ) );
+		Set( "EditorViewport.LitMode",            	  new IMAGE_BRUSH_SVG("Starship/EditorViewport/View/Lit", Icon16x16 ) );
+		Set( "EditorViewport.UnlitMode",          	  new IMAGE_BRUSH_SVG("Starship/EditorViewport/View/Unlit", Icon16x16 ) );
+		Set( "EditorViewport.WireframeMode",      	  new IMAGE_BRUSH_SVG("Starship/EditorViewport/View/BrushWireframe", Icon16x16 ) );
+		Set( "EditorViewport.DetailLightingMode", 	  new IMAGE_BRUSH_SVG("Starship/EditorViewport/View/DetailLighting", Icon16x16 ) );
+		Set( "EditorViewport.LightingOnlyMode",   	  new IMAGE_BRUSH_SVG("Starship/EditorViewport/View/LightingOnly", Icon16x16 ) );
+
 		Set( "EditorViewport.LightComplexityMode", new IMAGE_BRUSH( "Icons/icon_ViewMode_LightComplexity_16px", Icon16x16 ) );
 		Set( "EditorViewport.ShaderComplexityMode", new IMAGE_BRUSH( "Icons/icon_ViewMode_Shadercomplexity_16px", Icon16x16 ) );
-		Set( "EditorViewport.QuadOverdrawMode", new IMAGE_BRUSH( "Icons/icon_ViewMode_QuadOverdraw_16px", Icon16x16 ) );
+
+		Set( "EditorViewport.QuadOverdrawMode", new IMAGE_BRUSH_SVG("Starship/EditorViewport/View/OptimizationViewmodes", Icon16x16 ) );
+
 		Set( "EditorViewport.ShaderComplexityWithQuadOverdrawMode", new IMAGE_BRUSH( "Icons/icon_ViewMode_Shadercomplexity_16px", Icon16x16 ) );
 		Set( "EditorViewport.TexStreamAccPrimitiveDistanceMode", new IMAGE_BRUSH( "Icons/icon_ViewMode_TextureStreamingAccuracy_16px", Icon16x16 ) );
 		Set( "EditorViewport.TexStreamAccMeshUVDensityMode", new IMAGE_BRUSH("Icons/icon_ViewMode_TextureStreamingAccuracy_16px", Icon16x16));
@@ -4240,21 +4244,26 @@ void FStarshipEditorStyle::FStyle::SetupLevelEditorStyle()
 		Set( "EditorViewport.LightmapDensityMode", new IMAGE_BRUSH( "Icons/icon_ViewMode_LightmapDensity_16px", Icon16x16 ) );
 
 		Set( "EditorViewport.LODColorationMode", new IMAGE_BRUSH("Icons/icon_ViewMode_LODColoration_16px", Icon16x16) );
-		Set( "EditorViewport.HLODColorationMode", new IMAGE_BRUSH("Icons/icon_ViewMode_LODColoration_16px", Icon16x16));
-		Set( "EditorViewport.GroupLODColorationMode", new IMAGE_BRUSH("Icons/icon_ViewMode_LODColoration_16px", Icon16x16));
-
-		Set( "EditorViewport.VisualizeGBufferMode", new IMAGE_BRUSH("Icons/icon_ViewMode_VisualisationGBuffer_16px", Icon16x16) );
-		Set( "EditorViewport.ReflectionOverrideMode", new IMAGE_BRUSH( "Icons/icon_ViewMode_ReflectionOverride_16px", Icon16x16 ) );
-		Set( "EditorViewport.VisualizeBufferMode", new IMAGE_BRUSH( "Icons/icon_ViewMode_VisualisationGBuffer_16px", Icon16x16 ) );
-		Set( "EditorViewport.CollisionPawn", new IMAGE_BRUSH( "Icons/icon_ViewMode_CollsionPawn_16px", Icon16x16 ) );
-		Set( "EditorViewport.CollisionVisibility", new IMAGE_BRUSH( "Icons/icon_ViewMode_CollisionVisibility_16px", Icon16x16 ) );
-		Set( "EditorViewport.Perspective", new IMAGE_BRUSH( "Icons/icon_ViewMode_ViewPerspective_16px", Icon16x16 ) );
-		Set( "EditorViewport.Top", new IMAGE_BRUSH( "Icons/icon_ViewMode_ViewTop_16px", Icon16x16 ) );
-		Set( "EditorViewport.Left", new IMAGE_BRUSH( "Icons/icon_ViewMode_ViewLeft_16px", Icon16x16 ) );
-		Set( "EditorViewport.Front", new IMAGE_BRUSH( "Icons/icon_ViewMode_ViewFront_16px", Icon16x16 ) );
-		Set( "EditorViewport.Bottom", new IMAGE_BRUSH("Icons/icon_ViewMode_ViewBottom_16px", Icon16x16 ) );
-		Set( "EditorViewport.Right", new IMAGE_BRUSH("Icons/icon_ViewMode_ViewRight_16px", Icon16x16 ) );
-		Set( "EditorViewport.Back", new IMAGE_BRUSH("Icons/icon_ViewMode_ViewBack_16px", Icon16x16 ) );
+		Set( "EditorViewport.HLODColorationMode", new IMAGE_BRUSH("Icons/icon_ViewMode_LODColoration_16px", Icon16x16));	
+		Set( "EditorViewport.GroupLODColorationMode", new IMAGE_BRUSH_SVG("Starship/EditorViewport/View/LODColorization", Icon16x16) );
+
+		Set( "EditorViewport.VisualizeGBufferMode",   new IMAGE_BRUSH_SVG("Starship/EditorViewport/View/BufferVisualization", Icon16x16) );
+
+		Set( "EditorViewport.Visualizers", 			  new IMAGE_BRUSH_SVG("Starship/EditorViewport/View/Visualizer", Icon16x16) );
+		Set( "EditorViewport.LOD", 			  		  new IMAGE_BRUSH_SVG("Starship/EditorViewport/View/LOD", Icon16x16) );
+
+		Set( "EditorViewport.ReflectionOverrideMode", new IMAGE_BRUSH_SVG("Starship/EditorViewport/View/Reflections", Icon16x16 ) );
+		Set( "EditorViewport.VisualizeBufferMode",    new IMAGE_BRUSH_SVG("Starship/EditorViewport/View/BufferVisualization", Icon16x16 ) );
+		Set( "EditorViewport.CollisionPawn",          new IMAGE_BRUSH_SVG("Starship/EditorViewport/View/PlayerCollision", Icon16x16 ) );
+		Set( "EditorViewport.CollisionVisibility",    new IMAGE_BRUSH_SVG("Starship/EditorViewport/View/VisibilityCollision", Icon16x16 ) );
+		
+		Set( "EditorViewport.Perspective", new IMAGE_BRUSH_SVG("Starship/EditorViewport/Perspective/Perspective", Icon16x16 ) );
+		Set( "EditorViewport.Top",         new IMAGE_BRUSH_SVG("Starship/EditorViewport/Perspective/Top", Icon16x16 ) );
+		Set( "EditorViewport.Left",        new IMAGE_BRUSH_SVG("Starship/EditorViewport/Perspective/Left", Icon16x16 ) );
+		Set( "EditorViewport.Front",       new IMAGE_BRUSH_SVG("Starship/EditorViewport/Perspective/Front", Icon16x16 ) );
+		Set( "EditorViewport.Bottom",      new IMAGE_BRUSH_SVG("Starship/EditorViewport/Perspective/Bottom", Icon16x16 ) );
+		Set( "EditorViewport.Right",       new IMAGE_BRUSH_SVG("Starship/EditorViewport/Perspective/Right", Icon16x16 ) );
+		Set( "EditorViewport.Back",        new IMAGE_BRUSH_SVG("Starship/EditorViewport/Perspective/Back", Icon16x16 ) );
 
 		Set("EditorViewport.ActiveBorderColor", FStyleColors::Primary);
 
diff --git a/Engine/Source/Editor/LandscapeEditor/Private/LandscapeEditorModule.cpp b/Engine/Source/Editor/LandscapeEditor/Private/LandscapeEditorModule.cpp
--- a/Engine/Source/Editor/LandscapeEditor/Private/LandscapeEditorModule.cpp
+++ b/Engine/Source/Editor/LandscapeEditor/Private/LandscapeEditorModule.cpp
@@ -175,7 +175,13 @@ public:
 				InMenuBuilder.EndSection();
 			}
 		};
-		MenuBuilder.AddSubMenu(LOCTEXT("LandscapeSubMenu", "Visualizers"), LOCTEXT("LandscapeSubMenu_ToolTip", "Select a Landscape visualiser"), FNewMenuDelegate::CreateStatic(&Local::BuildLandscapeVisualizersMenu));
+		MenuBuilder.AddSubMenu(
+			LOCTEXT("LandscapeSubMenu", "Visualizers"), 
+			LOCTEXT("LandscapeSubMenu_ToolTip", "Select a Landscape visualiser"), 
+			FNewMenuDelegate::CreateStatic(&Local::BuildLandscapeVisualizersMenu), 
+			false, 
+			FSlateIcon(FAppStyle::GetAppStyleSetName(), "EditorViewport.Visualizers")
+		);
 	}
 
 	static void ChangeLandscapeViewMode(ELandscapeViewMode::Type ViewMode)
diff --git a/Engine/Source/Editor/LevelEditor/Private/SLevelViewportToolBar.cpp b/Engine/Source/Editor/LevelEditor/Private/SLevelViewportToolBar.cpp
--- a/Engine/Source/Editor/LevelEditor/Private/SLevelViewportToolBar.cpp
+++ b/Engine/Source/Editor/LevelEditor/Private/SLevelViewportToolBar.cpp
@@ -1522,7 +1522,7 @@ void SLevelViewportToolBar::FillViewMenu(UToolMenu* Menu)
 			LOCTEXT("LandscapeLODMenu_ToolTip", "Override Landscape LOD in this viewport"),
 			FNewToolMenuDelegate::CreateStatic(&Local::BuildLandscapeLODMenu, this),
 			/*bInOpenSubMenuOnClick=*/ false,
-			FSlateIcon()
+			FSlateIcon(FAppStyle::GetAppStyleSetName(), "EditorViewport.LOD")
 		);
 	}
 }
