Fix output log colors not functioning properly.

diff --git a/Engine/Source/Developer/OutputLog/Private/SOutputLog.cpp b/Engine/Source/Developer/OutputLog/Private/SOutputLog.cpp
--- a/Engine/Source/Developer/OutputLog/Private/SOutputLog.cpp
+++ b/Engine/Source/Developer/OutputLog/Private/SOutputLog.cpp
@@ -847,7 +847,6 @@ void SOutputLog::Construct( const FArguments& InArgs )
 	MessagesTextBox = SNew(SMultiLineEditableTextBox)
 		.Style(FEditorStyle::Get(), "Log.TextBox")
 		.TextStyle(FEditorStyle::Get(), "Log.Normal")
-		.ForegroundColor(FLinearColor::Gray)
 		.Marshaller(MessagesTextMarshaller)
 		.IsReadOnly(true)
 		.AlwaysShowScrollbars(true)
@@ -1071,7 +1070,7 @@ bool SOutputLog::CreateLogMessages( const TCHAR* V, ELogVerbosity::Type Verbosit
 				Line = Line.ConvertTabsToSpaces(4);
 
 				// Hard-wrap lines to avoid them being too long
-				static const int32 HardWrapLen = 360;
+				static const int32 HardWrapLen = 600;
 				for (int32 CurrentStartIndex = 0; CurrentStartIndex < Line.Len();)
 				{
 					int32 HardWrapLineLen = 0;
diff --git a/Engine/Source/Editor/EditorStyle/Private/EditorStyleSettingsCustomization.cpp b/Engine/Source/Editor/EditorStyle/Private/EditorStyleSettingsCustomization.cpp
--- a/Engine/Source/Editor/EditorStyle/Private/EditorStyleSettingsCustomization.cpp
+++ b/Engine/Source/Editor/EditorStyle/Private/EditorStyleSettingsCustomization.cpp
@@ -4,6 +4,42 @@
 #include "DetailCategoryBuilder.h"
 #include "Styling/StyleColors.h"
 #include "DetailLayoutBuilder.h"
+#include "IDetailChildrenBuilder.h"
+
+TSharedRef<IPropertyTypeCustomization> FStyleColorListCustomization::MakeInstance()
+{
+	return MakeShared<FStyleColorListCustomization>();
+}
+
+void FStyleColorListCustomization::CustomizeHeader(TSharedRef<IPropertyHandle> PropertyHandle, FDetailWidgetRow& HeaderRow, IPropertyTypeCustomizationUtils& CustomizationUtils)
+{
+	
+}
+
+void FStyleColorListCustomization::CustomizeChildren(TSharedRef<IPropertyHandle> PropertyHandle, IDetailChildrenBuilder& ChildBuilder, IPropertyTypeCustomizationUtils& CustomizationUtils)
+{
+	uint32 NumChildren = 0;
+	TSharedPtr<IPropertyHandle> ColorArrayProperty = PropertyHandle->GetChildHandle(GET_MEMBER_NAME_CHECKED(FStyleColorList, StyleColors));
+
+	ColorArrayProperty->GetNumChildren(NumChildren);
+	for (uint32 ChildIndex = 0; ChildIndex < NumChildren; ++ChildIndex)
+	{
+		if (ChildIndex < (uint32)EStyleColor::User1)
+		{
+			ChildBuilder.AddProperty(ColorArrayProperty->GetChildHandle(ChildIndex).ToSharedRef());
+		}
+		else
+		{
+			// user colors are added if they have been customized with a display name
+			FText DisplayName = UStyleColorTable::Get().GetColorDisplayName((EStyleColor)ChildIndex);
+			if (!DisplayName.IsEmpty())
+			{
+				IDetailPropertyRow& Row = ChildBuilder.AddProperty(ColorArrayProperty->GetChildHandle(ChildIndex).ToSharedRef());
+				Row.DisplayName(DisplayName);
+			}
+		}
+	}
+}
 
 
 TSharedRef<IDetailCustomization> FEditorStyleSettingsCustomization::MakeInstance()
@@ -14,9 +50,10 @@ TSharedRef<IDetailCustomization> FEditorStyleSettingsCustomization::MakeInstance
 
 void FEditorStyleSettingsCustomization::CustomizeDetails(IDetailLayoutBuilder& DetailLayout)
 {
-	IDetailCategoryBuilder& ColorCategory = DetailLayout.EditCategory("Color");
+	IDetailCategoryBuilder& ColorCategory = DetailLayout.EditCategory("Colors");
 
 	TArray<UObject*> Objects = { &UStyleColorTable::Get() };
 
-	ColorCategory.AddExternalObjectProperty(Objects, "Colors");
-}
\ No newline at end of file
+	ColorCategory.AddExternalObjectProperty(Objects, "ColorList");
+}
+
diff --git a/Engine/Source/Editor/EditorStyle/Private/EditorStyleSettingsCustomization.h b/Engine/Source/Editor/EditorStyle/Private/EditorStyleSettingsCustomization.h
--- a/Engine/Source/Editor/EditorStyle/Private/EditorStyleSettingsCustomization.h
+++ b/Engine/Source/Editor/EditorStyle/Private/EditorStyleSettingsCustomization.h
@@ -3,9 +3,20 @@
 #pragma once
 
 #include "IDetailCustomization.h"
+#include "IPropertyTypeCustomization.h"
 
 class IDetailLayoutBuilder;
 
+class FStyleColorListCustomization : public IPropertyTypeCustomization
+{
+public:
+	static TSharedRef<IPropertyTypeCustomization> MakeInstance();
+
+	virtual void CustomizeHeader(TSharedRef<IPropertyHandle> PropertyHandle, FDetailWidgetRow& HeaderRow, IPropertyTypeCustomizationUtils& CustomizationUtils) override;
+	virtual void CustomizeChildren(TSharedRef<IPropertyHandle> PropertyHandle, IDetailChildrenBuilder& ChildBuilder, IPropertyTypeCustomizationUtils& CustomizationUtils) override;
+
+};
+
 class FEditorStyleSettingsCustomization : public IDetailCustomization
 {
 public:
diff --git a/Engine/Source/Editor/EditorStyle/Private/SlateEditorStyle.cpp b/Engine/Source/Editor/EditorStyle/Private/SlateEditorStyle.cpp
--- a/Engine/Source/Editor/EditorStyle/Private/SlateEditorStyle.cpp
+++ b/Engine/Source/Editor/EditorStyle/Private/SlateEditorStyle.cpp
@@ -127,13 +127,6 @@ void FSlateEditorStyle::FStyle::SyncSettings()
 		// Sync the colors used by FEditorStyle
 		SetColor( SelectionColor_LinearRef, Settings->SelectionColor );
 
-		SetColor( LogColor_Background_LinearRef, Settings->LogBackgroundColor );
-		SetColor( LogColor_SelectionBackground_LinearRef, Settings->LogSelectionBackgroundColor );
-		SetColor( LogColor_Normal_LinearRef, Settings->LogNormalColor );
-		SetColor( LogColor_Command_LinearRef, Settings->LogCommandColor );
-		SetColor( LogColor_Warning_LinearRef, Settings->LogWarningColor );
-		SetColor( LogColor_Error_LinearRef, Settings->LogErrorColor );
-
 		// The subdued selection color is derived from the selection color
 		auto SubduedSelectionColor = Settings->GetSubduedSelectionColor();
 		SetColor( SelectionColor_Subdued_LinearRef, SubduedSelectionColor );
@@ -148,17 +141,6 @@ void FSlateEditorStyle::FStyle::SyncSettings()
 
 		WindowBackgroundColor = Settings->EditorWindowBackgroundColor;
 
-		FSlateBrush DummyBrush;
-		if (Settings->EditorMainWindowBackgroundOverride != DummyBrush)
-		{
-			WindowBackgroundMain = Settings->EditorMainWindowBackgroundOverride;
-		}
-
-		if (Settings->EditorChildWindowBackgroundOverride != DummyBrush)
-		{
-			WindowBackgroundChild = Settings->EditorChildWindowBackgroundOverride;
-		}
-
 		FWindowStyle& WindowStyle = const_cast<FWindowStyle&>(FCoreStyle::Get().GetWidgetStyle<FWindowStyle>("Window"));
 		WindowStyle.SetBackgroundColor(WindowBackgroundColor)
 			.SetBackgroundBrush(WindowBackgroundMain)
@@ -3043,7 +3025,7 @@ void FSlateEditorStyle::FStyle::SetupWindowStyles()
 		{
 			WindowBackgroundColor = Settings->EditorWindowBackgroundColor;
 
-			FSlateBrush DummyBrush;
+		/*	FSlateBrush DummyBrush;
 			if (Settings->EditorMainWindowBackgroundOverride != DummyBrush)
 			{
 				WindowBackgroundMain = Settings->EditorMainWindowBackgroundOverride;
@@ -3052,7 +3034,7 @@ void FSlateEditorStyle::FStyle::SetupWindowStyles()
 			if (Settings->EditorChildWindowBackgroundOverride != DummyBrush)
 			{
 				WindowBackgroundChild = Settings->EditorChildWindowBackgroundOverride;
-			}
+			}*/
 		}
 
 		FWindowStyle& WindowStyle = const_cast<FWindowStyle&>(FCoreStyle::Get().GetWidgetStyle<FWindowStyle>("Window"));
diff --git a/Engine/Source/Editor/EditorStyle/Private/StarshipStyle.cpp b/Engine/Source/Editor/EditorStyle/Private/StarshipStyle.cpp
--- a/Engine/Source/Editor/EditorStyle/Private/StarshipStyle.cpp
+++ b/Engine/Source/Editor/EditorStyle/Private/StarshipStyle.cpp
@@ -20,6 +20,8 @@
 #define DEFAULT_FONT(...) FStarshipCoreStyle::GetDefaultFontStyle(__VA_ARGS__)
 #define ICON_FONT(...) FSlateFontInfo(RootToContentDir("Fonts/FontAwesome", TEXT(".ttf")), __VA_ARGS__)
 
+#define LOCTEXT_NAMESPACE "EditorStyle"
+
 void FStarshipEditorStyle::Initialize()
 {
 	Settings = NULL;
@@ -34,8 +36,8 @@ void FStarshipEditorStyle::Initialize()
 	if (SettingsModule != nullptr)
 	{
 		SettingsModule->RegisterSettings("Editor", "General", "Appearance",
-			NSLOCTEXT("EditorStyle", "Appearance_UserSettingsName", "Appearance"),
-			NSLOCTEXT("EditorStyle", "Appearance_UserSettingsDescription", "Customize the look of the editor."),
+			LOCTEXT("Appearance_UserSettingsName", "Appearance"),
+			LOCTEXT("Appearance_UserSettingsDescription", "Customize the look of the editor."),
 			Settings
 		);
 	}
@@ -43,8 +45,28 @@ void FStarshipEditorStyle::Initialize()
 
 	FPropertyEditorModule& PropertyEditorModule = FModuleManager::Get().GetModuleChecked<FPropertyEditorModule>("PropertyEditor");
 	PropertyEditorModule.RegisterCustomClassLayout("EditorStyleSettings", FOnGetDetailCustomizationInstance::CreateStatic(&FEditorStyleSettingsCustomization::MakeInstance));
+	PropertyEditorModule.RegisterCustomPropertyTypeLayout("StyleColorList", FOnGetPropertyTypeCustomizationInstance::CreateStatic(&FStyleColorListCustomization::MakeInstance));
 #endif
 
+	// Background
+	UStyleColorTable::Get().SetColor(EStyleColor::User1, UStyleColorTable::Get().GetColor(EStyleColor::Input));
+	UStyleColorTable::Get().SetColorDisplayName(EStyleColor::User1, LOCTEXT("UserColor_OutputLogBackground", "Log Background"));
+	// Selection highlight
+	UStyleColorTable::Get().SetColor(EStyleColor::User2, UStyleColorTable::Get().GetColor(EStyleColor::Highlight));
+	UStyleColorTable::Get().SetColorDisplayName(EStyleColor::User2, LOCTEXT("UserColor_OutputLogHighlight", "Log Highlight"));
+	// Normal
+	UStyleColorTable::Get().SetColor(EStyleColor::User3, UStyleColorTable::Get().GetColor(EStyleColor::Foreground));
+	UStyleColorTable::Get().SetColorDisplayName(EStyleColor::User3, LOCTEXT("UserColor_OutputLogText", "Log Text"));
+	// Command
+	UStyleColorTable::Get().SetColor(EStyleColor::User4, UStyleColorTable::Get().GetColor(EStyleColor::AccentGreen));
+	UStyleColorTable::Get().SetColorDisplayName(EStyleColor::User4, LOCTEXT("UserColor_OutputLogCommand", "Log Command"));
+	// Warning
+	UStyleColorTable::Get().SetColor(EStyleColor::User5, UStyleColorTable::Get().GetColor(EStyleColor::AccentYellow));
+	UStyleColorTable::Get().SetColorDisplayName(EStyleColor::User5, LOCTEXT("UserColor_OutputLogWarning", "Log Warning"));
+	// Error
+	UStyleColorTable::Get().SetColor(EStyleColor::User6, UStyleColorTable::Get().GetColor(EStyleColor::AccentRed));
+	UStyleColorTable::Get().SetColorDisplayName(EStyleColor::User6, LOCTEXT("UserColor_OutputLogError", "Log Error"));
+
 	StyleInstance = Create(Settings);
 	SetStyle(StyleInstance.ToSharedRef());
 }
@@ -128,30 +150,24 @@ FStarshipEditorStyle::FStyle::FStyle( const TWeakObjectPtr< UEditorStyleSettings
 	, HighlightColor_LinearRef( MakeShareable( new FLinearColor(0.068f, 0.068f, 0.068f) ) ) 
 	, WindowHighlightColor_LinearRef(MakeShareable(new FLinearColor(0,0,0,0)))
 
-	, LogColor_Background_LinearRef(MakeShareable(new FLinearColor(FColor(0x040404FF))))
-	, LogColor_SelectionBackground_LinearRef(MakeShareable(new FLinearColor(FColor(0x020202FF))))
-	, LogColor_Normal_LinearRef(MakeShareable(new FLinearColor(FColor(0xffaaaaaa))))
-	, LogColor_Command_LinearRef(MakeShareable(new FLinearColor(FColor(0xff33dd33))))
-	, LogColor_Warning_LinearRef(MakeShareable(new FLinearColor(FColor(0xEBB000FF))))
-	, LogColor_Error_LinearRef(MakeShareable(new FLinearColor(FColor(0xFF0D0FFF))))
+
 
 	// These are the Slate colors which reference those above; these are the colors to put into the style
 	, SelectionColor_Subdued( SelectionColor_Subdued_LinearRef )
 	, HighlightColor( HighlightColor_LinearRef )
 	, WindowHighlightColor(WindowHighlightColor_LinearRef)
 
-	, LogColor_Background( LogColor_Background_LinearRef )
-	, LogColor_SelectionBackground( LogColor_SelectionBackground_LinearRef )
-	, LogColor_Normal( LogColor_Normal_LinearRef )
-	, LogColor_Command( LogColor_Command_LinearRef )
-	, LogColor_Warning( LogColor_Warning_LinearRef )
-	, LogColor_Error( LogColor_Error_LinearRef )
+	, LogColor_Background(EStyleColor::User1)
+	, LogColor_SelectionBackground(EStyleColor::User2)
+	, LogColor_Normal(EStyleColor::User3)
+	, LogColor_Command(EStyleColor::User4)
+	, LogColor_Warning(EStyleColor::User5)
+	, LogColor_Error(EStyleColor::User6)
 
 	, InheritedFromBlueprintTextColor(FLinearColor(0.25f, 0.5f, 1.0f))
 
 	, Settings( InSettings )
 {
-
 }
 
 
@@ -167,14 +183,6 @@ void FStarshipEditorStyle::FStyle::SyncSettings()
 {
 	if (Settings.IsValid())
 	{
-		SetColor(LogColor_Background_LinearRef, Settings->LogBackgroundColor);
-		SetColor(LogColor_SelectionBackground_LinearRef, Settings->LogSelectionBackgroundColor);
-		SetColor(LogColor_Normal_LinearRef, Settings->LogNormalColor);
-		SetColor(LogColor_Command_LinearRef, Settings->LogCommandColor);
-		SetColor(LogColor_Warning_LinearRef, Settings->LogWarningColor);
-		SetColor(LogColor_Error_LinearRef, Settings->LogErrorColor);
-
-
 		// The subdued selection color is derived from the selection color
 		auto SubduedSelectionColor = Settings->GetSubduedSelectionColor();
 		SetColor(SelectionColor_Subdued_LinearRef, SubduedSelectionColor);
@@ -189,17 +197,6 @@ void FStarshipEditorStyle::FStyle::SyncSettings()
 		// Sync the window background settings
 		FWindowStyle& WindowStyle = const_cast<FWindowStyle&>(FStarshipCoreStyle::GetCoreStyle().GetWidgetStyle<FWindowStyle>("Window"));
 
-		FSlateBrush DummyBrush;
-		if (Settings->EditorMainWindowBackgroundOverride != DummyBrush)
-		{
-			WindowStyle.SetBackgroundBrush(Settings->EditorMainWindowBackgroundOverride);
-		}
-
-		if (Settings->EditorChildWindowBackgroundOverride != DummyBrush)
-		{
-			WindowStyle.SetChildBackgroundBrush(Settings->EditorChildWindowBackgroundOverride);
-		}
-
 		if (Settings->bEnableEditorWindowBackgroundColor)
 		{
 			SetColor(WindowHighlightColor_LinearRef, Settings->EditorWindowBackgroundColor);
@@ -960,9 +957,10 @@ void FStarshipEditorStyle::FStyle::SetupGeneralStyles()
 		const int32 LogFontSize = Settings.IsValid() ? Settings->LogFontSize : 9;
 
 		const FTextBlockStyle NormalLogText = FTextBlockStyle(NormalText)
-			.SetFont( DEFAULT_FONT( "Mono", LogFontSize ) )
-			.SetColorAndOpacity( LogColor_Normal )
-			.SetSelectedBackgroundColor( LogColor_SelectionBackground );
+			.SetFont(DEFAULT_FONT("Mono", LogFontSize))
+			.SetColorAndOpacity(LogColor_Normal)
+			.SetSelectedBackgroundColor(LogColor_SelectionBackground)
+			.SetHighlightColor(FLinearColor::White);
 
 		Set("Log.Normal", NormalLogText );
 
@@ -7780,5 +7778,6 @@ void FStarshipEditorStyle::FStyle::SetupStatusBarStyle()
 #undef BORDER_BRUSH
 #undef DEFAULT_FONT
 #undef ICON_FONT
+#undef LOCTEXT_NAMESPACE
 
 END_SLATE_FUNCTION_BUILD_OPTIMIZATION
diff --git a/Engine/Source/Editor/EditorStyle/Private/StarshipStyle.h b/Engine/Source/Editor/EditorStyle/Private/StarshipStyle.h
--- a/Engine/Source/Editor/EditorStyle/Private/StarshipStyle.h
+++ b/Engine/Source/Editor/EditorStyle/Private/StarshipStyle.h
@@ -105,13 +105,6 @@ public:
 		const TSharedRef< FLinearColor > HighlightColor_LinearRef;
 		const TSharedRef< FLinearColor > WindowHighlightColor_LinearRef;
 
-		const TSharedRef< FLinearColor > LogColor_Background_LinearRef;
-		const TSharedRef< FLinearColor > LogColor_SelectionBackground_LinearRef;
-		const TSharedRef< FLinearColor > LogColor_Normal_LinearRef;
-		const TSharedRef< FLinearColor > LogColor_Command_LinearRef;
-		const TSharedRef< FLinearColor > LogColor_Warning_LinearRef;
-		const TSharedRef< FLinearColor > LogColor_Error_LinearRef;
-
 		// These are the Slate colors which reference those above; these are the colors to put into the style
 		// Most of these are owned by our parent style
 		FSlateColor DefaultForeground;
diff --git a/Engine/Source/Editor/EditorStyle/Public/Classes/EditorStyleSettings.h b/Engine/Source/Editor/EditorStyle/Public/Classes/EditorStyleSettings.h
--- a/Engine/Source/Editor/EditorStyle/Public/Classes/EditorStyleSettings.h
+++ b/Engine/Source/Editor/EditorStyle/Public/Classes/EditorStyleSettings.h
@@ -71,7 +71,7 @@ public:
 	bool bColorVisionDeficiencyCorrectionPreviewWithDeficiency;
 
 	/** The color used to represent selection */
-	UPROPERTY(EditAnywhere, config, Category=Colors, meta=(DisplayName="Selection Color"))
+	UPROPERTY(EditAnywhere, config, Category=Colors, meta=(DisplayName="Viewport Selection Color"))
 	FLinearColor SelectionColor;
 
 	UPROPERTY(config)
@@ -81,14 +81,6 @@ public:
 	UPROPERTY(EditAnywhere, config, Category=Colors, meta=(EditCondition="bEnableEditorWindowBackgroundColor"))
 	FLinearColor EditorWindowBackgroundColor;
 
-	/** The override for the background of the main window (if not modified, the defaults will be used) */
-	UPROPERTY(EditAnywhere, config, Category=Colors)
-	FSlateBrush EditorMainWindowBackgroundOverride;
-
-	/** The override for the background of the child window (if not modified, the defaults will be used) */
-	UPROPERTY(EditAnywhere, config, Category=Colors)
-	FSlateBrush EditorChildWindowBackgroundOverride;
-
 	/** Whether to use small toolbar icons without labels or not. */
 	UPROPERTY(EditAnywhere, config, Category=UserInterface)
 	uint32 bUseSmallToolBarIcons:1;
@@ -133,30 +125,6 @@ public:
 	UPROPERTY(config)
 	uint32 bShowLaunchMenus : 1;
 
-	/** The color used for the background in the output log */
-	UPROPERTY(EditAnywhere, config, Category="Output Log", meta=(DisplayName="Background Color"))
-	FLinearColor LogBackgroundColor;
-
-	/** The color used for the background of selected text in the output log */
-	UPROPERTY(EditAnywhere, config, Category="Output Log", meta=(DisplayName="Selection Background Color"))
-	FLinearColor LogSelectionBackgroundColor;
-
-	/** The color used for normal text in the output log */
-	UPROPERTY(EditAnywhere, config, Category="Output Log", meta=(DisplayName="Normal Text Color"))
-	FLinearColor LogNormalColor;
-
-	/** The color used for normal text in the output log */
-	UPROPERTY(EditAnywhere, config, Category="Output Log", meta=(DisplayName="Command Text Color"))
-	FLinearColor LogCommandColor;
-
-	/** The color used for warning log lines */
-	UPROPERTY(EditAnywhere, config, Category="Output Log", meta=(DisplayName="Warning Text Color"))
-	FLinearColor LogWarningColor;
-
-	/** The color used for error log lines */
-	UPROPERTY(EditAnywhere, config, Category="Output Log", meta=(DisplayName="Error Text Color"))
-	FLinearColor LogErrorColor;
-
 	/** When enabled, the Advanced Details will always auto expand. */
 	UPROPERTY(config)
 	uint32 bShowAllAdvancedDetails : 1;
diff --git a/Engine/Source/Runtime/SlateCore/Public/Styling/StyleColors.h b/Engine/Source/Runtime/SlateCore/Public/Styling/StyleColors.h
--- a/Engine/Source/Runtime/SlateCore/Public/Styling/StyleColors.h
+++ b/Engine/Source/Runtime/SlateCore/Public/Styling/StyleColors.h
@@ -49,9 +49,43 @@ enum class EStyleColor : uint8
 	AccentWhite,
 	AccentFolder,
 
+	/** Only user colors should be below this line
+	 * To use user colors:
+	 * 1. Set an unused user enum value below as the color value for an FSlateColor. E.g. FSlateColor MyCustomColor(EStyleColors::User1)
+	 * 2. Set the actual color. E.g UStyleColorTable::Get().SetColor(EStyleColor::User1, FLinearColor::White)
+	 * 3. Give it a display name if you want it to be configurable by editor users. E.g.  UStyleColorTable::Get().SetColorDisplayName(EUserStyleColor::User1, "My Color Name")
+	 */
+	User1,
+	User2,
+	User3,
+	User4,
+	User5,
+	User6,
+	User7,
+	User8,
+	User9,
+	User10,
+	User11,
+	User12,
+	User13,
+	User14,
+	User15,
+	User16,
+
 	MAX
 };
 
+USTRUCT()
+struct FStyleColorList
+{
+	GENERATED_BODY()
+
+	UPROPERTY(EditAnywhere, Config, Category = Colors)
+	FLinearColor StyleColors[(int32)EStyleColor::MAX];
+
+	FText DisplayNames[(int32)EStyleColor::MAX];
+};
+
 UCLASS(Config=EditorSettings)
 class SLATECORE_API UStyleColorTable : public UObject 
 {
@@ -64,12 +98,22 @@ public:
 
 	const FLinearColor& GetColor(EStyleColor Color)
 	{
-		return Colors[static_cast<int32>(Color)];
+		return ColorList.StyleColors[static_cast<int32>(Color)];
 	}
 
 	void SetColor(EStyleColor InColorId, FLinearColor InColor)
 	{
-		Colors[static_cast<int32>(InColorId)] = InColor;
+		ColorList.StyleColors[static_cast<int32>(InColorId)] = InColor;
+	}
+
+	void SetColorDisplayName(EStyleColor InColorId, FText DisplayName)
+	{
+		ColorList.DisplayNames[static_cast<int32>(InColorId)] = DisplayName;
+	}
+
+	FText GetColorDisplayName(EStyleColor InColorId) const
+	{
+		return ColorList.DisplayNames[static_cast<int32>(InColorId)];
 	}
 
 	UStyleColorTable();
@@ -89,8 +133,8 @@ public:
 	void InitalizeDefaults();
 
 private:
-	UPROPERTY(EditAnywhere, Config, Category=Colors)
-	FLinearColor Colors[static_cast<int32>(EStyleColor::MAX)];
+	UPROPERTY(EditAnywhere, Config, Category = Colors)
+	FStyleColorList ColorList;
 };
 
 
