Updated Icons

diff --git a/Engine/Source/Editor/EditorStyle/Private/StarshipStyle.cpp b/Engine/Source/Editor/EditorStyle/Private/StarshipStyle.cpp
--- a/Engine/Source/Editor/EditorStyle/Private/StarshipStyle.cpp
+++ b/Engine/Source/Editor/EditorStyle/Private/StarshipStyle.cpp
@@ -1038,13 +1038,13 @@ void FStarshipEditorStyle::FStyle::SetupGeneralStyles()
 
 	// Developer Tools Menu
 	{
-		Set("DeveloperTools.MenuIcon", new IMAGE_BRUSH_SVG( "Starship/Menus/LevelEditor/DeveloperTools", Icon16x16 ) );
+		Set("DeveloperTools.MenuIcon", new IMAGE_BRUSH_SVG( "Starship/LevelEditor/Menus/DeveloperTools", Icon16x16 ) );
 	}
 
 	// Automation Tools Menu
 	{
-		Set("AutomationTools.MenuIcon", new IMAGE_BRUSH_SVG("Starship/Menus/LevelEditor/AutomationTools", Icon16x16));
-		Set("AutomationTools.TestAutomation", new IMAGE_BRUSH_SVG("Starship/Menus/LevelEditor/TestAutomation", Icon16x16));
+		Set("AutomationTools.MenuIcon", new IMAGE_BRUSH_SVG("Starship/LevelEditor/Menus/AutomationTools", Icon16x16));
+		Set("AutomationTools.TestAutomation", new IMAGE_BRUSH_SVG("Starship/LevelEditor/Menus/TestAutomation", Icon16x16));
 	}
 
 	// Session Browser tab
@@ -1080,7 +1080,7 @@ void FStarshipEditorStyle::FStyle::SetupGeneralStyles()
 
 	// Launcher Window
 	{
-		Set("Launcher.TabIcon", new IMAGE_BRUSH_SVG( "Starship/Menus/LevelEditor/ProjectLauncher", Icon16x16 ) );
+		Set("Launcher.TabIcon", new IMAGE_BRUSH_SVG( "Starship/LevelEditor/Menus/ProjectLauncher", Icon16x16 ) );
 		Set("Launcher.Tabs.Tools", new IMAGE_BRUSH( "/Icons/icon_tab_Tools_16x", Icon16x16 ) );
 	}
 
@@ -4214,10 +4214,10 @@ void FStarshipEditorStyle::FStyle::SetupLevelEditorStyle()
 		Set("LevelEditor.Tutorials",               new IMAGE_BRUSH_SVG("Starship/Menus/Help/Tutorials", Icon16x16));
 		Set("LevelEditor.BrowseViewportControls",  new IMAGE_BRUSH_SVG("Starship/Menus/Help/ViewportControls", Icon16x16));
 
-		Set("MainFrame.ToggleFullscreen",          new IMAGE_BRUSH_SVG("Starship/Menus/LevelEditor/EnableFullscreen", Icon16x16));
-		Set("MainFrame.LoadLayout",                new IMAGE_BRUSH_SVG("Starship/Menus/LevelEditor/LayoutLoad", Icon16x16));
-		Set("MainFrame.SaveLayout",                new IMAGE_BRUSH_SVG("Starship/Menus/LevelEditor/LayoutSave", Icon16x16));
-		Set("MainFrame.RemoveLayout",              new IMAGE_BRUSH_SVG("Starship/Menus/LevelEditor/LayoutRemove", Icon16x16));
+		Set("MainFrame.ToggleFullscreen",          new IMAGE_BRUSH_SVG("Starship/LevelEditor/Menus/EnableFullscreen", Icon16x16));
+		Set("MainFrame.LoadLayout",                new IMAGE_BRUSH_SVG("Starship/LevelEditor/Menus/LayoutLoad", Icon16x16));
+		Set("MainFrame.SaveLayout",                new IMAGE_BRUSH_SVG("Starship/LevelEditor/Menus/LayoutSave", Icon16x16));
+		Set("MainFrame.RemoveLayout",              new IMAGE_BRUSH_SVG("Starship/LevelEditor/Menus/LayoutRemove", Icon16x16));
 
 		Set("MainFrame.OpenIssueTracker",          new IMAGE_BRUSH_SVG("Starship/Menus/Help/IssueTracker", Icon16x16));
 		Set("MainFrame.ReportABug",                new IMAGE_BRUSH_SVG("Starship/Menus/Help/Bug", Icon16x16));
@@ -4301,32 +4301,32 @@ void FStarshipEditorStyle::FStyle::SetupLevelEditorStyle()
 
 #if WITH_EDITOR || (IS_PROGRAM && WITH_UNREAL_DEVELOPER_TOOLS)
 		{
-			Set( "LevelEditor.Tabs.Details",                new IMAGE_BRUSH_SVG("Starship/Menus/LevelEditor/Details", Icon16x16) );
-			Set( "LevelEditor.Tabs.Cinematics",             new IMAGE_BRUSH_SVG("Starship/Menus/LevelEditor/Cinematics", Icon16x16) );
-			Set( "LevelEditor.Tabs.EditorModes",            new IMAGE_BRUSH_SVG("Starship/Menus/LevelEditor/EditorModes", Icon16x16) );
-			Set( "LevelEditor.Tabs.Modes",                  new IMAGE_BRUSH_SVG("Starship/Menus/LevelEditor/EditorModes", Icon16x16) );
-			Set( "LevelEditor.Tabs.PlacementBrowser",       new IMAGE_BRUSH_SVG("Starship/Menus/LevelEditor/PlaceActors", Icon16x16) );
-			Set( "LevelEditor.Tabs.Properties",             new IMAGE_BRUSH_SVG("Starship/Menus/LevelEditor/Properties", Icon16x16) );
-			Set( "LevelEditor.Tabs.Outliner",               new IMAGE_BRUSH_SVG("Starship/Menus/LevelEditor/WorldOutliner", Icon16x16) );
-			Set( "LevelEditor.Tabs.ContentBrowser",         new IMAGE_BRUSH_SVG("Starship/Menus/LevelEditor/ContentBrowser", Icon16x16) );
-			Set( "LevelEditor.Tabs.Levels",                 new IMAGE_BRUSH_SVG("Starship/Menus/LevelEditor/Levels", Icon16x16) );
-			Set( "LevelEditor.Tabs.WorldBrowser",           new IMAGE_BRUSH_SVG("Starship/Menus/LevelEditor/Levels", Icon16x16) );
-			Set( "LevelEditor.Tabs.WorldBrowserDetails",    new IMAGE_BRUSH_SVG("Starship/Menus/LevelEditor/WorldSettings", Icon16x16) );
+			Set( "LevelEditor.Tabs.Details",                new IMAGE_BRUSH_SVG("Starship/LevelEditor/Menus/Details", Icon16x16) );
+			Set( "LevelEditor.Tabs.Cinematics",             new IMAGE_BRUSH_SVG("Starship/LevelEditor/Menus/Cinematics", Icon16x16) );
+			Set( "LevelEditor.Tabs.EditorModes",            new IMAGE_BRUSH_SVG("Starship/LevelEditor/Menus/EditorModes", Icon16x16) );
+			Set( "LevelEditor.Tabs.Modes",                  new IMAGE_BRUSH_SVG("Starship/LevelEditor/Menus/EditorModes", Icon16x16) );
+			Set( "LevelEditor.Tabs.PlacementBrowser",       new IMAGE_BRUSH_SVG("Starship/LevelEditor/Menus/PlaceActors", Icon16x16) );
+			Set( "LevelEditor.Tabs.Properties",             new IMAGE_BRUSH_SVG("Starship/LevelEditor/Menus/Properties", Icon16x16) );
+			Set( "LevelEditor.Tabs.Outliner",               new IMAGE_BRUSH_SVG("Starship/LevelEditor/Menus/WorldOutliner", Icon16x16) );
+			Set( "LevelEditor.Tabs.ContentBrowser",         new IMAGE_BRUSH_SVG("Starship/LevelEditor/Menus/ContentBrowser", Icon16x16) );
+			Set( "LevelEditor.Tabs.Levels",                 new IMAGE_BRUSH_SVG("Starship/LevelEditor/Menus/Levels", Icon16x16) );
+			Set( "LevelEditor.Tabs.WorldBrowser",           new IMAGE_BRUSH_SVG("Starship/LevelEditor/Menus/Levels", Icon16x16) );
+			Set( "LevelEditor.Tabs.WorldBrowserDetails",    new IMAGE_BRUSH_SVG("Starship/LevelEditor/Menus/WorldSettings", Icon16x16) );
 
 			Set( "LevelEditor.Tabs.WorldBrowserComposition",new IMAGE_BRUSH("Icons/icon_levels_compositionbutton_16x", Icon16x16 ) );
 
-			Set( "LevelEditor.Tabs.Layers",                 new IMAGE_BRUSH_SVG("Starship/Menus/LevelEditor/Layers", Icon16x16) );
-			Set( "LevelEditor.Tabs.ComposureCompositing",   new IMAGE_BRUSH_SVG("Starship/Menus/LevelEditor/ComposureCompositing", Icon16x16) );
-			Set( "LevelEditor.Tabs.USDStage",   			new IMAGE_BRUSH_SVG("Starship/Menus/LevelEditor/USDStage", Icon16x16) );
+			Set( "LevelEditor.Tabs.Layers",                 new IMAGE_BRUSH_SVG("Starship/LevelEditor/Menus/Layers", Icon16x16) );
+			Set( "LevelEditor.Tabs.ComposureCompositing",   new IMAGE_BRUSH_SVG("Starship/LevelEditor/Menus/ComposureCompositing", Icon16x16) );
+			Set( "LevelEditor.Tabs.USDStage",   			new IMAGE_BRUSH_SVG("Starship/LevelEditor/Menus/USDStage", Icon16x16) );
 
 			Set( "LevelEditor.Tabs.BuildAndSubmit",         new IMAGE_BRUSH("Icons/icon_tab_BuildSubmit_16x", Icon16x16 ) );
 
-			Set( "LevelEditor.Tabs.StatsViewer",            new IMAGE_BRUSH_SVG("Starship/Menus/LevelEditor/Statistics", Icon16x16) );
+			Set( "LevelEditor.Tabs.StatsViewer",            new IMAGE_BRUSH_SVG("Starship/LevelEditor/Menus/Statistics", Icon16x16) );
 
 			Set( "LevelEditor.Tabs.Toolbar",                new IMAGE_BRUSH("Icons/icon_tab_Toolbars_16x", Icon16x16 ) );
 
-			Set( "LevelEditor.Tabs.Viewports",              new IMAGE_BRUSH_SVG("Starship/Menus/LevelEditor/Viewports", Icon16x16) );
-			Set( "LevelEditor.Tabs.HLOD",                   new IMAGE_BRUSH_SVG("Starship/Menus/LevelEditor/HierarchicalLODOutliner", Icon16x16) );
+			Set( "LevelEditor.Tabs.Viewports",              new IMAGE_BRUSH_SVG("Starship/LevelEditor/Menus/Viewports", Icon16x16) );
+			Set( "LevelEditor.Tabs.HLOD",                   new IMAGE_BRUSH_SVG("Starship/LevelEditor/Menus/HierarchicalLODOutliner", Icon16x16) );
 		}
 #endif // WITH_EDITOR || (IS_PROGRAM && WITH_UNREAL_DEVELOPER_TOOLS)
 
@@ -4432,7 +4432,7 @@ void FStarshipEditorStyle::FStyle::SetupLevelEditorStyle()
 		Set( "LevelEditor.WorldProperties", new IMAGE_BRUSH( "Icons/icon_worldscript_40x", Icon40x40 ) );
 		Set( "LevelEditor.WorldProperties.Small", new IMAGE_BRUSH( "Icons/icon_worldscript_40x", Icon20x20 ) );
 
-		Set( "LevelEditor.WorldProperties.Tab", new IMAGE_BRUSH_SVG( "Starship/Menus/LevelEditor/WorldSettings", Icon16x16 ) );
+		Set( "LevelEditor.WorldProperties.Tab", new IMAGE_BRUSH_SVG( "Starship/LevelEditor/Menus/WorldSettings", Icon16x16 ) );
 
 		Set( "LevelEditor.OpenContentBrowser", new IMAGE_BRUSH_SVG( "Starship/MainToolBar/content", Icon20x20 ) );
 		Set( "LevelEditor.OpenMarketplace", new IMAGE_BRUSH_SVG( "Starship/MainToolBar/marketplace", Icon20x20) );
@@ -4594,8 +4594,8 @@ void FStarshipEditorStyle::FStyle::SetupLevelEditorStyle()
 		Set( "LevelViewport.PilotSelectedActor", new IMAGE_BRUSH( "Icons/icon_PilotSelectedActor_16x", Icon16x16 ) );
 		Set( "LevelViewport.PilotSelectedActor.Small", new IMAGE_BRUSH( "Icons/icon_PilotSelectedActor_16x", Icon16x16 ) );
 		
-		Set( "LevelViewport.ToggleActorPilotCameraView",       new IMAGE_BRUSH_SVG( "Starship/Menus/LevelEditor/CameraShake", Icon16x16 ) );
-		Set( "LevelViewport.ToggleActorPilotCameraView.Small", new IMAGE_BRUSH_SVG( "Starship/Menus/LevelEditor/CameraShake", Icon16x16 ) );
+		Set( "LevelViewport.ToggleActorPilotCameraView",       new IMAGE_BRUSH_SVG( "Starship/LevelEditor/Menus/CameraShake", Icon16x16 ) );
+		Set( "LevelViewport.ToggleActorPilotCameraView.Small", new IMAGE_BRUSH_SVG( "Starship/LevelEditor/Menus/CameraShake", Icon16x16 ) );
 
 		Set( "LevelViewport.ActorPilotText", FTextBlockStyle()
 			.SetFont( DEFAULT_FONT( "BoldCondensed", 12 ) )
@@ -5012,7 +5012,7 @@ void FStarshipEditorStyle::FStyle::SetupPersonaStyle()
 
 		Set( "BlueprintEditor.ActionMenu.ContextDescriptionFont",  DEFAULT_FONT("Regular", 12) );
 
-		Set( "BlueprintEditor.FindInBlueprints.MenuIcon", new IMAGE_BRUSH_SVG("Starship/Menus/LevelEditor/FindInBlueprints", Icon16x16) );
+		Set( "BlueprintEditor.FindInBlueprints.MenuIcon", new IMAGE_BRUSH_SVG("Starship/LevelEditor/Menus/FindInBlueprints", Icon16x16) );
 
 		Set( "BlueprintEditor.FindInBlueprint", new IMAGE_BRUSH( "Icons/icon_Blueprint_Find_40px", Icon40x40 ) );
 		Set( "BlueprintEditor.FindInBlueprint.Small", new IMAGE_BRUSH( "Icons/icon_Blueprint_Find_40px", Icon20x20 ) );
@@ -5907,7 +5907,7 @@ void FStarshipEditorStyle::FStyle::SetupContentBrowserStyle()
 	// Content Browser
 	{
 		// Tab and menu icon
-		Set( "ContentBrowser.TabIcon", new IMAGE_BRUSH_SVG("Starship/Menus/LevelEditor/ContentBrowser", Icon16x16) );
+		Set( "ContentBrowser.TabIcon", new IMAGE_BRUSH_SVG("Starship/LevelEditor/Menus/ContentBrowser", Icon16x16) );
 
 		// Sources View
 		Set( "ContentBrowser.SourceTitleFont", DEFAULT_FONT( "Regular", 12 ) );
@@ -7698,7 +7698,7 @@ void FStarshipEditorStyle::FStyle::SetupTranslationEditorStyles()
 
 void FStarshipEditorStyle::FStyle::SetupLocalizationDashboardStyles()
 {
-	Set("LocalizationDashboard.MenuIcon", new IMAGE_BRUSH_SVG("Starship/Menus/LevelEditor/LocalizationDashboard", Icon16x16));
+	Set("LocalizationDashboard.MenuIcon", new IMAGE_BRUSH_SVG("Starship/LevelEditor/Menus/LocalizationDashboard", Icon16x16));
 
 	Set("LocalizationDashboard.GatherTextAllTargets", new IMAGE_BRUSH("Icons/Icon_Localisation_Gather_All_40x", Icon40x40));
 	Set("LocalizationDashboard.ImportTextAllTargetsAllCultures", new IMAGE_BRUSH("Icons/Icon_Localisation_Import_All_40x", Icon40x40));
diff --git a/Engine/Source/Runtime/SlateCore/Private/Styling/StarshipCoreStyle.cpp b/Engine/Source/Runtime/SlateCore/Private/Styling/StarshipCoreStyle.cpp
--- a/Engine/Source/Runtime/SlateCore/Private/Styling/StarshipCoreStyle.cpp
+++ b/Engine/Source/Runtime/SlateCore/Private/Styling/StarshipCoreStyle.cpp
@@ -348,6 +348,8 @@ TSharedRef<ISlateStyle> FStarshipCoreStyle::Create()
 
 		Style->Set("Icons.CircleArrowUp", new IMAGE_BRUSH_SVG("Starship/Common/circle-arrow-up", Icon16x16));
 		Style->Set("Icons.CircleArrowDown", new IMAGE_BRUSH_SVG("Starship/Common/circle-arrow-down", Icon16x16));
+
+		Style->Set("Icons.Check", new IMAGE_BRUSH_SVG("Starship/Common/check", Icon16x16));
 	}
 
 	// Tool panels
@@ -1782,9 +1784,9 @@ void FStarshipCoreStyle::SetupMultiboxStyles(TSharedRef<FStyle>& Style)
 			.SetUncheckedImage(FSlateNoResource(Icon14x14))
 			.SetUncheckedHoveredImage(FSlateNoResource(Icon14x14))
 			.SetUncheckedPressedImage(FSlateNoResource(Icon14x14))
-			.SetCheckedImage(IMAGE_BRUSH("Starship/Icons/Quixel/Common/Check_32x", Icon16x16))
-			.SetCheckedHoveredImage(IMAGE_BRUSH("Starship/Icons/Quixel/Common/Check_32x", Icon16x16))
-			.SetCheckedPressedImage(IMAGE_BRUSH("Starship/Icons/Quixel/Common/Check_32x", Icon16x16))
+			.SetCheckedImage(	    IMAGE_BRUSH_SVG("Starship/Common/check", Icon16x16))
+			.SetCheckedHoveredImage(IMAGE_BRUSH_SVG("Starship/Common/check", Icon16x16))
+			.SetCheckedPressedImage(IMAGE_BRUSH_SVG("Starship/Common/check", Icon16x16))
 			.SetUndeterminedImage(FSlateNoResource())
 			.SetUndeterminedHoveredImage(FSlateNoResource())
 			.SetUndeterminedPressedImage(FSlateNoResource());
