Updated Editor Viewport Icons

diff --git a/Engine/Source/Editor/EditorStyle/Private/StarshipStyle.cpp b/Engine/Source/Editor/EditorStyle/Private/StarshipStyle.cpp
--- a/Engine/Source/Editor/EditorStyle/Private/StarshipStyle.cpp
+++ b/Engine/Source/Editor/EditorStyle/Private/StarshipStyle.cpp
@@ -2448,8 +2448,8 @@ void FStarshipEditorStyle::FStyle::SetupViewportStyles()
 
 		
 		Set("EditorViewportToolBar.MenuDropdown", new IMAGE_BRUSH("Common/ComboArrow", Icon8x8));
-		Set("EditorViewportToolBar.Maximize.Normal", new IMAGE_BRUSH_SVG("Starship/EditorViewport/maximize", Icon16x16));
-		Set("EditorViewportToolBar.Maximize.Checked", new IMAGE_BRUSH_SVG("Starship/EditorViewport/restore", Icon16x16));
+		Set("EditorViewportToolBar.Maximize.Normal", new IMAGE_BRUSH_SVG("Starship/EditorViewport/square", Icon16x16));
+		Set("EditorViewportToolBar.Maximize.Checked", new IMAGE_BRUSH_SVG("Starship/EditorViewport/quad", Icon16x16));
 		Set("EditorViewportToolBar.RestoreFromImmersive.Normal", new IMAGE_BRUSH("Icons/icon_RestoreFromImmersive_16px", Icon16x16));
 	}
 
@@ -4240,17 +4240,17 @@ void FStarshipEditorStyle::FStyle::SetupLevelEditorStyle()
 
 		Set( "EditorViewport.ToggleRealTime", new IMAGE_BRUSH("Icons/icon_MatEd_Realtime_40x", Icon40x40));
 
-		Set( "EditorViewport.LocationGridSnap", new IMAGE_BRUSH_SVG("Starship/EditorViewport/location-grid-snap", Icon16x16));
-		Set( "EditorViewport.RotationGridSnap", new IMAGE_BRUSH_SVG("Starship/EditorViewport/rotation-grid-snap", Icon16x16));
+		Set( "EditorViewport.LocationGridSnap", new IMAGE_BRUSH_SVG("Starship/EditorViewport/grid", Icon16x16));
+		Set( "EditorViewport.RotationGridSnap", new IMAGE_BRUSH_SVG("Starship/EditorViewport/angle", Icon16x16));
 
 		Set( "EditorViewport.Layer2DSnap", new IMAGE_BRUSH("Old/LevelEditor/Layer2DSnap", Icon14x14));
 
-		Set("EditorViewport.ScaleGridSnap", new IMAGE_BRUSH( "Old/LevelEditor/ScaleGridSnap", Icon14x14 ) );
+		Set("EditorViewport.ScaleGridSnap", new IMAGE_BRUSH( "Starship/EditorViewport/scale-grid-snap", Icon16x16 ) );
 		Set("EditorViewport.ToggleSurfaceSnapping", new IMAGE_BRUSH_SVG( "Starship/EditorViewport/surface-snap", Icon16x16 ) );
 		Set("EditorViewport.ToggleSurfaceSnapping", new IMAGE_BRUSH_SVG("Starship/EditorViewport/surface-snap", Icon16x16));
 
 		Set( "EditorViewport.RelativeCoordinateSystem_Local", new IMAGE_BRUSH_SVG( "Starship/EditorViewport/transform-local", Icon16x16 ) );
-		Set( "EditorViewport.RelativeCoordinateSystem_World", new IMAGE_BRUSH_SVG( "Starship/EditorViewport/transform-world", Icon16x16 ) );
+		Set( "EditorViewport.RelativeCoordinateSystem_World", new IMAGE_BRUSH_SVG( "Starship/EditorViewport/globe", Icon16x16 ) );
 
 		Set( "EditorViewport.CamSpeedSetting", new IMAGE_BRUSH_SVG( "Starship/EditorViewport/camera", Icon16x16) );
 		
