Fixing missing Editor Viewport Icon

diff --git a/Engine/Source/Editor/EditorStyle/Private/StarshipStyle.cpp b/Engine/Source/Editor/EditorStyle/Private/StarshipStyle.cpp
--- a/Engine/Source/Editor/EditorStyle/Private/StarshipStyle.cpp
+++ b/Engine/Source/Editor/EditorStyle/Private/StarshipStyle.cpp
@@ -4245,7 +4245,7 @@ void FStarshipEditorStyle::FStyle::SetupLevelEditorStyle()
 
 		Set( "EditorViewport.Layer2DSnap", new IMAGE_BRUSH("Old/LevelEditor/Layer2DSnap", Icon14x14));
 
-		Set("EditorViewport.ScaleGridSnap", new IMAGE_BRUSH( "Starship/EditorViewport/scale-grid-snap", Icon16x16 ) );
+		Set("EditorViewport.ScaleGridSnap", new IMAGE_BRUSH_SVG( "Starship/EditorViewport/scale-grid-snap", Icon16x16 ) );
 		Set("EditorViewport.ToggleSurfaceSnapping", new IMAGE_BRUSH_SVG( "Starship/EditorViewport/surface-snap", Icon16x16 ) );
 		Set("EditorViewport.ToggleSurfaceSnapping", new IMAGE_BRUSH_SVG("Starship/EditorViewport/surface-snap", Icon16x16));
 
