Updated Main Toolbar Icons

diff --git a/Engine/Plugins/Experimental/ControlRig/Source/ControlRigEditor/Private/ControlRigEditorStyle.h b/Engine/Plugins/Experimental/ControlRig/Source/ControlRigEditor/Private/ControlRigEditorStyle.h
--- a/Engine/Plugins/Experimental/ControlRig/Source/ControlRigEditor/Private/ControlRigEditorStyle.h
+++ b/Engine/Plugins/Experimental/ControlRig/Source/ControlRigEditor/Private/ControlRigEditorStyle.h
@@ -4,10 +4,8 @@
 
 #include "Styling/SlateStyle.h"
 #include "Styling/SlateStyleRegistry.h"
+#include "Styling/SlateStyleMacros.h"
 
-#define IMAGE_BRUSH(RelativePath, ...) FSlateImageBrush(RootToContentDir(RelativePath, TEXT(".png")), __VA_ARGS__)
-#define BOX_BRUSH(RelativePath, ...) FSlateBoxBrush(RootToContentDir(RelativePath, TEXT(".png")), __VA_ARGS__)
-#define BORDER_BRUSH(RelativePath, ...) FSlateBorderBrush(RootToContentDir(RelativePath, TEXT(".png")), __VA_ARGS__)
 #define TTF_FONT(RelativePath, ...) FSlateFontInfo(RootToContentDir(RelativePath, TEXT(".ttf")), __VA_ARGS__)
 #define OTF_FONT(RelativePath, ...) FSlateFontInfo(RootToContentDir(RelativePath, TEXT(".otf")), __VA_ARGS__)
 
@@ -36,6 +34,13 @@ public:
 		}
 
 		// Edit mode styles
+
+		if (FCoreStyle::IsStarshipStyle())
+		{
+			Set("ControlRigEditMode", new IMAGE_BRUSH_SVG("Slate/animation", Icon40x40));
+			Set("ControlRigEditMode.Small", new IMAGE_BRUSH_SVG("Slate/animation", Icon20x20));
+		}
+		else
 		{
 			Set("ControlRigEditMode", new IMAGE_BRUSH("Slate/ControlRigMode_40x", Icon40x40));
 			Set("ControlRigEditMode.Small", new IMAGE_BRUSH("Slate/ControlRigMode_40x", Icon20x20));
diff --git a/Engine/Source/Editor/EditorStyle/Private/StarshipStyle.cpp b/Engine/Source/Editor/EditorStyle/Private/StarshipStyle.cpp
--- a/Engine/Source/Editor/EditorStyle/Private/StarshipStyle.cpp
+++ b/Engine/Source/Editor/EditorStyle/Private/StarshipStyle.cpp
@@ -1264,7 +1264,7 @@ void FStarshipEditorStyle::FStyle::SetupGeneralStyles()
 	// Asset editors (common)
 	{
 		Set( "AssetEditor.SaveAsset.Greyscale", new IMAGE_BRUSH( "Icons/icon_file_save_16px", Icon16x16 ) );
-		Set( "AssetEditor.SaveAsset", new IMAGE_BRUSH_SVG( "Starship/MainToolBar/save", Icon40x40 ) );
+		Set( "AssetEditor.SaveAsset", new IMAGE_BRUSH_SVG( "Starship/MainToolbar/save", Icon40x40 ) );
 		Set( "AssetEditor.SaveAssetAs", new IMAGE_BRUSH( "Icons/icon_file_saveas_40x", Icon40x40 ) );
 		Set( "AssetEditor.SaveAssetAs.Small", new IMAGE_BRUSH( "Icons/icon_file_saveas_40x", Icon20x20 ) );
 		Set( "AssetEditor.ReimportAsset", new IMAGE_BRUSH( "Icons/icon_TextureEd_Reimport_40x", Icon40x40 ) );
@@ -4343,8 +4343,8 @@ void FStarshipEditorStyle::FStyle::SetupLevelEditorStyle()
 		Set( "LevelEditor.MapCheck", new IMAGE_BRUSH( "Icons/icon_MapCheck_40x", Icon40x40 ) );
 		Set( "LevelEditor.MapCheck.Small", new IMAGE_BRUSH( "Icons/icon_MapCheck_40x", Icon20x20 ) );
 
-		Set( "LevelEditor.Recompile", new IMAGE_BRUSH( "Icons/icon_compile_40x", Icon40x40 ) );
-		Set( "LevelEditor.Recompile.Small", new IMAGE_BRUSH( "Icons/icon_compile_40x", Icon20x20 ) );
+		Set( "LevelEditor.Recompile", new IMAGE_BRUSH_SVG( "Starship/MainToolbar/compile", Icon40x40 ) );
+		Set( "LevelEditor.Recompile.Small", new IMAGE_BRUSH_SVG( "Starship/MainToolbar/compile", Icon20x20 ) );
 
 		Set("LevelEditor.PreviewMode.Enabled", new IMAGE_BRUSH("Icons/icon_PreviewMode_SM5_Enabled_40x", Icon40x40));
 		Set("LevelEditor.PreviewMode.Enabled.Small", new IMAGE_BRUSH("Icons/icon_PreviewMode_SM5_Enabled_40x", Icon20x20));
@@ -4378,7 +4378,7 @@ void FStarshipEditorStyle::FStyle::SetupLevelEditorStyle()
 		Set("LevelEditor.ViewOptions", new IMAGE_BRUSH("Icons/icon_view_40x", Icon40x40));
 		Set( "LevelEditor.ViewOptions.Small", new IMAGE_BRUSH( "Icons/icon_view_40x", Icon20x20 ) );
 
-		Set( "LevelEditor.GameSettings", new IMAGE_BRUSH_SVG( "Starship/MainToolBar/settings", Icon20x20) );
+		Set( "LevelEditor.GameSettings", new IMAGE_BRUSH_SVG( "Starship/MainToolbar/settings", Icon20x20) );
 
 		Set( "LevelEditor.Create", new IMAGE_BRUSH( "Icons/icon_Mode_Placement_40px", Icon40x40 ) );
 		Set( "LevelEditor.Create.Small", new IMAGE_BRUSH( "Icons/icon_Mode_Placement_40px", Icon20x20 ) );
@@ -4407,9 +4407,9 @@ void FStarshipEditorStyle::FStyle::SetupLevelEditorStyle()
 		);
 
 
-		Set("LevelEditor.SelectMode", new IMAGE_BRUSH_SVG("Starship/MainToolBar/select", Icon20x20));
+		Set("LevelEditor.SelectMode", new IMAGE_BRUSH_SVG("Starship/MainToolbar/select", Icon20x20));
 
-		Set( "LevelEditor.MeshPaintMode", new IMAGE_BRUSH_SVG( "Starship/MainToolBar/paint", Icon20x20 ) );
+		Set( "LevelEditor.MeshPaintMode", new IMAGE_BRUSH_SVG( "Starship/MainToolbar/paint", Icon20x20 ) );
 		
 		Set("LevelEditor.MeshPaintMode.TexturePaint", new IMAGE_BRUSH("Icons/TexturePaint_40x", Icon40x40));
 		Set("LevelEditor.MeshPaintMode.TexturePaint.Small", new IMAGE_BRUSH("Icons/TexturePaint_40x", Icon20x20));
@@ -4418,11 +4418,11 @@ void FStarshipEditorStyle::FStyle::SetupLevelEditorStyle()
 		Set("LevelEditor.MeshPaintMode.WeightPaint", new IMAGE_BRUSH("Icons/WeightPaint_40x", Icon40x40));
 		Set("LevelEditor.MeshPaintMode.WeightPaint.Small", new IMAGE_BRUSH("Icons/WeightPaint_40x", Icon20x20));
 
-		Set( "LevelEditor.LandscapeMode", new IMAGE_BRUSH_SVG( "Starship/MainToolBar/landscape", Icon20x20 ) );
+		Set( "LevelEditor.LandscapeMode", new IMAGE_BRUSH_SVG( "Starship/MainToolbar/landscape", Icon20x20 ) );
 		Set( "LevelEditor.LandscapeMode.Selected", new IMAGE_BRUSH( "Icons/icon_Mode_Landscape_selected_40x", Icon40x40 ) );
 		Set( "LevelEditor.LandscapeMode.Selected.Small", new IMAGE_BRUSH( "Icons/icon_Mode_Landscape_selected_40x", Icon20x20 ) );
 
-		Set( "LevelEditor.FoliageMode", new IMAGE_BRUSH_SVG( "Starship/MainToolBar/foliage", Icon20x20 ) );
+		Set( "LevelEditor.FoliageMode", new IMAGE_BRUSH_SVG( "Starship/MainToolbar/foliage", Icon20x20 ) );
 		Set( "LevelEditor.FoliageMode.Selected", new IMAGE_BRUSH( "Icons/icon_Mode_Foliage_selected_40x", Icon40x40 ) );
 		Set( "LevelEditor.FoliageMode.Selected.Small", new IMAGE_BRUSH( "Icons/icon_Mode_Foliage_selected_40x", Icon20x20 ) );
 
@@ -4431,13 +4431,13 @@ void FStarshipEditorStyle::FStyle::SetupLevelEditorStyle()
 
 		Set( "LevelEditor.WorldProperties.Tab", new IMAGE_BRUSH_SVG( "Starship/LevelEditor/Menus/WorldSettings", Icon16x16 ) );
 
-		Set( "LevelEditor.OpenContentBrowser", new IMAGE_BRUSH_SVG( "Starship/MainToolBar/content", Icon20x20 ) );
-		Set( "LevelEditor.OpenMarketplace", new IMAGE_BRUSH_SVG( "Starship/MainToolBar/marketplace", Icon20x20) );
-		Set( "LevelEditor.OpenLevelBlueprint", new IMAGE_BRUSH_SVG( "Starship/MainToolBar/blueprints", Icon20x20) );
+		Set( "LevelEditor.OpenContentBrowser", new IMAGE_BRUSH_SVG( "Starship/MainToolbar/content", Icon20x20 ) );
+		Set( "LevelEditor.OpenMarketplace", new IMAGE_BRUSH_SVG( "Starship/MainToolbar/marketplace", Icon20x20) );
+		Set( "LevelEditor.OpenLevelBlueprint", new IMAGE_BRUSH_SVG( "Starship/MainToolbar/blueprints", Icon20x20) );
 
 		Set( "LevelEditor.CreateClassBlueprint", new IMAGE_BRUSH("Icons/icon_class_Blueprint_New_16x", Icon16x16));
 		Set( "LevelEditor.OpenClassBlueprint", new IMAGE_BRUSH("Icons/icon_class_Blueprint_Open_16x", Icon16x16));
-		Set( "LevelEditor.EditMatinee", new IMAGE_BRUSH_SVG( "Starship/MainToolBar/cinematics", Icon20x20 ) );
+		Set( "LevelEditor.EditMatinee", new IMAGE_BRUSH_SVG( "Starship/MainToolbar/cinematics", Icon20x20 ) );
 
 		Set( "LevelEditor.ToggleVR", new IMAGE_BRUSH( "Icons/VREditor/VR_Editor_Toolbar_Icon", Icon40x40 ) );
 		Set( "LevelEditor.ToggleVR.Small", new IMAGE_BRUSH( "Icons/VREditor/VR_Editor_Toolbar_Icon_Small", Icon20x20 ) );
@@ -5277,7 +5277,7 @@ void FStarshipEditorStyle::FStyle::SetupPersonaStyle()
 
 		Set( "PlayWorld.RepeatLastPlay", new IMAGE_BRUSH( "Icons/icon_simulate_40x", Icon40x40 ) );
 		Set( "PlayWorld.RepeatLastPlay.Small", new IMAGE_BRUSH( "Icons/icon_simulate_40x", Icon20x20 ) );
-		Set( "PlayWorld.PlayInViewport", new IMAGE_BRUSH_SVG("Starship/MainToolBar/play", Icon20x20));
+		Set( "PlayWorld.PlayInViewport", new IMAGE_BRUSH_SVG("Starship/MainToolbar/play", Icon20x20));
 
 		Set( "PlayWorld.PlayInEditorFloating", new IMAGE_BRUSH( "Icons/icon_playInWindow_40x", Icon40x40 ) );
 		Set( "PlayWorld.PlayInEditorFloating.Small", new IMAGE_BRUSH( "Icons/icon_playInWindow_40x", Icon20x20 ) );
@@ -5289,8 +5289,10 @@ void FStarshipEditorStyle::FStyle::SetupPersonaStyle()
 		Set( "PlayWorld.PlayInVulkanPreview.Small", new IMAGE_BRUSH( "Icons/icon_PlayMobilePreview_16x", Icon20x20 ) );
 		Set( "PlayWorld.PlayInNewProcess", new IMAGE_BRUSH( "Icons/icon_PlayStandalone_40x", Icon40x40 ) );
 		Set( "PlayWorld.PlayInNewProcess.Small", new IMAGE_BRUSH( "Icons/icon_PlayStandalone_40x", Icon20x20 ) );
-		Set( "PlayWorld.RepeatLastLaunch", new IMAGE_BRUSH( "Icons/icon_PlayOnDevice_40px", Icon40x40 ) );
-		Set( "PlayWorld.RepeatLastLaunch.Small", new IMAGE_BRUSH( "Icons/icon_PlayOnDevice_40px", Icon20x20 ) );
+
+		Set( "PlayWorld.RepeatLastLaunch", new IMAGE_BRUSH_SVG( "Starship/MainToolbar/launch", Icon40x40 ) );
+		Set( "PlayWorld.RepeatLastLaunch.Small", new IMAGE_BRUSH_SVG( "Starship/MainToolbar/launch", Icon20x20 ) );
+
 		Set( "PlayWorld.PlayInCameraLocation", new IMAGE_BRUSH( "Icons/icon_PlayCameraLocation_40x", Icon40x40 ) );
 		Set( "PlayWorld.PlayInDefaultPlayerStart", new IMAGE_BRUSH( "Icons/icon_PlayDefaultPlayerStart_40x", Icon40x40 ) );
 
