Icon Reorganization

diff --git a/Engine/Source/Editor/EditorStyle/Private/StarshipStyle.cpp b/Engine/Source/Editor/EditorStyle/Private/StarshipStyle.cpp
--- a/Engine/Source/Editor/EditorStyle/Private/StarshipStyle.cpp
+++ b/Engine/Source/Editor/EditorStyle/Private/StarshipStyle.cpp
@@ -999,7 +999,7 @@ void FStarshipEditorStyle::FStyle::SetupGeneralStyles()
 
 		Set("DebugConsole.ComboButton", DebugConsoleComboButton);
 
-		Set("DebugConsole.Icon", new IMAGE_BRUSH_SVG("Starship/Menus/Help/Console", Icon16x16));
+		Set("DebugConsole.Icon", new IMAGE_BRUSH_SVG("Starship/Common/Console", Icon16x16));
 
 		Set("OutputLog.OpenSourceLocation", new IMAGE_BRUSH("Icons/icon_Asset_Open_Source_Location_16x", Icon16x16));
 		Set("OutputLog.OpenInExternalEditor", new IMAGE_BRUSH("Icons/icon_Asset_Open_In_External_Editor_16x", Icon16x16));
@@ -1038,13 +1038,13 @@ void FStarshipEditorStyle::FStyle::SetupGeneralStyles()
 
 	// Developer Tools Menu
 	{
-		Set("DeveloperTools.MenuIcon", new IMAGE_BRUSH_SVG( "Starship/LevelEditor/Menus/DeveloperTools", Icon16x16 ) );
+		Set("DeveloperTools.MenuIcon", new IMAGE_BRUSH_SVG( "Starship/Common/DeveloperTools", Icon16x16 ) );
 	}
 
 	// Automation Tools Menu
 	{
-		Set("AutomationTools.MenuIcon", new IMAGE_BRUSH_SVG("Starship/LevelEditor/Menus/AutomationTools", Icon16x16));
-		Set("AutomationTools.TestAutomation", new IMAGE_BRUSH_SVG("Starship/LevelEditor/Menus/TestAutomation", Icon16x16));
+		Set("AutomationTools.MenuIcon", new IMAGE_BRUSH_SVG("Starship/Common/AutomationTools", Icon16x16));
+		Set("AutomationTools.TestAutomation", new IMAGE_BRUSH_SVG("Starship/Common/TestAutomation", Icon16x16));
 	}
 
 	// Session Browser tab
@@ -1080,13 +1080,13 @@ void FStarshipEditorStyle::FStyle::SetupGeneralStyles()
 
 	// Launcher Window
 	{
-		Set("Launcher.TabIcon", new IMAGE_BRUSH_SVG( "Starship/LevelEditor/Menus/ProjectLauncher", Icon16x16 ) );
+		Set("Launcher.TabIcon", new IMAGE_BRUSH_SVG( "Starship/Common/ProjectLauncher", Icon16x16 ) );
 		Set("Launcher.Tabs.Tools", new IMAGE_BRUSH( "/Icons/icon_tab_Tools_16x", Icon16x16 ) );
 	}
 
 	// Undo History Window
 	{
-		Set("UndoHistory.TabIcon", new IMAGE_BRUSH_SVG( "Starship/Menus/Edit/UndoHistory", Icon16x16 ) );
+		Set("UndoHistory.TabIcon", new CORE_IMAGE_BRUSH_SVG( "Starship/Common/UndoHistory", Icon16x16 ) );
 	}
 
 	// InputBinding editor
@@ -1289,7 +1289,7 @@ void FStarshipEditorStyle::FStyle::SetupGeneralStyles()
 
 	// Open any asset dialog
 	{
-		Set( "SystemWideCommands.SummonOpenAssetDialog", new IMAGE_BRUSH_SVG( "Starship/Menus/File/OpenAsset", Icon16x16 ) );
+		Set( "SystemWideCommands.SummonOpenAssetDialog", new IMAGE_BRUSH_SVG( "Starship/Common/OpenAsset", Icon16x16 ) );
 	
 		Set( "GlobalAssetPicker.Background", new BOX_BRUSH( "Old/Menu_Background", FMargin(8.0f/64.0f) ) );
 		Set( "GlobalAssetPicker.OutermostMargin", FMargin(4, 4, 4, 4) );
@@ -1307,22 +1307,22 @@ void FStarshipEditorStyle::FStyle::SetupGeneralStyles()
 	{
 		Set( "MainFrame.AutoSaveImage", new IMAGE_BRUSH( "Icons/icon_Autosave", Icon24x24 ) );
 
-		Set( "MainFrame.SaveAll",              new IMAGE_BRUSH_SVG( "Starship/Menus/File/SaveAll", Icon16x16 ) );
-		Set( "MainFrame.ChoosePackagesToSave", new IMAGE_BRUSH_SVG( "Starship/Menus/File/icon_file_choosepackages_16px", Icon16x16 ) );
-		Set( "MainFrame.NewProject",           new IMAGE_BRUSH_SVG( "Starship/Menus/File/ProjectNew", Icon16x16 ) );
-		Set( "MainFrame.OpenProject",          new IMAGE_BRUSH_SVG( "Starship/Menus/File/ProjectOpen", Icon16x16 ) );
-		Set( "MainFrame.AddCodeToProject",     new IMAGE_BRUSH_SVG( "Starship/Menus/File/ProjectC++", Icon16x16 ) );
-		Set( "MainFrame.Exit",                 new IMAGE_BRUSH_SVG( "Starship/Menus/File/Exit", Icon16x16 ) );
-		Set( "MainFrame.CookContent",          new IMAGE_BRUSH_SVG( "Starship/Menus/File/CookContent", Icon16x16 ) );
-		Set( "MainFrame.OpenVisualStudio",     new IMAGE_BRUSH_SVG( "Starship/Menus/File/VisualStudio", Icon16x16 ) );
-		Set( "MainFrame.RefreshVisualStudio",  new IMAGE_BRUSH_SVG( "Starship/Menus/File/RefreshVisualStudio", Icon16x16 ) );
-		Set( "MainFrame.PackageProject",       new IMAGE_BRUSH_SVG( "Starship/Menus/File/ProjectPackage", Icon16x16 ) );
-		Set( "MainFrame.RecentProjects",       new IMAGE_BRUSH_SVG( "Starship/Menus/File/ProjectsRecent", Icon16x16 ) );
-		Set( "MainFrame.RecentLevels",         new IMAGE_BRUSH_SVG( "Starship/Menus/File/LevelRecent", Icon16x16 ) );
-		Set( "MainFrame.FavoriteLevels",       new IMAGE_BRUSH_SVG( "Starship/Menus/File/LevelFavorite", Icon16x16 ) );
-
-		Set( "MainFrame.ChooseFilesToSave",       new IMAGE_BRUSH_SVG( "Starship/Menus/File/SaveChoose", Icon16x16 ) );
-		Set( "MainFrame.ConnectToSourceControl",  new IMAGE_BRUSH_SVG( "Starship/Menus/File/SourceControl", Icon16x16 ) );
+		Set( "MainFrame.SaveAll",              new IMAGE_BRUSH_SVG( "Starship/Common/SaveAll", Icon16x16 ) );
+		Set( "MainFrame.ChoosePackagesToSave", new IMAGE_BRUSH_SVG( "Starship/Common/icon_file_choosepackages_16px", Icon16x16 ) );
+		Set( "MainFrame.NewProject",           new IMAGE_BRUSH_SVG( "Starship/Common/ProjectNew", Icon16x16 ) );
+		Set( "MainFrame.OpenProject",          new IMAGE_BRUSH_SVG( "Starship/Common/ProjectOpen", Icon16x16 ) );
+		Set( "MainFrame.AddCodeToProject",     new IMAGE_BRUSH_SVG( "Starship/Common/ProjectC++", Icon16x16 ) );
+		Set( "MainFrame.Exit",                 new IMAGE_BRUSH_SVG( "Starship/Common/Exit", Icon16x16 ) );
+		Set( "MainFrame.CookContent",          new IMAGE_BRUSH_SVG( "Starship/Common/CookContent", Icon16x16 ) );
+		Set( "MainFrame.OpenVisualStudio",     new IMAGE_BRUSH_SVG( "Starship/Common/VisualStudio", Icon16x16 ) );
+		Set( "MainFrame.RefreshVisualStudio",  new IMAGE_BRUSH_SVG( "Starship/Common/RefreshVisualStudio", Icon16x16 ) );
+		Set( "MainFrame.PackageProject",       new IMAGE_BRUSH_SVG( "Starship/Common/ProjectPackage", Icon16x16 ) );
+		Set( "MainFrame.RecentProjects",       new IMAGE_BRUSH_SVG( "Starship/Common/ProjectsRecent", Icon16x16 ) );
+		Set( "MainFrame.RecentLevels",         new IMAGE_BRUSH_SVG( "Starship/Common/LevelRecent", Icon16x16 ) );
+		Set( "MainFrame.FavoriteLevels",       new IMAGE_BRUSH_SVG( "Starship/Common/LevelFavorite", Icon16x16 ) );
+
+		Set( "MainFrame.ChooseFilesToSave",       new IMAGE_BRUSH_SVG( "Starship/Common/SaveChoose", Icon16x16 ) );
+		Set( "MainFrame.ConnectToSourceControl",  new IMAGE_BRUSH_SVG( "Starship/Common/SourceControl", Icon16x16 ) );
 
 		Set( "MainFrame.DebugTools.SmallFont", DEFAULT_FONT( "Regular", 8 ) );
 		Set( "MainFrame.DebugTools.NormalFont", DEFAULT_FONT( "Regular", 9 ) );
@@ -1331,12 +1331,12 @@ void FStarshipEditorStyle::FStyle::SetupGeneralStyles()
 
 	// Editor preferences
 	{
-		Set("EditorPreferences.TabIcon", new IMAGE_BRUSH_SVG("Starship/Menus/Edit/Preferences", Icon16x16));
+		Set("EditorPreferences.TabIcon", new CORE_IMAGE_BRUSH_SVG("Starship/Common/Preferences", Icon16x16));
 	}
 
 	// Project settings
 	{
-		Set("ProjectSettings.TabIcon", new IMAGE_BRUSH_SVG("Starship/Menus/Edit/ProjectSettings", Icon16x16));
+		Set("ProjectSettings.TabIcon", new IMAGE_BRUSH_SVG("Starship/Common/ProjectSettings", Icon16x16));
 	}
 
 	// Main frame
@@ -2433,7 +2433,7 @@ void FStarshipEditorStyle::FStyle::SetupViewportStyles()
 		ToolbarBackgroundColor.A = .75f;
 
 		Set("EditorViewportToolBar.Background", new FSlateColorBrush(ToolbarBackgroundColor));
-		Set("EditorViewportToolBar.OptionsDropdown", new CORE_IMAGE_BRUSH_SVG("Starship/Common/hamburger", FVector2D(13,8)));
+		Set("EditorViewportToolBar.OptionsDropdown", new CORE_IMAGE_BRUSH_SVG("Starship/EditorViewport/menu", FVector2D(13,8)));
 
 		Set("EditorViewportToolBar.Font", FStyleFonts::Get().Normal);
 
@@ -4205,28 +4205,28 @@ void FStarshipEditorStyle::FStyle::SetupLevelEditorStyle()
 	// Level editor tool bar icons
 #if WITH_EDITOR
 	{
-		Set("LevelEditor.BrowseDocumentation",     new IMAGE_BRUSH_SVG("Starship/Menus/Help/Documentation", Icon16x16));
-		Set("LevelEditor.BrowseAPIReference",      new IMAGE_BRUSH_SVG("Starship/Menus/Help/Documentation", Icon16x16));
-		Set("LevelEditor.BrowseCVars",             new IMAGE_BRUSH_SVG("Starship/Menus/Help/Console", Icon16x16));
-		Set("LevelEditor.Tutorials",               new IMAGE_BRUSH_SVG("Starship/Menus/Help/Tutorials", Icon16x16));
-		Set("LevelEditor.BrowseViewportControls",  new IMAGE_BRUSH_SVG("Starship/Menus/Help/ViewportControls", Icon16x16));
+		Set("LevelEditor.BrowseDocumentation",     new IMAGE_BRUSH_SVG("Starship/Common/Documentation", Icon16x16));
+		Set("LevelEditor.BrowseAPIReference",      new IMAGE_BRUSH_SVG("Starship/Common/Documentation", Icon16x16));
+		Set("LevelEditor.BrowseCVars",             new IMAGE_BRUSH_SVG("Starship/Common/Console", Icon16x16));
+		Set("LevelEditor.Tutorials",               new IMAGE_BRUSH_SVG("Starship/Common/Tutorials", Icon16x16));
+		Set("LevelEditor.BrowseViewportControls",  new IMAGE_BRUSH_SVG("Starship/Common/ViewportControls", Icon16x16));
 
-		Set("MainFrame.ToggleFullscreen",          new IMAGE_BRUSH_SVG("Starship/LevelEditor/Menus/EnableFullscreen", Icon16x16));
-		Set("MainFrame.LoadLayout",                new IMAGE_BRUSH_SVG("Starship/LevelEditor/Menus/LayoutLoad", Icon16x16));
-		Set("MainFrame.SaveLayout",                new IMAGE_BRUSH_SVG("Starship/LevelEditor/Menus/LayoutSave", Icon16x16));
-		Set("MainFrame.RemoveLayout",              new IMAGE_BRUSH_SVG("Starship/LevelEditor/Menus/LayoutRemove", Icon16x16));
+		Set("MainFrame.ToggleFullscreen",          new IMAGE_BRUSH_SVG("Starship/Common/EnableFullscreen", Icon16x16));
+		Set("MainFrame.LoadLayout",                new IMAGE_BRUSH_SVG("Starship/Common/LayoutLoad", Icon16x16));
+		Set("MainFrame.SaveLayout",                new IMAGE_BRUSH_SVG("Starship/Common/LayoutSave", Icon16x16));
+		Set("MainFrame.RemoveLayout",              new IMAGE_BRUSH_SVG("Starship/Common/LayoutRemove", Icon16x16));
 
-		Set("MainFrame.OpenIssueTracker",          new IMAGE_BRUSH_SVG("Starship/Menus/Help/IssueTracker", Icon16x16));
-		Set("MainFrame.ReportABug",                new IMAGE_BRUSH_SVG("Starship/Menus/Help/Bug", Icon16x16));
+		Set("MainFrame.OpenIssueTracker",          new IMAGE_BRUSH_SVG("Starship/Common/IssueTracker", Icon16x16));
+		Set("MainFrame.ReportABug",                new IMAGE_BRUSH_SVG("Starship/Common/Bug", Icon16x16));
 
 		Set("MainFrame.VisitAskAQuestionPage",     new IMAGE_BRUSH_SVG("Starship/Menus/Bug", Icon16x16));
-		Set("MainFrame.VisitWiki",                 new IMAGE_BRUSH_SVG("Starship/Menus/Help/WebBrowser", Icon16x16));
-		Set("MainFrame.VisitForums",               new IMAGE_BRUSH_SVG("Starship/Menus/Help/WebBrowser", Icon16x16));
-		Set("MainFrame.VisitSearchForAnswersPage", new IMAGE_BRUSH_SVG("Starship/Menus/Help/WebBrowser", Icon16x16));
-		Set("MainFrame.VisitSupportWebSite",       new IMAGE_BRUSH_SVG("Starship/Menus/Help/Support", Icon16x16));
+		Set("MainFrame.VisitWiki",                 new IMAGE_BRUSH_SVG("Starship/Common/WebBrowser", Icon16x16));
+		Set("MainFrame.VisitForums",               new IMAGE_BRUSH_SVG("Starship/Common/WebBrowser", Icon16x16));
+		Set("MainFrame.VisitSearchForAnswersPage", new IMAGE_BRUSH_SVG("Starship/Common/WebBrowser", Icon16x16));
+		Set("MainFrame.VisitSupportWebSite",       new IMAGE_BRUSH_SVG("Starship/Common/Support", Icon16x16));
 		Set("MainFrame.VisitEpicGamesDotCom",      new IMAGE_BRUSH("Icons/Help/icon_Help_epic_16x", Icon16x16));
 		Set("MainFrame.AboutUnrealEd",             new IMAGE_BRUSH("Icons/Help/icon_Help_unreal_16x", Icon16x16));
-		Set("MainFrame.CreditsUnrealEd",           new IMAGE_BRUSH_SVG("Starship/Menus/Help/Credits", Icon16x16));
+		Set("MainFrame.CreditsUnrealEd",           new IMAGE_BRUSH_SVG("Starship/Common/Credits", Icon16x16));
 
 		Set( "EditorViewport.TranslateMode", new IMAGE_BRUSH_SVG( "Starship/EditorViewport/translate", Icon16x16 ) );
 		Set( "EditorViewport.RotateMode", new IMAGE_BRUSH_SVG("Starship/EditorViewport/rotate", Icon16x16 ) );
@@ -4251,16 +4251,16 @@ void FStarshipEditorStyle::FStyle::SetupLevelEditorStyle()
 
 		Set( "EditorViewport.CamSpeedSetting", new IMAGE_BRUSH_SVG( "Starship/EditorViewport/camera", Icon16x16) );
 		
-		Set( "EditorViewport.LitMode",            	  new IMAGE_BRUSH_SVG("Starship/EditorViewport/View/Lit", Icon16x16 ) );
-		Set( "EditorViewport.UnlitMode",          	  new IMAGE_BRUSH_SVG("Starship/EditorViewport/View/Unlit", Icon16x16 ) );
-		Set( "EditorViewport.WireframeMode",      	  new IMAGE_BRUSH_SVG("Starship/EditorViewport/View/BrushWireframe", Icon16x16 ) );
-		Set( "EditorViewport.DetailLightingMode", 	  new IMAGE_BRUSH_SVG("Starship/EditorViewport/View/DetailLighting", Icon16x16 ) );
-		Set( "EditorViewport.LightingOnlyMode",   	  new IMAGE_BRUSH_SVG("Starship/EditorViewport/View/LightingOnly", Icon16x16 ) );
+		Set( "EditorViewport.LitMode",            	  new IMAGE_BRUSH_SVG("Starship/Common/LitCube", Icon16x16 ) );
+		Set( "EditorViewport.UnlitMode",          	  new IMAGE_BRUSH_SVG("Starship/Common/UnlitCube", Icon16x16 ) );
+		Set( "EditorViewport.WireframeMode",      	  new IMAGE_BRUSH_SVG("Starship/Common/BrushWireframe", Icon16x16 ) );
+		Set( "EditorViewport.DetailLightingMode", 	  new IMAGE_BRUSH_SVG("Starship/Common/DetailLighting", Icon16x16 ) );
+		Set( "EditorViewport.LightingOnlyMode",   	  new IMAGE_BRUSH_SVG("Starship/Common/LightingOnly", Icon16x16 ) );
 
 		Set( "EditorViewport.LightComplexityMode", new IMAGE_BRUSH( "Icons/icon_ViewMode_LightComplexity_16px", Icon16x16 ) );
 		Set( "EditorViewport.ShaderComplexityMode", new IMAGE_BRUSH( "Icons/icon_ViewMode_Shadercomplexity_16px", Icon16x16 ) );
 
-		Set( "EditorViewport.QuadOverdrawMode", new IMAGE_BRUSH_SVG("Starship/EditorViewport/View/OptimizationViewmodes", Icon16x16 ) );
+		Set( "EditorViewport.QuadOverdrawMode", new IMAGE_BRUSH_SVG("Starship/Common/OptimizationViewmodes", Icon16x16 ) );
 
 		Set( "EditorViewport.ShaderComplexityWithQuadOverdrawMode", new IMAGE_BRUSH( "Icons/icon_ViewMode_Shadercomplexity_16px", Icon16x16 ) );
 		Set( "EditorViewport.TexStreamAccPrimitiveDistanceMode", new IMAGE_BRUSH( "Icons/icon_ViewMode_TextureStreamingAccuracy_16px", Icon16x16 ) );
@@ -4272,25 +4272,25 @@ void FStarshipEditorStyle::FStyle::SetupLevelEditorStyle()
 
 		Set( "EditorViewport.LODColorationMode", new IMAGE_BRUSH("Icons/icon_ViewMode_LODColoration_16px", Icon16x16) );
 		Set( "EditorViewport.HLODColorationMode", new IMAGE_BRUSH("Icons/icon_ViewMode_LODColoration_16px", Icon16x16));	
-		Set( "EditorViewport.GroupLODColorationMode", new IMAGE_BRUSH_SVG("Starship/EditorViewport/View/LODColorization", Icon16x16) );
+		Set( "EditorViewport.GroupLODColorationMode", new IMAGE_BRUSH_SVG("Starship/Common/LODColorization", Icon16x16) );
 
-		Set( "EditorViewport.VisualizeGBufferMode",   new IMAGE_BRUSH_SVG("Starship/EditorViewport/View/BufferVisualization", Icon16x16) );
+		Set( "EditorViewport.VisualizeGBufferMode",   new IMAGE_BRUSH_SVG("Starship/Common/BufferVisualization", Icon16x16) );
 
-		Set( "EditorViewport.Visualizers", 			  new IMAGE_BRUSH_SVG("Starship/EditorViewport/View/Visualizer", Icon16x16) );
-		Set( "EditorViewport.LOD", 			  		  new IMAGE_BRUSH_SVG("Starship/EditorViewport/View/LOD", Icon16x16) );
+		Set( "EditorViewport.Visualizers", 			  new IMAGE_BRUSH_SVG("Starship/Common/Visualizer", Icon16x16) );
+		Set( "EditorViewport.LOD", 			  		  new IMAGE_BRUSH_SVG("Starship/Common/LOD", Icon16x16) );
 
-		Set( "EditorViewport.ReflectionOverrideMode", new IMAGE_BRUSH_SVG("Starship/EditorViewport/View/Reflections", Icon16x16 ) );
-		Set( "EditorViewport.VisualizeBufferMode",    new IMAGE_BRUSH_SVG("Starship/EditorViewport/View/BufferVisualization", Icon16x16 ) );
-		Set( "EditorViewport.CollisionPawn",          new IMAGE_BRUSH_SVG("Starship/EditorViewport/View/PlayerCollision", Icon16x16 ) );
-		Set( "EditorViewport.CollisionVisibility",    new IMAGE_BRUSH_SVG("Starship/EditorViewport/View/VisibilityCollision", Icon16x16 ) );
+		Set( "EditorViewport.ReflectionOverrideMode", new IMAGE_BRUSH_SVG("Starship/Common/Reflections", Icon16x16 ) );
+		Set( "EditorViewport.VisualizeBufferMode",    new IMAGE_BRUSH_SVG("Starship/Common/BufferVisualization", Icon16x16 ) );
+		Set( "EditorViewport.CollisionPawn",          new IMAGE_BRUSH_SVG("Starship/Common/PlayerCollision", Icon16x16 ) );
+		Set( "EditorViewport.CollisionVisibility",    new IMAGE_BRUSH_SVG("Starship/Common/VisibilityCollision", Icon16x16 ) );
 		
-		Set( "EditorViewport.Perspective", new IMAGE_BRUSH_SVG("Starship/EditorViewport/Perspective/Perspective", Icon16x16 ) );
-		Set( "EditorViewport.Top",         new IMAGE_BRUSH_SVG("Starship/EditorViewport/Perspective/Top", Icon16x16 ) );
-		Set( "EditorViewport.Left",        new IMAGE_BRUSH_SVG("Starship/EditorViewport/Perspective/Left", Icon16x16 ) );
-		Set( "EditorViewport.Front",       new IMAGE_BRUSH_SVG("Starship/EditorViewport/Perspective/Front", Icon16x16 ) );
-		Set( "EditorViewport.Bottom",      new IMAGE_BRUSH_SVG("Starship/EditorViewport/Perspective/Bottom", Icon16x16 ) );
-		Set( "EditorViewport.Right",       new IMAGE_BRUSH_SVG("Starship/EditorViewport/Perspective/Right", Icon16x16 ) );
-		Set( "EditorViewport.Back",        new IMAGE_BRUSH_SVG("Starship/EditorViewport/Perspective/Back", Icon16x16 ) );
+		Set( "EditorViewport.Perspective", new IMAGE_BRUSH_SVG("Starship/Common/ViewPerspective", Icon16x16 ) );
+		Set( "EditorViewport.Top",         new IMAGE_BRUSH_SVG("Starship/Common/ViewTop", Icon16x16 ) );
+		Set( "EditorViewport.Left",        new IMAGE_BRUSH_SVG("Starship/Common/ViewLeft", Icon16x16 ) );
+		Set( "EditorViewport.Front",       new IMAGE_BRUSH_SVG("Starship/Common/ViewFront", Icon16x16 ) );
+		Set( "EditorViewport.Bottom",      new IMAGE_BRUSH_SVG("Starship/Common/ViewBottom", Icon16x16 ) );
+		Set( "EditorViewport.Right",       new IMAGE_BRUSH_SVG("Starship/Common/ViewRight", Icon16x16 ) );
+		Set( "EditorViewport.Back",        new IMAGE_BRUSH_SVG("Starship/Common/ViewBack", Icon16x16 ) );
 
 		Set("EditorViewport.ActiveBorderColor", FStyleColors::Primary);
 
@@ -4298,45 +4298,45 @@ void FStarshipEditorStyle::FStyle::SetupLevelEditorStyle()
 
 #if WITH_EDITOR || (IS_PROGRAM && WITH_UNREAL_DEVELOPER_TOOLS)
 		{
-			Set( "LevelEditor.Tabs.Details",                new IMAGE_BRUSH_SVG("Starship/LevelEditor/Menus/Details", Icon16x16) );
-			Set( "LevelEditor.Tabs.Cinematics",             new IMAGE_BRUSH_SVG("Starship/LevelEditor/Menus/Cinematics", Icon16x16) );
-			Set( "LevelEditor.Tabs.EditorModes",            new IMAGE_BRUSH_SVG("Starship/LevelEditor/Menus/EditorModes", Icon16x16) );
-			Set( "LevelEditor.Tabs.Modes",                  new IMAGE_BRUSH_SVG("Starship/LevelEditor/Menus/EditorModes", Icon16x16) );
-			Set( "LevelEditor.Tabs.PlacementBrowser",       new IMAGE_BRUSH_SVG("Starship/LevelEditor/Menus/PlaceActors", Icon16x16) );
-			Set( "LevelEditor.Tabs.Properties",             new IMAGE_BRUSH_SVG("Starship/LevelEditor/Menus/Properties", Icon16x16) );
-			Set( "LevelEditor.Tabs.Outliner",               new IMAGE_BRUSH_SVG("Starship/LevelEditor/Menus/WorldOutliner", Icon16x16) );
-			Set( "LevelEditor.Tabs.ContentBrowser",         new IMAGE_BRUSH_SVG("Starship/LevelEditor/Menus/ContentBrowser", Icon16x16) );
-			Set( "LevelEditor.Tabs.Levels",                 new IMAGE_BRUSH_SVG("Starship/LevelEditor/Menus/Levels", Icon16x16) );
-			Set( "LevelEditor.Tabs.WorldBrowser",           new IMAGE_BRUSH_SVG("Starship/LevelEditor/Menus/Levels", Icon16x16) );
-			Set( "LevelEditor.Tabs.WorldBrowserDetails",    new IMAGE_BRUSH_SVG("Starship/LevelEditor/Menus/WorldSettings", Icon16x16) );
+			Set( "LevelEditor.Tabs.Details",                new IMAGE_BRUSH_SVG("Starship/Common/Details", Icon16x16) );
+			Set( "LevelEditor.Tabs.Cinematics",             new IMAGE_BRUSH_SVG("Starship/Common/Cinematics", Icon16x16) );
+			Set( "LevelEditor.Tabs.EditorModes",            new IMAGE_BRUSH_SVG("Starship/Common/EditorModes", Icon16x16) );
+			Set( "LevelEditor.Tabs.Modes",                  new IMAGE_BRUSH_SVG("Starship/Common/EditorModes", Icon16x16) );
+			Set( "LevelEditor.Tabs.PlacementBrowser",       new IMAGE_BRUSH_SVG("Starship/Common/PlaceActors", Icon16x16) );
+			Set( "LevelEditor.Tabs.Properties",             new IMAGE_BRUSH_SVG("Starship/Common/Properties", Icon16x16) );
+			Set( "LevelEditor.Tabs.Outliner",               new IMAGE_BRUSH_SVG("Starship/Common/WorldOutliner", Icon16x16) );
+			Set( "LevelEditor.Tabs.ContentBrowser",         new IMAGE_BRUSH_SVG("Starship/Common/ContentBrowser", Icon16x16) );
+			Set( "LevelEditor.Tabs.Levels",                 new IMAGE_BRUSH_SVG("Starship/Common/Levels", Icon16x16) );
+			Set( "LevelEditor.Tabs.WorldBrowser",           new IMAGE_BRUSH_SVG("Starship/Common/Levels", Icon16x16) );
+			Set( "LevelEditor.Tabs.WorldBrowserDetails",    new IMAGE_BRUSH_SVG("Starship/Common/WorldSettings", Icon16x16) );
 
 			Set( "LevelEditor.Tabs.WorldBrowserComposition",new IMAGE_BRUSH("Icons/icon_levels_compositionbutton_16x", Icon16x16 ) );
 
-			Set( "LevelEditor.Tabs.Layers",                 new IMAGE_BRUSH_SVG("Starship/LevelEditor/Menus/Layers", Icon16x16) );
-			Set( "LevelEditor.Tabs.ComposureCompositing",   new IMAGE_BRUSH_SVG("Starship/LevelEditor/Menus/ComposureCompositing", Icon16x16) );
-			Set( "LevelEditor.Tabs.USDStage",   			new IMAGE_BRUSH_SVG("Starship/LevelEditor/Menus/USDStage", Icon16x16) );
+			Set( "LevelEditor.Tabs.Layers",                 new IMAGE_BRUSH_SVG("Starship/Common/Layers", Icon16x16) );
+			Set( "LevelEditor.Tabs.ComposureCompositing",   new IMAGE_BRUSH_SVG("Starship/Common/ComposureCompositing", Icon16x16) );
+			Set( "LevelEditor.Tabs.USDStage",   			new IMAGE_BRUSH_SVG("Starship/Common/USDStage", Icon16x16) );
 
 			Set( "LevelEditor.Tabs.BuildAndSubmit",         new IMAGE_BRUSH("Icons/icon_tab_BuildSubmit_16x", Icon16x16 ) );
 
-			Set( "LevelEditor.Tabs.StatsViewer",            new IMAGE_BRUSH_SVG("Starship/LevelEditor/Menus/Statistics", Icon16x16) );
+			Set( "LevelEditor.Tabs.StatsViewer",            new IMAGE_BRUSH_SVG("Starship/Common/Statistics", Icon16x16) );
 
 			Set( "LevelEditor.Tabs.Toolbar",                new IMAGE_BRUSH("Icons/icon_tab_Toolbars_16x", Icon16x16 ) );
 
-			Set( "LevelEditor.Tabs.Viewports",              new IMAGE_BRUSH_SVG("Starship/LevelEditor/Menus/Viewports", Icon16x16) );
-			Set( "LevelEditor.Tabs.HLOD",                   new IMAGE_BRUSH_SVG("Starship/LevelEditor/Menus/HierarchicalLODOutliner", Icon16x16) );
+			Set( "LevelEditor.Tabs.Viewports",              new IMAGE_BRUSH_SVG("Starship/Common/Viewports", Icon16x16) );
+			Set( "LevelEditor.Tabs.HLOD",                   new IMAGE_BRUSH_SVG("Starship/Common/HierarchicalLODOutliner", Icon16x16) );
 		}
 #endif // WITH_EDITOR || (IS_PROGRAM && WITH_UNREAL_DEVELOPER_TOOLS)
 
 #if WITH_EDITOR
-		Set( "LevelEditor.NewLevel",      new IMAGE_BRUSH_SVG( "Starship/Menus/File/LevelNew", Icon16x16 ) );
-		Set( "LevelEditor.OpenLevel",     new IMAGE_BRUSH_SVG( "Starship/Menus/File/LevelOpen", Icon16x16 ) );
-		Set( "LevelEditor.Save",          new IMAGE_BRUSH_SVG( "Starship/Menus/File/SaveCurrent", Icon16x16 ) );
-		Set( "LevelEditor.SaveAs",        new IMAGE_BRUSH_SVG( "Starship/Menus/File/SaveCurrentAs", Icon16x16 ) );
-		Set( "LevelEditor.SaveAllLevels", new IMAGE_BRUSH_SVG( "Starship/Menus/File/LevelSaveAll", Icon16x16 ) );
+		Set( "LevelEditor.NewLevel",      new IMAGE_BRUSH_SVG( "Starship/Common/LevelNew", Icon16x16 ) );
+		Set( "LevelEditor.OpenLevel",     new IMAGE_BRUSH_SVG( "Starship/Common/LevelOpen", Icon16x16 ) );
+		Set( "LevelEditor.Save",          new IMAGE_BRUSH_SVG( "Starship/Common/SaveCurrent", Icon16x16 ) );
+		Set( "LevelEditor.SaveAs",        new IMAGE_BRUSH_SVG( "Starship/Common/SaveCurrentAs", Icon16x16 ) );
+		Set( "LevelEditor.SaveAllLevels", new IMAGE_BRUSH_SVG( "Starship/Common/LevelSaveAll", Icon16x16 ) );
 
-		Set( "LevelEditor.ImportScene",    new IMAGE_BRUSH_SVG( "Starship/Menus/File/LevelImportInto", Icon16x16 ) );
-		Set( "LevelEditor.ExportAll",      new IMAGE_BRUSH_SVG( "Starship/Menus/File/ExportAll", Icon16x16 ) );
-		Set( "LevelEditor.ExportSelected", new IMAGE_BRUSH_SVG( "Starship/Menus/File/ExportSelected", Icon16x16 ) );
+		Set( "LevelEditor.ImportScene",    new IMAGE_BRUSH_SVG( "Starship/Common/LevelImportInto", Icon16x16 ) );
+		Set( "LevelEditor.ExportAll",      new IMAGE_BRUSH_SVG( "Starship/Common/ExportAll", Icon16x16 ) );
+		Set( "LevelEditor.ExportSelected", new IMAGE_BRUSH_SVG( "Starship/Common/ExportSelected", Icon16x16 ) );
 
 		Set( "LevelEditor.Build", new IMAGE_BRUSH( "Icons/icon_build_40x", Icon40x40 ) );
 		Set( "LevelEditor.Build.Small", new IMAGE_BRUSH( "Icons/icon_build_40x", Icon20x20 ) );
@@ -4429,7 +4429,7 @@ void FStarshipEditorStyle::FStyle::SetupLevelEditorStyle()
 		Set( "LevelEditor.WorldProperties", new IMAGE_BRUSH( "Icons/icon_worldscript_40x", Icon40x40 ) );
 		Set( "LevelEditor.WorldProperties.Small", new IMAGE_BRUSH( "Icons/icon_worldscript_40x", Icon20x20 ) );
 
-		Set( "LevelEditor.WorldProperties.Tab", new IMAGE_BRUSH_SVG( "Starship/LevelEditor/Menus/WorldSettings", Icon16x16 ) );
+		Set( "LevelEditor.WorldProperties.Tab", new IMAGE_BRUSH_SVG( "Starship/Common/WorldSettings", Icon16x16 ) );
 
 		Set( "LevelEditor.OpenContentBrowser", new IMAGE_BRUSH_SVG( "Starship/MainToolbar/content", Icon20x20 ) );
 		Set( "LevelEditor.OpenMarketplace", new IMAGE_BRUSH_SVG( "Starship/MainToolbar/marketplace", Icon20x20) );
@@ -4591,8 +4591,8 @@ void FStarshipEditorStyle::FStyle::SetupLevelEditorStyle()
 		Set( "LevelViewport.PilotSelectedActor", new IMAGE_BRUSH( "Icons/icon_PilotSelectedActor_16x", Icon16x16 ) );
 		Set( "LevelViewport.PilotSelectedActor.Small", new IMAGE_BRUSH( "Icons/icon_PilotSelectedActor_16x", Icon16x16 ) );
 		
-		Set( "LevelViewport.ToggleActorPilotCameraView",       new IMAGE_BRUSH_SVG( "Starship/LevelEditor/Menus/CameraShake", Icon16x16 ) );
-		Set( "LevelViewport.ToggleActorPilotCameraView.Small", new IMAGE_BRUSH_SVG( "Starship/LevelEditor/Menus/CameraShake", Icon16x16 ) );
+		Set( "LevelViewport.ToggleActorPilotCameraView",       new IMAGE_BRUSH_SVG( "Starship/Common/CameraShake", Icon16x16 ) );
+		Set( "LevelViewport.ToggleActorPilotCameraView.Small", new IMAGE_BRUSH_SVG( "Starship/Common/CameraShake", Icon16x16 ) );
 
 		Set( "LevelViewport.ActorPilotText", FTextBlockStyle()
 			.SetFont( DEFAULT_FONT( "BoldCondensed", 12 ) )
@@ -5009,7 +5009,7 @@ void FStarshipEditorStyle::FStyle::SetupPersonaStyle()
 
 		Set( "BlueprintEditor.ActionMenu.ContextDescriptionFont",  DEFAULT_FONT("Regular", 12) );
 
-		Set( "BlueprintEditor.FindInBlueprints.MenuIcon", new IMAGE_BRUSH_SVG("Starship/LevelEditor/Menus/FindInBlueprints", Icon16x16) );
+		Set( "BlueprintEditor.FindInBlueprints.MenuIcon", new IMAGE_BRUSH_SVG("Starship/Common/FindInBlueprints", Icon16x16) );
 
 		Set( "BlueprintEditor.FindInBlueprint", new IMAGE_BRUSH( "Icons/icon_Blueprint_Find_40px", Icon40x40 ) );
 		Set( "BlueprintEditor.FindInBlueprint.Small", new IMAGE_BRUSH( "Icons/icon_Blueprint_Find_40px", Icon20x20 ) );
@@ -5906,7 +5906,7 @@ void FStarshipEditorStyle::FStyle::SetupContentBrowserStyle()
 	// Content Browser
 	{
 		// Tab and menu icon
-		Set("ContentBrowser.TabIcon", new IMAGE_BRUSH_SVG("Starship/LevelEditor/Menus/ContentBrowser", Icon16x16));
+		Set("ContentBrowser.TabIcon", new IMAGE_BRUSH_SVG("Starship/Common/ContentBrowser", Icon16x16));
 
 		// Sources View
 		Set("ContentBrowser.SourceTitleFont", DEFAULT_FONT( "Regular", 12 ) );
@@ -6005,7 +6005,7 @@ void FStarshipEditorStyle::FStyle::SetupContentBrowserStyle()
 			.SetCheckedPressedForegroundColor(FStyleColors::PrimaryHover)
 			.SetPressedForegroundColor(FStyleColors::ForegroundHover);
 
-		Set("ContentBrowser.FilterImage", new CORE_IMAGE_BRUSH_SVG("Starship/Common/filled-circle", Icon8x8));
+		Set("ContentBrowser.FilterImage",  new FSlateRoundedBoxBrush(FStyleColors::White, Icon8x8));// new CORE_IMAGE_BRUSH_SVG("Starship/Common/filled-circle", Icon8x8));
 		Set("ContentBrowser.FilterBackgroundUnchecked", new FSlateRoundedBoxBrush(FStyleColors::Secondary, 4.0f));
 		Set("ContentBrowser.FilterBackgroundChecked", new FSlateRoundedBoxBrush(FStyleColors::Recessed, 4.0f));
 		/* ... and add the new style */
@@ -7706,7 +7706,7 @@ void FStarshipEditorStyle::FStyle::SetupTranslationEditorStyles()
 
 void FStarshipEditorStyle::FStyle::SetupLocalizationDashboardStyles()
 {
-	Set("LocalizationDashboard.MenuIcon", new IMAGE_BRUSH_SVG("Starship/LevelEditor/Menus/LocalizationDashboard", Icon16x16));
+	Set("LocalizationDashboard.MenuIcon", new IMAGE_BRUSH_SVG("Starship/Common/LocalizationDashboard", Icon16x16));
 
 	Set("LocalizationDashboard.GatherTextAllTargets", new IMAGE_BRUSH("Icons/Icon_Localisation_Gather_All_40x", Icon40x40));
 	Set("LocalizationDashboard.ImportTextAllTargetsAllCultures", new IMAGE_BRUSH("Icons/Icon_Localisation_Import_All_40x", Icon40x40));
diff --git a/Engine/Source/Runtime/AppFramework/Private/Framework/Testing/SStarshipGallery.cpp b/Engine/Source/Runtime/AppFramework/Private/Framework/Testing/SStarshipGallery.cpp
--- a/Engine/Source/Runtime/AppFramework/Private/Framework/Testing/SStarshipGallery.cpp
+++ b/Engine/Source/Runtime/AppFramework/Private/Framework/Testing/SStarshipGallery.cpp
@@ -453,14 +453,15 @@ public:
             .Padding(48)
             [
                 SNew(SVerticalBox)
-                +SVerticalBox::Slot().AutoHeight()[ GenerateIconLibrarySVG(NSLOCTEXT("StarshipGallery", "Starship Common", "Common"), "Content/Slate/Starship/Common")]
+                +SVerticalBox::Slot().AutoHeight()[ GenerateIconLibrarySVG(NSLOCTEXT("StarshipGallery", "SlateCore", "Core"), "Content/Slate/Starship/Common")]
+                +SVerticalBox::Slot().AutoHeight()[ GenerateIconLibrarySVG(NSLOCTEXT("StarshipGallery", "Editor Common", "Editor"), "Content/Editor/Slate/Starship/Common")]
 
                 +SVerticalBox::Slot().AutoHeight()[ GenerateIconLibrarySVG(NSLOCTEXT("StarshipGallery", "SceneOutliner", "SceneOutliner"), "Content/Editor/Slate/Starship/SceneOutliner")]
-                +SVerticalBox::Slot().AutoHeight()[ GenerateIconLibrarySVG(NSLOCTEXT("StarshipGallery", "LevelEditor", "LevelEditor"), "Content/Editor/Slate/Starship/LevelEditor/Menus")]
+                // +SVerticalBox::Slot().AutoHeight()[ GenerateIconLibrarySVG(NSLOCTEXT("StarshipGallery", "LevelEditor", "LevelEditor"), "Content/Editor/Slate/Starship/LevelEditor/Menus")]
                 +SVerticalBox::Slot().AutoHeight()[ GenerateIconLibrarySVG(NSLOCTEXT("StarshipGallery", "MainToolbar", "MainToolbar"), "Content/Editor/Slate/Starship/MainToolbar")]
-                +SVerticalBox::Slot().AutoHeight()[ GenerateIconLibrarySVG(NSLOCTEXT("StarshipGallery", "FileMenu", "FileMenu"), "Content/Editor/Slate/Starship/Menus/File")]
-                +SVerticalBox::Slot().AutoHeight()[ GenerateIconLibrarySVG(NSLOCTEXT("StarshipGallery", "EditMenu", "EditMenu"), "Content/Editor/Slate/Starship/Menus/Edit")]
-                +SVerticalBox::Slot().AutoHeight()[ GenerateIconLibrarySVG(NSLOCTEXT("StarshipGallery", "HelpMenu", "HelpMenu"), "Content/Editor/Slate/Starship/Menus/Help")]
+                // +SVerticalBox::Slot().AutoHeight()[ GenerateIconLibrarySVG(NSLOCTEXT("StarshipGallery", "FileMenu", "FileMenu"), "Content/Editor/Slate/Starship/Menus/File")]
+                // +SVerticalBox::Slot().AutoHeight()[ GenerateIconLibrarySVG(NSLOCTEXT("StarshipGallery", "EditMenu", "EditMenu"), "Content/Editor/Slate/Starship/Menus/Edit")]
+                // +SVerticalBox::Slot().AutoHeight()[ GenerateIconLibrarySVG(NSLOCTEXT("StarshipGallery", "HelpMenu", "HelpMenu"), "Content/Editor/Slate/Starship/Menus/Help")]
 
                 +SVerticalBox::Slot().AutoHeight()[ GenerateIconLibrary(NSLOCTEXT("StarshipGallery", "PaintIconTitle", "Paint"), "Content/Editor/Slate/Icons/Paint")]
                 +SVerticalBox::Slot().AutoHeight()[ GenerateIconLibrary(NSLOCTEXT("StarshipGallery", "LandscapeIconTitle", "Landscape"), "Content/Editor/Slate/Icons/Landscape")]
