Fixed up missing viewport toolbar options dropdown image

diff --git a/Engine/Source/Editor/EditorStyle/Private/StarshipStyle.cpp b/Engine/Source/Editor/EditorStyle/Private/StarshipStyle.cpp
--- a/Engine/Source/Editor/EditorStyle/Private/StarshipStyle.cpp
+++ b/Engine/Source/Editor/EditorStyle/Private/StarshipStyle.cpp
@@ -2448,7 +2448,7 @@ void FStarshipEditorStyle::FStyle::SetupViewportStyles()
 		ToolbarBackgroundColor.A = .75f;
 
 		Set("EditorViewportToolBar.Background", new FSlateColorBrush(ToolbarBackgroundColor));
-		Set("EditorViewportToolBar.OptionsDropdown", new CORE_IMAGE_BRUSH_SVG("Starship/EditorViewport/menu", FVector2D(13,8)));
+		Set("EditorViewportToolBar.OptionsDropdown", new IMAGE_BRUSH_SVG("Starship/EditorViewport/menu", FVector2D(13,8)));
 
 		Set("EditorViewportToolBar.Font", FStyleFonts::Get().Normal);
 
