Add Foundation Icon to StarshipStyle (copy from SlateEditorStyle.cpp)

diff --git a/Engine/Source/Editor/EditorStyle/Private/StarshipStyle.cpp b/Engine/Source/Editor/EditorStyle/Private/StarshipStyle.cpp
--- a/Engine/Source/Editor/EditorStyle/Private/StarshipStyle.cpp
+++ b/Engine/Source/Editor/EditorStyle/Private/StarshipStyle.cpp
@@ -5903,7 +5903,8 @@ void FStarshipEditorStyle::FStyle::SetupClassIconsAndThumbnails()
 			TEXT("CineCameraActor"),
 			TEXT("CameraRig_Crane"),
 			TEXT("CameraRig_Rail"),
-			TEXT("FoliageType_Actor")
+			TEXT("FoliageType_Actor"),
+			TEXT("FoundationActor")
 		};
 
 		for (int32 TypeIndex = 0; TypeIndex < UE_ARRAY_COUNT(AssetTypes); ++TypeIndex)
