Updated starship style from UE4 editor style changes

diff --git a/Engine/Source/Editor/EditorStyle/Private/StarshipStyle.cpp b/Engine/Source/Editor/EditorStyle/Private/StarshipStyle.cpp
--- a/Engine/Source/Editor/EditorStyle/Private/StarshipStyle.cpp
+++ b/Engine/Source/Editor/EditorStyle/Private/StarshipStyle.cpp
@@ -2008,7 +2008,7 @@ void FStarshipEditorStyle::FStyle::SetupWorldBrowserStyles()
 	// World Browser
 	{
 		Set("WorldBrowser.AddLayer", new IMAGE_BRUSH("Icons/icon_levels_addlayer_16x", Icon16x16));
-		Set("WorldBrowser.SimulationViewPositon", new IMAGE_BRUSH("Icons/icon_levels_simulationviewpos_16x", Icon16x16));
+		Set("WorldBrowser.SimulationViewPosition", new IMAGE_BRUSH("Icons/icon_levels_simulationviewpos_16x", Icon16x16));
 		Set("WorldBrowser.MouseLocation", new IMAGE_BRUSH("Icons/icon_levels_mouselocation_16x", Icon16x16));
 		Set("WorldBrowser.MarqueeRectSize", new IMAGE_BRUSH("Icons/icon_levels_marqueerectsize_16x", Icon16x16));
 		Set("WorldBrowser.WorldSize", new IMAGE_BRUSH("Icons/icon_levels_worldsize_16x", Icon16x16));
@@ -5659,7 +5659,8 @@ void FStarshipEditorStyle::FStyle::SetupClassIconsAndThumbnails()
 		Set("ClassIcon.VectorFieldComponent", new IMAGE_BRUSH("Icons/ActorIcons/VectorFieldVolume_16x", Icon16x16));
 		Set("ClassIcon.ArrowComponent", new IMAGE_BRUSH("Icons/ActorIcons/Arrow_16px", Icon16x16));
 		Set("ClassIcon.AtmosphericFogComponent", new IMAGE_BRUSH("Icons/AssetIcons/AtmosphericFog_16x", Icon16x16));
-		Set("ClassIcon.SkyAtmosphereComponent", new IMAGE_BRUSH("Icons/AssetIcons/AtmosphericFog_16x", Icon16x16));
+		Set("ClassIcon.SkyAtmosphereComponent", new IMAGE_BRUSH("Icons/AssetIcons/SkyAtmosphere_16x", Icon16x16));
+		Set("ClassIcon.VolumetricCloudComponent", new IMAGE_BRUSH("Icons/AssetIcons/SkyAtmosphere_16x", Icon16x16));
 		Set("ClassIcon.BoxComponent", new IMAGE_BRUSH("Icons/ActorIcons/Box_16px", Icon16x16));
 		Set("ClassIcon.CapsuleComponent", new IMAGE_BRUSH("Icons/ActorIcons/Capsule_16px", Icon16x16));
 		Set("ClassIcon.InstancedStaticMeshComponent", new IMAGE_BRUSH("Icons/ActorIcons/InstancedStaticMesh_16px", Icon16x16));
@@ -5863,6 +5864,7 @@ void FStarshipEditorStyle::FStyle::SetupClassIconsAndThumbnails()
 			TEXT("SkyLight"),
 			TEXT("SkyLightComponent"),
 			TEXT("SkyAtmosphere"),
+			TEXT("VolumetricCloud"),
 			TEXT("SkeletalMesh"),
 			TEXT("Skeleton"),
 			TEXT("SlateBrushAsset"),
@@ -7430,7 +7432,7 @@ void FStarshipEditorStyle::FStyle::SetupAutomationStyles()
 			FString NormalIconPath = PlatformInfo.GetIconPath(EPlatformIconSize::Normal);
 			if (NormalIconPath.StartsWith(TEXT("/Platforms/")))
 			{
-				#define PLATFORM_IMAGE_BRUSH( PlatformPath, ... ) FSlateImageBrush( PlatformPath.Replace(TEXT("/Platforms/"), *FPaths::EnginePlatformExtensionsDir()) + TEXT(".png") , __VA_ARGS__ )
+#define PLATFORM_IMAGE_BRUSH( PlatformPath, ... ) FSlateImageBrush( PlatformPath.Replace(TEXT("/Platforms/"), *FPaths::EnginePlatformExtensionsDir()) + TEXT(".png") , __VA_ARGS__ )
 				Set(PlatformInfo.GetIconStyleName(EPlatformIconSize::Normal), new PLATFORM_IMAGE_BRUSH(NormalIconPath, Icon24x24));
 				Set(PlatformInfo.GetIconStyleName(EPlatformIconSize::Large), new PLATFORM_IMAGE_BRUSH(PlatformInfo.GetIconPath(EPlatformIconSize::Large), Icon64x64));
 				Set(PlatformInfo.GetIconStyleName(EPlatformIconSize::XLarge), new PLATFORM_IMAGE_BRUSH(PlatformInfo.GetIconPath(EPlatformIconSize::XLarge), Icon128x128));
