Starship style cleanup

diff --git a/Engine/Plugins/Editor/ContentBrowser/ContentBrowserAssetDataSource/Source/ContentBrowserAssetDataSource/Private/NewAssetContextMenu.cpp b/Engine/Plugins/Editor/ContentBrowser/ContentBrowserAssetDataSource/Source/ContentBrowserAssetDataSource/Private/NewAssetContextMenu.cpp
--- a/Engine/Plugins/Editor/ContentBrowser/ContentBrowserAssetDataSource/Source/ContentBrowserAssetDataSource/Private/NewAssetContextMenu.cpp
+++ b/Engine/Plugins/Editor/ContentBrowser/ContentBrowserAssetDataSource/Source/ContentBrowserAssetDataSource/Private/NewAssetContextMenu.cpp
@@ -238,7 +238,7 @@ void FNewAssetContextMenu::MakeContextMenu(
 				"ImportAsset",
 				FText::Format(LOCTEXT("ImportAsset", "Import to {0}..."), FText::FromName(FirstSelectedPath)),
 				LOCTEXT("ImportAssetTooltip_NewAsset", "Imports an asset from file to this folder."),
-				FSlateIcon(FEditorStyle::GetStyleSetName(), "ContentBrowser.ImportIcon"),
+				FSlateIcon(FAppStyle::GetAppStyleSetName(), "Icons.Import"),
 				FUIAction(
 					FExecuteAction::CreateStatic(&FNewAssetContextMenu::ExecuteImportAsset, InOnImportAssetRequested, FirstSelectedPath),
 					CanExecuteAssetActionsDelegate
diff --git a/Engine/Plugins/Editor/FacialAnimation/Source/FacialAnimationEditor/Private/FacialAnimationEditorModule.cpp b/Engine/Plugins/Editor/FacialAnimation/Source/FacialAnimationEditor/Private/FacialAnimationEditorModule.cpp
--- a/Engine/Plugins/Editor/FacialAnimation/Source/FacialAnimationEditor/Private/FacialAnimationEditorModule.cpp
+++ b/Engine/Plugins/Editor/FacialAnimation/Source/FacialAnimationEditor/Private/FacialAnimationEditorModule.cpp
@@ -58,7 +58,7 @@ void FFacialAnimationEditorModule::StartupModule()
 		{
 			const TSharedRef<SDockTab> Tab =
 				SNew(SDockTab)
-				.Icon(FEditorStyle::GetBrush("ContentBrowser.ImportIcon"))
+				.Icon(FAppStyle::Get().GetBrush("Icons.Import"))
 				.TabRole(ETabRole::NomadTab);
 
 			Tab->SetContent(SNew(SFacialAnimationBulkImporter));
@@ -74,7 +74,7 @@ void FFacialAnimationEditorModule::StartupModule()
 			.SetDisplayName(LOCTEXT("FacialAnimationBulkImporterTabTitle", "Facial Anim Importer"))
 			.SetTooltipText(LOCTEXT("FacialAnimationBulkImporterTooltipText", "Open the Facial Animation Bulk Importer tab."))
 			.SetGroup(WorkspaceMenu::GetMenuStructure().GetDeveloperToolsMiscCategory())
-			.SetIcon(FSlateIcon(FEditorStyle::GetStyleSetName(), "ContentBrowser.ImportIcon"));
+			.SetIcon(FSlateIcon(FAppStyle::GetAppStyleSetName(), "Icons..Import"));
 
 		// register for when persona is loaded
 		OnModulesChangedDelegate = FModuleManager::Get().OnModulesChanged().AddRaw(this, &FFacialAnimationEditorModule::HandleModulesChanged);
diff --git a/Engine/Source/Editor/ContentBrowser/Private/SThumbnailEditModeTools.cpp b/Engine/Source/Editor/ContentBrowser/Private/SThumbnailEditModeTools.cpp
--- a/Engine/Source/Editor/ContentBrowser/Private/SThumbnailEditModeTools.cpp
+++ b/Engine/Source/Editor/ContentBrowser/Private/SThumbnailEditModeTools.cpp
@@ -43,14 +43,14 @@ void SThumbnailEditModeTools::Construct( const FArguments& InArgs, const TShared
 
 		// Primitive tools
 		+SHorizontalBox::Slot()
-		.AutoWidth()
 		.VAlign(VAlign_Top)
+		.HAlign(HAlign_Left)
 		.Padding(1)
 		[
 			SNew(SButton)
 			.Visibility(this, &SThumbnailEditModeTools::GetPrimitiveToolsVisibility)
 			.ContentPadding(0)
-			.ButtonStyle(FEditorStyle::Get(), "ToggleButton")
+			.ButtonStyle(FEditorStyle::Get(), "SimpleButton")
 			.OnClicked(this, &SThumbnailEditModeTools::ChangePrimitive)
 			.ToolTipText(LOCTEXT("CyclePrimitiveThumbnailShapes", "Cycle through primitive shape for this thumbnail"))
 			.Content()
@@ -61,17 +61,19 @@ void SThumbnailEditModeTools::Construct( const FArguments& InArgs, const TShared
 		+SHorizontalBox::Slot()
 		.HAlign(HAlign_Right)
 		.VAlign(VAlign_Top)
+		.Padding(1)
 		[
 			SNew(SButton)
 			.Visibility(this, &SThumbnailEditModeTools::GetPrimitiveToolsResetToDefaultVisibility)
 			.ContentPadding(0)
-			.ButtonStyle(FEditorStyle::Get(), "ToggleButton")
+			.ButtonStyle(FAppStyle::Get(), "SimpleButton")
 			.OnClicked(this, &SThumbnailEditModeTools::ResetToDefault)
 			.ToolTipText(LOCTEXT("ResetThumbnailToDefault", "Resets thumbnail to the default"))
 			.Content()
 			[
 				SNew(SImage)
 				.Image(FEditorStyle::GetBrush("ContentBrowser.ResetPrimitiveToDefault"))
+				.ColorAndOpacity(FSlateColor::UseForeground())
 			]
 		]
 	];
diff --git a/Engine/Source/Editor/EditorStyle/Private/StarshipStyle.cpp b/Engine/Source/Editor/EditorStyle/Private/StarshipStyle.cpp
--- a/Engine/Source/Editor/EditorStyle/Private/StarshipStyle.cpp
+++ b/Engine/Source/Editor/EditorStyle/Private/StarshipStyle.cpp
@@ -1976,29 +1976,29 @@ void FStarshipEditorStyle::FStyle::SetupLevelGeneralStyles()
 {
 // Levels General
 	{
-	Set("Level.VisibleIcon16x",             new CORE_IMAGE_BRUSH_SVG("Starship/Common/visible", Icon16x16));
-	Set("Level.VisibleHighlightIcon16x",    new CORE_IMAGE_BRUSH_SVG("Starship/Common/visible", Icon16x16));
-	Set("Level.NotVisibleIcon16x",          new CORE_IMAGE_BRUSH_SVG("Starship/Common/hidden",  Icon16x16));
-	Set("Level.NotVisibleHighlightIcon16x", new CORE_IMAGE_BRUSH_SVG("Starship/Common/hidden",  Icon16x16));
-
-	Set("Level.LightingScenarioIcon16x", new IMAGE_BRUSH("Icons/icon_levels_LightingScenario_16px", Icon16x16));
-	Set("Level.LightingScenarioNotIcon16x", new IMAGE_BRUSH("Icons/icon_levels_LightingScenarioNot_16px", Icon16x16));
-	Set("Level.LockedIcon16x", new IMAGE_BRUSH("Icons/icon_locked_16px", Icon16x16));
-	Set("Level.LockedHighlightIcon16x", new IMAGE_BRUSH("Icons/icon_locked_highlight_16px", Icon16x16));
-	Set("Level.UnlockedIcon16x", new IMAGE_BRUSH("Icons/icon_levels_unlocked_16px", Icon16x16));
-	Set("Level.UnlockedHighlightIcon16x", new IMAGE_BRUSH("Icons/icon_levels_unlocked_hi_16px", Icon16x16));
-	Set("Level.ReadOnlyLockedIcon16x", new IMAGE_BRUSH("Icons/icon_levels_LockedReadOnly_16px", Icon16x16));
-	Set("Level.ReadOnlyLockedHighlightIcon16x", new IMAGE_BRUSH("Icons/icon_levels_LockedReadOnly_hi_16px", Icon16x16));
-	Set("Level.SaveIcon16x", new IMAGE_BRUSH("Icons/icon_levels_Save_16px", Icon16x16));
-	Set("Level.SaveHighlightIcon16x", new IMAGE_BRUSH("Icons/icon_levels_Save_hi_16px", Icon16x16));
-	Set("Level.SaveModifiedIcon16x", new IMAGE_BRUSH("Icons/icon_levels_SaveModified_16px", Icon16x16));
-	Set("Level.SaveModifiedHighlightIcon16x", new IMAGE_BRUSH("Icons/icon_levels_SaveModified_hi_16px", Icon16x16));
-	Set("Level.SaveDisabledIcon16x", new IMAGE_BRUSH("Icons/icon_levels_SaveDisabled_16px", Icon16x16));
-	Set("Level.SaveDisabledHighlightIcon16x", new IMAGE_BRUSH("Icons/icon_levels_SaveDisabled_hi_16px", Icon16x16));
-	Set("Level.ScriptIcon16x", new IMAGE_BRUSH("Icons/icon_levels_Blueprint_16px", Icon16x16));
-	Set("Level.ScriptHighlightIcon16x", new IMAGE_BRUSH("Icons/icon_levels_Blueprint_hi_16px", Icon16x16));
-	Set("Level.EmptyIcon16x", new IMAGE_BRUSH("Icons/Empty_16x", Icon16x16));
-	Set("Level.ColorIcon40x", new IMAGE_BRUSH("Icons/icon_levels_back_16px", Icon16x16));
+		Set("Level.VisibleIcon16x", new CORE_IMAGE_BRUSH_SVG("Starship/Common/visible", Icon16x16));
+		Set("Level.VisibleHighlightIcon16x", new CORE_IMAGE_BRUSH_SVG("Starship/Common/visible", Icon16x16));
+		Set("Level.NotVisibleIcon16x", new CORE_IMAGE_BRUSH_SVG("Starship/Common/hidden", Icon16x16));
+		Set("Level.NotVisibleHighlightIcon16x", new CORE_IMAGE_BRUSH_SVG("Starship/Common/hidden", Icon16x16));
+
+		Set("Level.LightingScenarioIcon16x", new IMAGE_BRUSH("Icons/icon_levels_LightingScenario_16px", Icon16x16));
+		Set("Level.LightingScenarioNotIcon16x", new IMAGE_BRUSH("Icons/icon_levels_LightingScenarioNot_16px", Icon16x16));
+		Set("Level.LockedIcon16x", new IMAGE_BRUSH("Icons/icon_locked_16px", Icon16x16));
+		Set("Level.LockedHighlightIcon16x", new IMAGE_BRUSH("Icons/icon_locked_highlight_16px", Icon16x16));
+		Set("Level.UnlockedIcon16x", new IMAGE_BRUSH("Icons/icon_levels_unlocked_16px", Icon16x16));
+		Set("Level.UnlockedHighlightIcon16x", new IMAGE_BRUSH("Icons/icon_levels_unlocked_hi_16px", Icon16x16));
+		Set("Level.ReadOnlyLockedIcon16x", new IMAGE_BRUSH("Icons/icon_levels_LockedReadOnly_16px", Icon16x16));
+		Set("Level.ReadOnlyLockedHighlightIcon16x", new IMAGE_BRUSH("Icons/icon_levels_LockedReadOnly_hi_16px", Icon16x16));
+		Set("Level.SaveIcon16x", new IMAGE_BRUSH("Icons/icon_levels_Save_16px", Icon16x16));
+		Set("Level.SaveHighlightIcon16x", new IMAGE_BRUSH("Icons/icon_levels_Save_hi_16px", Icon16x16));
+		Set("Level.SaveModifiedIcon16x", new IMAGE_BRUSH("Icons/icon_levels_SaveModified_16px", Icon16x16));
+		Set("Level.SaveModifiedHighlightIcon16x", new IMAGE_BRUSH("Icons/icon_levels_SaveModified_hi_16px", Icon16x16));
+		Set("Level.SaveDisabledIcon16x", new IMAGE_BRUSH("Icons/icon_levels_SaveDisabled_16px", Icon16x16));
+		Set("Level.SaveDisabledHighlightIcon16x", new IMAGE_BRUSH("Icons/icon_levels_SaveDisabled_hi_16px", Icon16x16));
+		Set("Level.ScriptIcon16x", new IMAGE_BRUSH("Icons/icon_levels_Blueprint_16px", Icon16x16));
+		Set("Level.ScriptHighlightIcon16x", new IMAGE_BRUSH("Icons/icon_levels_Blueprint_hi_16px", Icon16x16));
+		Set("Level.EmptyIcon16x", new IMAGE_BRUSH("Icons/Empty_16x", Icon16x16));
+		Set("Level.ColorIcon40x", new IMAGE_BRUSH("Icons/icon_levels_back_16px", Icon16x16));
 	}
 }
 
@@ -6079,7 +6079,7 @@ void FStarshipEditorStyle::FStyle::SetupContentBrowserStyle()
 		Set( "ContentBrowser.PrimitiveCube", new IMAGE_BRUSH( "ContentBrowser/ThumbnailCube", Icon32x32 ) );
 		Set( "ContentBrowser.PrimitivePlane", new IMAGE_BRUSH( "ContentBrowser/ThumbnailPlane", Icon32x32 ) );
 		Set( "ContentBrowser.PrimitiveCylinder", new IMAGE_BRUSH( "ContentBrowser/ThumbnailCylinder", Icon32x32 ) );
-		Set( "ContentBrowser.ResetPrimitiveToDefault", new IMAGE_BRUSH("ContentBrowser/ThumbnailReset", Icon32x32) );
+		Set( "ContentBrowser.ResetPrimitiveToDefault", new CORE_IMAGE_BRUSH_SVG("Starship/Common/Undo", Icon20x20) );
 
 		Set( "ContentBrowser.TopBar.Font", FTextBlockStyle( NormalText )
 			.SetFont( DEFAULT_FONT( "Bold", 11 ) )
@@ -6092,19 +6092,9 @@ void FStarshipEditorStyle::FStyle::SetupContentBrowserStyle()
 			.SetFont(DEFAULT_FONT("Regular", 7)));
 
 
-		// New Asset
-		Set( "ContentBrowser.NewAsset", new IMAGE_BRUSH( "Icons/icon_file_new_40x", Icon25x25 ) );
-
-		Set( "ContentBrowser.PathActions.NewAsset", new IMAGE_BRUSH( "Icons/icon_file_new_16px", Icon16x16 ) );
-		Set( "ContentBrowser.PathActions.SetColor", new IMAGE_BRUSH( "Icons/icon_Cascade_Color_40x", Icon16x16 ) );
-
-	
-
-		Set( "ContentBrowser.SaveDirtyPackages", new IMAGE_BRUSH( "Icons/icon_file_saveall_40x", Icon25x25 ) );
 		Set( "ContentBrowser.AddContent", new IMAGE_BRUSH( "Icons/icon_AddContent_40x", Icon25x25 ) );
 		Set( "ContentBrowser.ImportPackage", new IMAGE_BRUSH( "Icons/icon_Import_40x", Icon25x25 ) );
-		Set( "ContentBrowser.ImportIcon", new IMAGE_BRUSH( "Icons/icon_Import_16x", Icon16x16 ) );
-
+	
 		// Asset Context Menu
 		Set( "ContentBrowser.AssetActions", new IMAGE_BRUSH( "Icons/icon_tab_Tools_16x", Icon16x16 ) );
 		Set( "ContentBrowser.AssetActions.Edit", new IMAGE_BRUSH( "Icons/Edit/icon_Edit_16x", Icon16x16 ) );
diff --git a/Engine/Source/Runtime/SlateCore/Private/Styling/AppStyle.cpp b/Engine/Source/Runtime/SlateCore/Private/Styling/AppStyle.cpp
--- a/Engine/Source/Runtime/SlateCore/Private/Styling/AppStyle.cpp
+++ b/Engine/Source/Runtime/SlateCore/Private/Styling/AppStyle.cpp
@@ -24,7 +24,7 @@ void FAppStyle::SetAppStyleSetName(const FName& InName)
 	AppStyleName = InName;
 }
 
-const FName& FAppStyle::GetAppStyleSetName()
+const FName FAppStyle::GetAppStyleSetName()
 {
 	return AppStyleName;
 }
diff --git a/Engine/Source/Runtime/SlateCore/Public/Styling/AppStyle.h b/Engine/Source/Runtime/SlateCore/Public/Styling/AppStyle.h
--- a/Engine/Source/Runtime/SlateCore/Public/Styling/AppStyle.h
+++ b/Engine/Source/Runtime/SlateCore/Public/Styling/AppStyle.h
@@ -27,7 +27,7 @@ public:
 
 	static const ISlateStyle& Get();
 	
-	static const FName& GetAppStyleSetName();
+	static const FName GetAppStyleSetName();
 	static void SetAppStyleSetName(const FName& InName);
 
 	static void SetAppStyleSet(const ISlateStyle& InStyle);
