Moved stylable colors to advanced

diff --git a/Engine/Source/Editor/EditorStyle/Private/EditorStyleSettingsCustomization.cpp b/Engine/Source/Editor/EditorStyle/Private/EditorStyleSettingsCustomization.cpp
--- a/Engine/Source/Editor/EditorStyle/Private/EditorStyleSettingsCustomization.cpp
+++ b/Engine/Source/Editor/EditorStyle/Private/EditorStyleSettingsCustomization.cpp
@@ -54,6 +54,6 @@ void FEditorStyleSettingsCustomization::CustomizeDetails(IDetailLayoutBuilder& D
 
 	TArray<UObject*> Objects = { &UStyleColorTable::Get() };
 
-	ColorCategory.AddExternalObjectProperty(Objects, "ColorList");
+	ColorCategory.AddExternalObjectProperty(Objects, "Colors", EPropertyLocation::Advanced);
 }
 
