Added basic theme editing support and a Dark theme as an example.
Themes are basic json files which can change the named colors that the editor uses for its UI.  They work similarly to config files in that the same theme in a user dir overrides the project dir and the base dir.
Themes only work in the editor or standalone programs. Not supported in UMG or shipping games

diff --git a/Engine/Source/Editor/EditorStyle/Private/EditorStyleModule.cpp b/Engine/Source/Editor/EditorStyle/Private/EditorStyleModule.cpp
--- a/Engine/Source/Editor/EditorStyle/Private/EditorStyleModule.cpp
+++ b/Engine/Source/Editor/EditorStyle/Private/EditorStyleModule.cpp
@@ -23,8 +23,9 @@ public:
 	{
 		if (FCoreStyle::IsStarshipStyle())
 		{
-			// This is necessary because the core style  loads the color table before ProcessNewlyLoadedUObjects is called which means none of the config properties are in the classs property link at that time.
-			UStyleColorTable::Get().ReloadConfig();
+#if ALLOW_THEMES
+			USlateThemeManager::Get().ValidateActiveTheme();
+#endif
 
 			bUsingStarshipStyle = true;
 			FStarshipEditorStyle::Initialize();
diff --git a/Engine/Source/Editor/EditorStyle/Private/EditorStyleSettingsCustomization.cpp b/Engine/Source/Editor/EditorStyle/Private/EditorStyleSettingsCustomization.cpp
--- a/Engine/Source/Editor/EditorStyle/Private/EditorStyleSettingsCustomization.cpp
+++ b/Engine/Source/Editor/EditorStyle/Private/EditorStyleSettingsCustomization.cpp
@@ -5,6 +5,220 @@
 #include "Styling/StyleColors.h"
 #include "DetailLayoutBuilder.h"
 #include "IDetailChildrenBuilder.h"
+#include "Widgets/Input/STextComboBox.h"
+#include "DetailWidgetRow.h"
+#include "Widgets/Images/SImage.h"
+#include "Widgets/Input/SButton.h"
+#include "Modules/ModuleManager.h"
+#include "Widgets/Input/SEditableTextBox.h"
+
+#if ALLOW_THEMES
+#include "IDesktopPlatform.h"
+#include "DesktopPlatformModule.h"
+
+#define LOCTEXT_NAMESPACE "ThemeEditor"
+
+TWeakPtr<SWindow> ThemeEditorWindow;
+
+
+class SThemeEditor : public SCompoundWidget
+{
+	SLATE_BEGIN_ARGS(SThemeEditor)
+	{}
+		SLATE_EVENT(FOnThemeEditorClosed, OnThemeEditorClosed);
+	SLATE_END_ARGS()
+
+public:
+	void Construct(const FArguments& InArgs, TSharedRef<SWindow> InParentWindow)
+	{
+		OnThemeEditorClosed = InArgs._OnThemeEditorClosed;
+
+		ParentWindow = InParentWindow;
+		InParentWindow->SetOnWindowClosed(FOnWindowClosed::CreateSP(this, &SThemeEditor::OnParentWindowClosed));
+
+		FPropertyEditorModule& PropertyEditorModule = FModuleManager::Get().GetModuleChecked<FPropertyEditorModule>("PropertyEditor");
+
+		FDetailsViewArgs DetailsViewArgs(
+			/*bUpdateFromSelection=*/ false,
+			/*bLockable=*/ false,
+			/*bAllowSearch=*/ false,
+			FDetailsViewArgs::HideNameArea,
+			/*bHideSelectionTip=*/ true,
+			/*InNotifyHook=*/ nullptr,
+			/*InSearchInitialKeyFocus=*/ false,
+			/*InViewIdentifier=*/ NAME_None);
+		DetailsViewArgs.bShowOptions = false;
+
+		TSharedRef<IDetailsView> DetailsView = PropertyEditorModule.CreateDetailView(DetailsViewArgs);
+
+		DetailsView->SetIsPropertyVisibleDelegate(FIsPropertyVisible::CreateLambda([](const FPropertyAndParent& PropertyAndParent)
+			{
+				static const FName CurrentThemeIdName("CurrentThemeId");
+
+				return PropertyAndParent.Property.GetFName() != CurrentThemeIdName;
+			})
+		);
+
+		DetailsView->SetObject(&USlateThemeManager::Get());
+		ChildSlot
+		[
+			SNew(SBorder)
+			.BorderImage(FAppStyle::Get().GetBrush("Brushes.Background"))
+			[
+				SNew(SVerticalBox)
+				+ SVerticalBox::Slot()
+				.Padding(6.0f, 3.0f)
+				.AutoHeight()
+				[
+					SNew(SHorizontalBox)
+					+ SHorizontalBox::Slot()
+					.FillWidth(0.6f)
+					.HAlign(HAlign_Right)
+					.VAlign(VAlign_Center)
+					.Padding(5.0f, 2.0f)
+					[
+						SNew(STextBlock)
+						.Text(LOCTEXT("ThemeName", "Name"))
+					]
+					+ SHorizontalBox::Slot()
+					.FillWidth(2.0f)
+					.VAlign(VAlign_Center)
+					.Padding(5.0f, 2.0f)
+					[
+						SNew(SEditableTextBox)
+						.Text(this, &SThemeEditor::GetThemeName)
+						.OnTextCommitted(this, &SThemeEditor::OnThemeNameChanged)
+						//.IsReadOnly(true)
+					]
+				]
+				/*+ SVerticalBox::Slot()
+				.Padding(6.0f, 3.0f)
+				.AutoHeight()
+				[
+
+					SNew(SHorizontalBox)
+					+ SHorizontalBox::Slot()
+					.FillWidth(0.6f)
+					.HAlign(HAlign_Right)
+					.VAlign(VAlign_Center)
+					.Padding(5.0f, 2.0f)
+					[
+						SNew(STextBlock)
+						.Text(LOCTEXT("ThemeDescription", "Description"))
+					]
+					+ SHorizontalBox::Slot()
+					.FillWidth(2.0f)
+					.VAlign(VAlign_Center)
+					.Padding(5.0f, 2.0f)
+					[
+						SNew(SEditableTextBox)
+						.Text(FText::FromString("Test Theme Description"))
+						.IsReadOnly(true)
+					]
+				]*/
+				+ SVerticalBox::Slot()
+				.Padding(6.0f, 3.0f)
+				[
+					DetailsView
+				]
+				+ SVerticalBox::Slot()
+				.AutoHeight()
+				.HAlign(HAlign_Right)
+				.VAlign(VAlign_Bottom)
+				.Padding(6.0f, 3.0f)
+				[
+					SNew(SHorizontalBox)
+					+ SHorizontalBox::Slot()
+					.AutoWidth()
+					.VAlign(VAlign_Bottom)
+					.Padding(4, 3)
+					[
+						SNew(SButton)
+						.Text(LOCTEXT("SaveThemeButton", "Save"))
+						.ButtonStyle(FAppStyle::Get(), "PrimaryButton")
+						.OnClicked(this, &SThemeEditor::OnSaveClicked)
+					]
+					+ SHorizontalBox::Slot()
+					.AutoWidth()
+					.VAlign(VAlign_Bottom)
+					.Padding(4, 3)
+					[
+						SNew(SButton)
+						.Text(LOCTEXT("CancelThemeEditingButton", "Cancel"))
+						.OnClicked(this, &SThemeEditor::OnCancelClicked)
+					]
+				]
+			]
+		];
+
+	}
+
+private:
+	FText GetThemeName() const
+	{
+		return USlateThemeManager::Get().GetCurrentTheme().DisplayName;
+	}
+
+	void OnThemeNameChanged(const FText& NewName, ETextCommit::Type Type)
+	{
+		USlateThemeManager::Get().SetCurrentThemeDisplayName(NewName);
+	}
+
+	FReply OnSaveClicked()
+	{
+		static bool bIsRentrant = false;
+
+		FString Filename;
+
+		const FStyleTheme& Theme = USlateThemeManager::Get().GetCurrentTheme();
+		if (Theme.Filename.IsEmpty())
+		{
+			TSharedPtr<SWindow> MyWindow = FSlateApplication::Get().FindWidgetWindow(SharedThis(this));
+
+			TArray<FString> Filenames;
+			IDesktopPlatform* DesktopPlatform = FDesktopPlatformModule::Get();
+
+			if (DesktopPlatform->SaveFileDialog(MyWindow->GetNativeWindow()->GetOSWindowHandle(), TEXT("Save Theme As"), USlateThemeManager::Get().GetEngineThemeDir(), Theme.DisplayName.ToString() + TEXT(".json"), TEXT("Theme Files (*.json)|*.json"), EFileDialogFlags::None, Filenames))
+			{
+				Filename = Filenames[0];
+			}
+		}
+		else
+		{
+			Filename = Theme.Filename;
+		}
+
+		USlateThemeManager::Get().SaveCurrentThemeAs(Filename);
+
+		ParentWindow.Pin()->SetOnWindowClosed(FOnWindowClosed());
+		ParentWindow.Pin()->RequestDestroyWindow();
+
+		OnThemeEditorClosed.ExecuteIfBound(true);
+		return FReply::Handled();
+	}
+
+	FReply OnCancelClicked()
+	{
+		ParentWindow.Pin()->SetOnWindowClosed(FOnWindowClosed());
+		ParentWindow.Pin()->RequestDestroyWindow();
+
+		OnThemeEditorClosed.ExecuteIfBound(false);
+		return FReply::Handled();
+	}
+
+	void OnParentWindowClosed(const TSharedRef<SWindow>&)
+	{
+		OnCancelClicked();
+	}
+
+private:
+	FOnThemeEditorClosed OnThemeEditorClosed;
+	TWeakPtr<SWindow> ParentWindow;
+};
+
+#undef LOCTEXT_NAMESPACE
+
+#define LOCTEXT_NAMESPACE "EditorStyleSettingsCustomization"
 
 TSharedRef<IPropertyTypeCustomization> FStyleColorListCustomization::MakeInstance()
 {
@@ -13,7 +227,7 @@ TSharedRef<IPropertyTypeCustomization> FStyleColorListCustomization::MakeInstanc
 
 void FStyleColorListCustomization::CustomizeHeader(TSharedRef<IPropertyHandle> PropertyHandle, FDetailWidgetRow& HeaderRow, IPropertyTypeCustomizationUtils& CustomizationUtils)
 {
-	
+
 }
 
 void FStyleColorListCustomization::CustomizeChildren(TSharedRef<IPropertyHandle> PropertyHandle, IDetailChildrenBuilder& ChildBuilder, IPropertyTypeCustomizationUtils& CustomizationUtils)
@@ -24,24 +238,53 @@ void FStyleColorListCustomization::CustomizeChildren(TSharedRef<IPropertyHandle>
 	ColorArrayProperty->GetNumChildren(NumChildren);
 	for (uint32 ChildIndex = 0; ChildIndex < NumChildren; ++ChildIndex)
 	{
+		FResetToDefaultOverride ResetToDefaultOverride =
+			FResetToDefaultOverride::Create(
+				FIsResetToDefaultVisible::CreateSP(this, &FStyleColorListCustomization::IsResetToDefaultVisible, (EStyleColor)ChildIndex),
+				FResetToDefaultHandler::CreateSP(this, &FStyleColorListCustomization::OnResetColorToDefault, (EStyleColor)ChildIndex));
+
 		if (ChildIndex < (uint32)EStyleColor::User1)
 		{
-			ChildBuilder.AddProperty(ColorArrayProperty->GetChildHandle(ChildIndex).ToSharedRef());
+			IDetailPropertyRow& Row = ChildBuilder.AddProperty(ColorArrayProperty->GetChildHandle(ChildIndex).ToSharedRef());
+			Row.OverrideResetToDefault(ResetToDefaultOverride);
 		}
 		else
 		{
 			// user colors are added if they have been customized with a display name
-			FText DisplayName = UStyleColorTable::Get().GetColorDisplayName((EStyleColor)ChildIndex);
+			FText DisplayName = USlateThemeManager::Get().GetColorDisplayName((EStyleColor)ChildIndex);
 			if (!DisplayName.IsEmpty())
 			{
 				IDetailPropertyRow& Row = ChildBuilder.AddProperty(ColorArrayProperty->GetChildHandle(ChildIndex).ToSharedRef());
 				Row.DisplayName(DisplayName);
+				Row.OverrideResetToDefault(ResetToDefaultOverride);
 			}
 		}
 	}
 }
 
 
+void FStyleColorListCustomization::OnResetColorToDefault(TSharedPtr<IPropertyHandle> Handle, EStyleColor Color)
+{
+	FLinearColor CurrentColor = USlateThemeManager::Get().GetColor(Color);
+	const FStyleTheme& Theme = USlateThemeManager::Get().GetCurrentTheme();
+	if (Theme.LoadedDefaultColors.Num())
+	{
+		USlateThemeManager::Get().ResetActiveColorToDefault(Color);
+	}
+}
+
+bool FStyleColorListCustomization::IsResetToDefaultVisible(TSharedPtr<IPropertyHandle> Handle, EStyleColor Color)
+{
+	FLinearColor CurrentColor = USlateThemeManager::Get().GetColor(Color);
+	const FStyleTheme& Theme = USlateThemeManager::Get().GetCurrentTheme();
+	if (Theme.LoadedDefaultColors.Num())
+	{
+		return Theme.LoadedDefaultColors[(int32)Color] != CurrentColor;
+	}
+
+	return false;
+}
+
 TSharedRef<IDetailCustomization> FEditorStyleSettingsCustomization::MakeInstance()
 {
 	return MakeShared<FEditorStyleSettingsCustomization>();
@@ -52,8 +295,220 @@ void FEditorStyleSettingsCustomization::CustomizeDetails(IDetailLayoutBuilder& D
 {
 	IDetailCategoryBuilder& ColorCategory = DetailLayout.EditCategory("Colors");
 
-	TArray<UObject*> Objects = { &UStyleColorTable::Get() };
+	TArray<UObject*> Objects = { &USlateThemeManager::Get() };
 
-	ColorCategory.AddExternalObjectProperty(Objects, "Colors", EPropertyLocation::Advanced);
+	if (IDetailPropertyRow* ThemeRow = ColorCategory.AddExternalObjectProperty(Objects, "CurrentThemeId"))
+	{
+		MakeThemePickerRow(*ThemeRow);
+	}
+}
+
+void FEditorStyleSettingsCustomization::RefreshComboBox()
+{
+	TSharedPtr<FString> SelectedTheme;
+	GenerateThemeOptions(SelectedTheme);
+	ComboBox->RefreshOptions();
+	ComboBox->SetSelectedItem(SelectedTheme);
 }
 
+void FEditorStyleSettingsCustomization::GenerateThemeOptions(TSharedPtr<FString>& OutSelectedTheme)
+{
+	const TArray<FStyleTheme>& Themes = USlateThemeManager::Get().GetThemes();
+
+	ThemeOptions.Empty(Themes.Num());
+	int32 Index = 0;
+	for (const FStyleTheme& Theme : Themes)
+	{
+		TSharedRef<FString> ThemeString = MakeShared<FString>(FString::FromInt(Index));
+
+		if (USlateThemeManager::Get().GetCurrentTheme() == Theme)
+		{
+			OutSelectedTheme = ThemeString;
+		}
+
+		ThemeOptions.Add(ThemeString);
+		++Index;
+	}
+
+}
+
+void FEditorStyleSettingsCustomization::MakeThemePickerRow(IDetailPropertyRow& PropertyRow)
+{
+
+	TSharedPtr<FString> SelectedItem;
+	GenerateThemeOptions(SelectedItem);
+
+	// Make combo choices
+	ComboBox =
+		SNew(STextComboBox)
+		.OptionsSource(&ThemeOptions)
+		.InitiallySelectedItem(SelectedItem)
+		.Font(IDetailLayoutBuilder::GetDetailFont())
+		.OnGetTextLabelForItem(this, &FEditorStyleSettingsCustomization::GetTextLabelForThemeEntry)
+		.OnSelectionChanged(this, &FEditorStyleSettingsCustomization::OnThemePicked);
+
+
+	FDetailWidgetRow& CustomWidgetRow = PropertyRow.CustomWidget(false);
+
+	CustomWidgetRow
+	.NameContent()
+	[
+		PropertyRow.GetPropertyHandle()->CreatePropertyNameWidget(LOCTEXT("ActiveThemeDisplayName", "Active Theme"))
+	]
+	.ValueContent()
+	.MaxDesiredWidth(350.f)
+	[
+		SNew(SHorizontalBox)
+		.IsEnabled(this, &FEditorStyleSettingsCustomization::IsThemeEditingEnabled)
+		+SHorizontalBox::Slot()
+		[
+			SNew(SBox)
+			.WidthOverride(125.f)
+			[
+				ComboBox.ToSharedRef()
+			]
+		]
+		+ SHorizontalBox::Slot()
+		.VAlign(VAlign_Center)
+		.HAlign(HAlign_Center)
+		.AutoWidth()
+		[
+			SNew(SButton)
+			.ButtonStyle(FAppStyle::Get(), "SimpleButton")
+			.ToolTipText(LOCTEXT("EditThemeToolTip", "Edit this theme"))
+			.OnClicked(this, &FEditorStyleSettingsCustomization::OnEditThemeClicked)
+			[
+				SNew(SImage)
+				.ColorAndOpacity(FSlateColor::UseForeground())
+				.Image(FAppStyle::Get().GetBrush("Icons.Edit"))
+			]
+		]
+		+SHorizontalBox::Slot()
+		.VAlign(VAlign_Center)
+		.HAlign(HAlign_Center)
+		.AutoWidth()
+		[
+			SNew(SButton)
+			.ButtonStyle(FAppStyle::Get(), "SimpleButton")
+			.ToolTipText(LOCTEXT("DuplicateThemeToolTip", "Duplicate this theme and edit it"))
+			.OnClicked(this, &FEditorStyleSettingsCustomization::OnDuplicateAndEditThemeClicked)
+			[
+				SNew(SImage)
+				.ColorAndOpacity(FSlateColor::UseForeground())
+				.Image(FAppStyle::Get().GetBrush("Icons.Duplicate"))
+			]
+		]
+	];
+}
+
+static void OnThemeEditorClosed(bool bSaved, TWeakPtr<FEditorStyleSettingsCustomization> ActiveCustomization, FGuid CreatedThemeId, FGuid PreviousThemeId)
+{
+	if (!bSaved)
+	{
+		if (PreviousThemeId.IsValid())
+		{
+	
+			USlateThemeManager::Get().ApplyTheme(PreviousThemeId);
+
+			if (CreatedThemeId.IsValid())
+			{
+				USlateThemeManager::Get().RemoveTheme(CreatedThemeId);
+
+			}
+			if (ActiveCustomization.IsValid())
+			{
+
+				ActiveCustomization.Pin()->RefreshComboBox();
+			}
+		}
+		else
+		{
+			for (int32 ColorIndex = 0; ColorIndex < (int32)EStyleColor::MAX; ++ColorIndex)
+			{
+				USlateThemeManager::Get().ResetActiveColorToDefault((EStyleColor)ColorIndex);
+			}
+		}
+	}
+}
+
+FReply FEditorStyleSettingsCustomization::OnDuplicateAndEditThemeClicked()
+{
+	FGuid PreviouslyActiveTheme = USlateThemeManager::Get().GetCurrentTheme().Id;
+
+	FGuid NewThemeId = USlateThemeManager::Get().DuplicateActiveTheme();
+	USlateThemeManager::Get().ApplyTheme(NewThemeId);
+
+	RefreshComboBox();
+
+	OpenThemeEditorWindow(FOnThemeEditorClosed::CreateStatic(&OnThemeEditorClosed, TWeakPtr<FEditorStyleSettingsCustomization>(SharedThis(this)), NewThemeId, PreviouslyActiveTheme));
+
+	return FReply::Handled();
+}
+
+FReply FEditorStyleSettingsCustomization::OnEditThemeClicked()
+{
+	OpenThemeEditorWindow(FOnThemeEditorClosed::CreateStatic(&OnThemeEditorClosed, TWeakPtr<FEditorStyleSettingsCustomization>(SharedThis(this)), FGuid(), FGuid()));
+
+	return FReply::Handled();
+}
+
+FString FEditorStyleSettingsCustomization::GetTextLabelForThemeEntry(TSharedPtr<FString> Entry)
+{
+	const TArray<FStyleTheme>& Themes = USlateThemeManager::Get().GetThemes();
+	return Themes[TCString<TCHAR>::Atoi(**Entry)].DisplayName.ToString();
+}
+
+void FEditorStyleSettingsCustomization::OnThemePicked(TSharedPtr<FString> NewSelection, ESelectInfo::Type SelectInfo)
+{
+	// If set directly in code, the theme was already applied
+	if(SelectInfo != ESelectInfo::Direct)
+	{
+		const TArray<FStyleTheme>& Themes = USlateThemeManager::Get().GetThemes();
+
+		USlateThemeManager::Get().ApplyTheme(Themes[TCString<TCHAR>::Atoi(**NewSelection)].Id);
+	}
+}
+
+void FEditorStyleSettingsCustomization::OpenThemeEditorWindow(FOnThemeEditorClosed OnThemeEditorClosed)
+{
+	if(!ThemeEditorWindow.IsValid())
+	{
+		TSharedRef<SWindow> NewWindow = SNew(SWindow)
+			.Title(LOCTEXT("ThemeEditorWindowTitle", "Theme Editor"))
+			.SizingRule(ESizingRule::UserSized)
+			.ClientSize(FVector2D(600, 600))
+			.SupportsMaximize(false)
+			.SupportsMinimize(false);
+
+		TSharedRef<SThemeEditor> ThemeEditor =
+			SNew(SThemeEditor, NewWindow)
+			.OnThemeEditorClosed(OnThemeEditorClosed);
+			
+		NewWindow->SetContent(
+			ThemeEditor
+		);
+
+		if (TSharedPtr<SWindow> ParentWindow = FSlateApplication::Get().FindWidgetWindow(ComboBox.ToSharedRef()))
+		{
+			FSlateApplication::Get().AddWindowAsNativeChild(NewWindow, ParentWindow.ToSharedRef());
+		}
+		else
+		{
+			FSlateApplication::Get().AddWindow(NewWindow);
+		}
+
+		ThemeEditorWindow = NewWindow;
+	}
+
+
+}
+
+bool FEditorStyleSettingsCustomization::IsThemeEditingEnabled() const
+{
+	// Don't allow changing themes while editing them
+	return !ThemeEditorWindow.IsValid();
+}
+
+#undef LOCTEXT_NAMESPACE
+
+#endif // ALLOW_THEMES
diff --git a/Engine/Source/Editor/EditorStyle/Private/EditorStyleSettingsCustomization.h b/Engine/Source/Editor/EditorStyle/Private/EditorStyleSettingsCustomization.h
--- a/Engine/Source/Editor/EditorStyle/Private/EditorStyleSettingsCustomization.h
+++ b/Engine/Source/Editor/EditorStyle/Private/EditorStyleSettingsCustomization.h
@@ -4,8 +4,15 @@
 
 #include "IDetailCustomization.h"
 #include "IPropertyTypeCustomization.h"
+#include "Styling/StyleColors.h"
 
 class IDetailLayoutBuilder;
+class STextComboBox;
+class IDetailPropertyRow;
+
+#if ALLOW_THEMES
+
+DECLARE_DELEGATE_OneParam(FOnThemeEditorClosed, bool)
 
 class FStyleColorListCustomization : public IPropertyTypeCustomization
 {
@@ -14,7 +21,9 @@ public:
 
 	virtual void CustomizeHeader(TSharedRef<IPropertyHandle> PropertyHandle, FDetailWidgetRow& HeaderRow, IPropertyTypeCustomizationUtils& CustomizationUtils) override;
 	virtual void CustomizeChildren(TSharedRef<IPropertyHandle> PropertyHandle, IDetailChildrenBuilder& ChildBuilder, IPropertyTypeCustomizationUtils& CustomizationUtils) override;
-
+private:
+	void OnResetColorToDefault(TSharedPtr<IPropertyHandle> Handle, EStyleColor Color);
+	bool IsResetToDefaultVisible(TSharedPtr<IPropertyHandle> Handle, EStyleColor Color);
 };
 
 class FEditorStyleSettingsCustomization : public IDetailCustomization
@@ -25,5 +34,20 @@ public:
 
 	/** IDetailCustomization interface */
 	virtual void CustomizeDetails(IDetailLayoutBuilder& DetailLayout) override;
+	
+	void RefreshComboBox();
+private:
+	void GenerateThemeOptions(TSharedPtr<FString>& OutSelectedTheme);
+
+	void MakeThemePickerRow(IDetailPropertyRow& PropertyRow);
+	FReply OnDuplicateAndEditThemeClicked();
+	FReply OnEditThemeClicked();
+	FString GetTextLabelForThemeEntry(TSharedPtr<FString> Entry);
+	void OnThemePicked(TSharedPtr<FString> NewSelection, ESelectInfo::Type SelectInfo);
+	void OpenThemeEditorWindow(FOnThemeEditorClosed OnThemeEditorClosed);
+	bool IsThemeEditingEnabled() const;
+private:
+	TArray<TSharedPtr<FString>> ThemeOptions;
+	TSharedPtr<STextComboBox> ComboBox;
 };
-
+#endif
diff --git a/Engine/Source/Editor/EditorStyle/Private/StarshipStyle.cpp b/Engine/Source/Editor/EditorStyle/Private/StarshipStyle.cpp
--- a/Engine/Source/Editor/EditorStyle/Private/StarshipStyle.cpp
+++ b/Engine/Source/Editor/EditorStyle/Private/StarshipStyle.cpp
@@ -46,27 +46,33 @@ void FStarshipEditorStyle::Initialize()
 	FPropertyEditorModule& PropertyEditorModule = FModuleManager::Get().GetModuleChecked<FPropertyEditorModule>("PropertyEditor");
 	PropertyEditorModule.RegisterCustomClassLayout("EditorStyleSettings", FOnGetDetailCustomizationInstance::CreateStatic(&FEditorStyleSettingsCustomization::MakeInstance));
 	PropertyEditorModule.RegisterCustomPropertyTypeLayout("StyleColorList", FOnGetPropertyTypeCustomizationInstance::CreateStatic(&FStyleColorListCustomization::MakeInstance));
-#endif
 
 	// Background
-	UStyleColorTable::Get().SetColor(EStyleColor::User1, UStyleColorTable::Get().GetColor(EStyleColor::Input));
-	UStyleColorTable::Get().SetColorDisplayName(EStyleColor::User1, LOCTEXT("UserColor_OutputLogBackground", "Log Background"));
+	USlateThemeManager::Get().SetDefaultColor(EStyleColor::User1, USlateThemeManager::Get().GetColor(EStyleColor::Input));
+	USlateThemeManager::Get().SetColorDisplayName(EStyleColor::User1, LOCTEXT("UserColor_OutputLogBackground", "Log Background"));
 	// Selection highlight
-	UStyleColorTable::Get().SetColor(EStyleColor::User2, UStyleColorTable::Get().GetColor(EStyleColor::Highlight));
-	UStyleColorTable::Get().SetColorDisplayName(EStyleColor::User2, LOCTEXT("UserColor_OutputLogHighlight", "Log Highlight"));
+	USlateThemeManager::Get().SetDefaultColor(EStyleColor::User2, USlateThemeManager::Get().GetColor(EStyleColor::Highlight));
+	USlateThemeManager::Get().SetColorDisplayName(EStyleColor::User2, LOCTEXT("UserColor_OutputLogHighlight", "Log Highlight"));
 	// Normal
-	UStyleColorTable::Get().SetColor(EStyleColor::User3, UStyleColorTable::Get().GetColor(EStyleColor::Foreground));
-	UStyleColorTable::Get().SetColorDisplayName(EStyleColor::User3, LOCTEXT("UserColor_OutputLogText", "Log Text"));
+	USlateThemeManager::Get().SetDefaultColor(EStyleColor::User3, USlateThemeManager::Get().GetColor(EStyleColor::Foreground));
+	USlateThemeManager::Get().SetColorDisplayName(EStyleColor::User3, LOCTEXT("UserColor_OutputLogText", "Log Text"));
 	// Command
-	UStyleColorTable::Get().SetColor(EStyleColor::User4, UStyleColorTable::Get().GetColor(EStyleColor::AccentGreen));
-	UStyleColorTable::Get().SetColorDisplayName(EStyleColor::User4, LOCTEXT("UserColor_OutputLogCommand", "Log Command"));
+	USlateThemeManager::Get().SetDefaultColor(EStyleColor::User4, USlateThemeManager::Get().GetColor(EStyleColor::AccentGreen));
+	USlateThemeManager::Get().SetColorDisplayName(EStyleColor::User4, LOCTEXT("UserColor_OutputLogCommand", "Log Command"));
 	// Warning
-	UStyleColorTable::Get().SetColor(EStyleColor::User5, UStyleColorTable::Get().GetColor(EStyleColor::AccentYellow));
-	UStyleColorTable::Get().SetColorDisplayName(EStyleColor::User5, LOCTEXT("UserColor_OutputLogWarning", "Log Warning"));
+	USlateThemeManager::Get().SetDefaultColor(EStyleColor::User5, USlateThemeManager::Get().GetColor(EStyleColor::AccentYellow));
+	USlateThemeManager::Get().SetColorDisplayName(EStyleColor::User5, LOCTEXT("UserColor_OutputLogWarning", "Log Warning"));
 	// Error
-	UStyleColorTable::Get().SetColor(EStyleColor::User6, UStyleColorTable::Get().GetColor(EStyleColor::AccentRed));
-	UStyleColorTable::Get().SetColorDisplayName(EStyleColor::User6, LOCTEXT("UserColor_OutputLogError", "Log Error"));
+	USlateThemeManager::Get().SetDefaultColor(EStyleColor::User6, USlateThemeManager::Get().GetColor(EStyleColor::AccentRed));
+	USlateThemeManager::Get().SetColorDisplayName(EStyleColor::User6, LOCTEXT("UserColor_OutputLogError", "Log Error"));
+#endif
 
+	const FString ThemesSubDir = TEXT("Slate/Themes");
+
+#if ALLOW_THEMES
+	USlateThemeManager::Get().ApplyTheme(USlateThemeManager::Get().GetCurrentTheme().Id);
+	//UStyleColorTable::Get().SaveCurrentThemeAs(UStyleColorTable::Get().GetCurrentTheme().Filename);
+#endif
 	StyleInstance = Create(Settings);
 	SetStyle(StyleInstance.ToSharedRef());
 }
diff --git a/Engine/Source/Runtime/Slate/Private/Framework/Application/SlateApplication.cpp b/Engine/Source/Runtime/Slate/Private/Framework/Application/SlateApplication.cpp
--- a/Engine/Source/Runtime/Slate/Private/Framework/Application/SlateApplication.cpp
+++ b/Engine/Source/Runtime/Slate/Private/Framework/Application/SlateApplication.cpp
@@ -6702,7 +6702,9 @@ void FSlateApplication::InitializeCoreStyle()
 {
 	if (FCoreStyle::IsStarshipStyle() && !FStarshipCoreStyle::IsInitialized())
 	{
-		UStyleColorTable::Get();
+#if ALLOW_THEMES
+		USlateThemeManager::Get().LoadThemes();
+#endif
 		FStarshipCoreStyle::ResetToDefault();
 		FAppStyle::SetAppStyleSet(FStarshipCoreStyle::GetCoreStyle());
 	}
diff --git a/Engine/Source/Runtime/SlateCore/Private/Styling/SlateColor.cpp b/Engine/Source/Runtime/SlateCore/Private/Styling/SlateColor.cpp
--- a/Engine/Source/Runtime/SlateCore/Private/Styling/SlateColor.cpp
+++ b/Engine/Source/Runtime/SlateCore/Private/Styling/SlateColor.cpp
@@ -32,5 +32,5 @@ bool FSlateColor::SerializeFromMismatchedTag(const FPropertyTag& Tag, FStructure
 
 const FLinearColor& FSlateColor::GetColorFromTable() const
 {
-	return UStyleColorTable::Get().GetColor(ColorTableId);
+	return USlateThemeManager::Get().GetColor(ColorTableId);
 }
diff --git a/Engine/Source/Runtime/SlateCore/Private/Styling/StarshipCoreStyle.cpp b/Engine/Source/Runtime/SlateCore/Private/Styling/StarshipCoreStyle.cpp
--- a/Engine/Source/Runtime/SlateCore/Private/Styling/StarshipCoreStyle.cpp
+++ b/Engine/Source/Runtime/SlateCore/Private/Styling/StarshipCoreStyle.cpp
@@ -364,6 +364,10 @@ TSharedRef<ISlateStyle> FStarshipCoreStyle::Create()
 
 		Style->Set("Icons.FilledCircle", new IMAGE_BRUSH_SVG("Starship/Common/filled-circle", Icon16x16));
 
+		Style->Set("Icons.Duplicate", new IMAGE_BRUSH_SVG("Starship/Common/Duplicate", Icon16x16));
+		Style->Set("Icons.Edit", new IMAGE_BRUSH_SVG("Starship/Common/edit", Icon16x16));
+
+
 	}
 
 	// Tool panels
@@ -661,12 +665,6 @@ TSharedRef<ISlateStyle> FStarshipCoreStyle::Create()
 			.SetPressed(IMAGE_BRUSH_SVG("Starship/CoreWidgets/Window/close", FVector2D(42.0f, 34.0f), FStyleColors::Foreground));
 #endif
 
-		const FTextBlockStyle TitleTextStyle = FTextBlockStyle(NormalText)
-			.SetFont(StyleFonts.Normal)
-			.SetColorAndOpacity(FLinearColor::White)
-			.SetShadowOffset(FVector2D(1.0f, 1.0f))
-			.SetShadowColorAndOpacity(FLinearColor::Black);
-
 		Style->Set("Window", FWindowStyle()
 
 #if !PLATFORM_MAC
@@ -675,7 +673,7 @@ TSharedRef<ISlateStyle> FStarshipCoreStyle::Create()
 			.SetRestoreButtonStyle(RestoreButtonStyle)
 			.SetCloseButtonStyle(CloseButtonStyle)
 #endif
-			.SetTitleTextStyle(TitleTextStyle)
+			.SetTitleTextStyle(NormalText)
 			.SetActiveTitleBrush(FSlateNoResource())
 			.SetInactiveTitleBrush(FSlateNoResource())
 			.SetFlashTitleBrush(IMAGE_BRUSH("Common/Window/WindowTitle_Flashing", Icon24x24, FLinearColor(1.0f, 1.0f, 1.0f, 1.0f), ESlateBrushTileType::Horizontal))	
diff --git a/Engine/Source/Runtime/SlateCore/Private/Styling/StyleColors.cpp b/Engine/Source/Runtime/SlateCore/Private/Styling/StyleColors.cpp
--- a/Engine/Source/Runtime/SlateCore/Private/Styling/StyleColors.cpp
+++ b/Engine/Source/Runtime/SlateCore/Private/Styling/StyleColors.cpp
@@ -1,7 +1,14 @@
 // Copyright Epic Games, Inc. All Rights Reserved.
 
 #include "Styling/StyleColors.h"
+#include "Misc/Paths.h"
+#include "Serialization/JsonWriter.h"
+#include "Misc/FileHelper.h"
+#include "Serialization/JsonReader.h"
+#include "Serialization/JsonSerializer.h"
+#include "HAL/PlatformFileManager.h"
 
+// Note this value is not mutable by the user
 const FSlateColor FStyleColors::Transparent = FSlateColor(FLinearColor::Transparent);
 
 const FSlateColor FStyleColors::Black = EStyleColor::Black;
@@ -50,56 +57,346 @@ const FSlateColor FStyleColors::AccentWhite = EStyleColor::AccentWhite;
 const FSlateColor FStyleColors::AccentFolder = EStyleColor::AccentFolder;
 
 
-UStyleColorTable::UStyleColorTable()
+USlateThemeManager::USlateThemeManager()
 {
 	InitalizeDefaults();
 }
 
-void UStyleColorTable::InitalizeDefaults()
-{
-	SetColor(EStyleColor::Black, COLOR("#000000FF"));
-	SetColor(EStyleColor::Title, COLOR("#151515FF"));
-	SetColor(EStyleColor::WindowBorder, COLOR("0F0F0FFF"));
-	SetColor(EStyleColor::Foldout, COLOR("0F0F0FFF"));
-	SetColor(EStyleColor::Input, COLOR("0F0F0FFF"));
-	SetColor(EStyleColor::Recessed, COLOR("#1A1A1AFF"));
-	SetColor(EStyleColor::Background, COLOR("#242424FF"));
-	SetColor(EStyleColor::Header, COLOR("#2F2F2FFF"));
-	SetColor(EStyleColor::Dropdown, COLOR("#383838FF"));
-	SetColor(EStyleColor::Hover, COLOR("#575757FF"));
-	SetColor(EStyleColor::Hover2, COLOR("#808080FF"));
-
-	SetColor(EStyleColor::White, COLOR("#FFFFFFFF"));
-	SetColor(EStyleColor::White25, COLOR("#FFFFFF40"));
-	SetColor(EStyleColor::Highlight, COLOR("#0078D7FF"));
-
-	SetColor(EStyleColor::Primary, COLOR("#26BBFFFF"));
-	SetColor(EStyleColor::PrimaryHover, COLOR("#6FD2FFFF"));
-	SetColor(EStyleColor::PrimaryPress, COLOR("#1989BCFF"));
-	SetColor(EStyleColor::Secondary, COLOR("#383838FF"));
-
-	SetColor(EStyleColor::Foreground, COLOR("#A6A6A6FF"));
-	SetColor(EStyleColor::ForegroundHover, COLOR("#FFFFFFFF"));
-	SetColor(EStyleColor::ForegroundInverted, GetColor(EStyleColor::Input));
-	SetColor(EStyleColor::ForegroundHeader, COLOR("#C8C8C8FF"));
-
-	SetColor(EStyleColor::Select, GetColor(EStyleColor::Primary));
-	SetColor(EStyleColor::SelectInactive, COLOR("#99B3BFFF"));
-	SetColor(EStyleColor::SelectParent, COLOR("#2C323AFF"));
-	SetColor(EStyleColor::SelectHover, GetColor(EStyleColor::Background));
+void USlateThemeManager::InitalizeDefaults()
+{
+	// This loads defaults for a dark theme
+
+	SetDefaultColor(EStyleColor::Black, COLOR("#000000FF"));
+	SetDefaultColor(EStyleColor::Title, COLOR("#151515FF"));
+	SetDefaultColor(EStyleColor::WindowBorder, COLOR("0F0F0FFF"));
+	SetDefaultColor(EStyleColor::Foldout, COLOR("0F0F0FFF"));
+	SetDefaultColor(EStyleColor::Input, COLOR("0F0F0FFF"));
+	SetDefaultColor(EStyleColor::Recessed, COLOR("#1A1A1AFF"));
+	SetDefaultColor(EStyleColor::Background, COLOR("#242424FF"));
+	SetDefaultColor(EStyleColor::Header, COLOR("#2F2F2FFF"));
+	SetDefaultColor(EStyleColor::Dropdown, COLOR("#383838FF"));
+	SetDefaultColor(EStyleColor::Hover, COLOR("#575757FF"));
+	SetDefaultColor(EStyleColor::Hover2, COLOR("#808080FF"));
+
+	SetDefaultColor(EStyleColor::White, COLOR("#FFFFFFFF"));
+	SetDefaultColor(EStyleColor::White25, COLOR("#FFFFFF40"));
+	SetDefaultColor(EStyleColor::Highlight, COLOR("#0078D7FF"));
+
+	SetDefaultColor(EStyleColor::Primary, COLOR("#26BBFFFF"));
+	SetDefaultColor(EStyleColor::PrimaryHover, COLOR("#6FD2FFFF"));
+	SetDefaultColor(EStyleColor::PrimaryPress, COLOR("#1989BCFF"));
+	SetDefaultColor(EStyleColor::Secondary, COLOR("#383838FF"));
+
+	SetDefaultColor(EStyleColor::Foreground, COLOR("#A6A6A6FF"));
+	SetDefaultColor(EStyleColor::ForegroundHover, COLOR("#FFFFFFFF"));
+	SetDefaultColor(EStyleColor::ForegroundInverted, GetDefaultColor(EStyleColor::Input));
+	SetDefaultColor(EStyleColor::ForegroundHeader, COLOR("#C8C8C8FF"));
+
+	SetDefaultColor(EStyleColor::Select, GetDefaultColor(EStyleColor::Primary));
+	SetDefaultColor(EStyleColor::SelectInactive, COLOR("#99B3BFFF"));
+	SetDefaultColor(EStyleColor::SelectParent, COLOR("#2C323AFF"));
+	SetDefaultColor(EStyleColor::SelectHover, GetDefaultColor(EStyleColor::Background));
 	// if select ==  primary shouldnt we have a select pressed which is the same as primary press?
 
-	SetColor(EStyleColor::AccentBlue, COLOR("#26BBFFFF"));
-	SetColor(EStyleColor::AccentPurple, COLOR("#A139BFFF"));
-	SetColor(EStyleColor::AccentPink, COLOR("#FF729CFF"));
-	SetColor(EStyleColor::AccentRed, COLOR("#FF4040FF"));
-	SetColor(EStyleColor::AccentOrange, COLOR("#FE9B07FF"));
-	SetColor(EStyleColor::AccentYellow, COLOR("#FFDC1AFF"));
-	SetColor(EStyleColor::AccentGreen, COLOR("#8BC24AFF"));
-	SetColor(EStyleColor::AccentBrown, COLOR("#804D39FF"));
-	SetColor(EStyleColor::AccentBlack, COLOR("#242424FF"));
-	SetColor(EStyleColor::AccentGray, COLOR("#808080FF"));
-	SetColor(EStyleColor::AccentWhite, COLOR("#FFFFFFFF"));
-	SetColor(EStyleColor::AccentFolder, COLOR("#B68F55FF"));
+	SetDefaultColor(EStyleColor::AccentBlue, COLOR("#26BBFFFF"));
+	SetDefaultColor(EStyleColor::AccentPurple, COLOR("#A139BFFF"));
+	SetDefaultColor(EStyleColor::AccentPink, COLOR("#FF729CFF"));
+	SetDefaultColor(EStyleColor::AccentRed, COLOR("#FF4040FF"));
+	SetDefaultColor(EStyleColor::AccentOrange, COLOR("#FE9B07FF"));
+	SetDefaultColor(EStyleColor::AccentYellow, COLOR("#FFDC1AFF"));
+	SetDefaultColor(EStyleColor::AccentGreen, COLOR("#8BC24AFF"));
+	SetDefaultColor(EStyleColor::AccentBrown, COLOR("#804D39FF"));
+	SetDefaultColor(EStyleColor::AccentBlack, COLOR("#242424FF"));
+	SetDefaultColor(EStyleColor::AccentGray, COLOR("#808080FF"));
+	SetDefaultColor(EStyleColor::AccentWhite, COLOR("#FFFFFFFF"));
+	SetDefaultColor(EStyleColor::AccentFolder, COLOR("#B68F55FF"));
+}
+
+#if ALLOW_THEMES
+
+static const FString ThemesSubDir = TEXT("Slate/Themes");
+
+
+void USlateThemeManager::LoadThemes()
+{
+	Themes.Empty();
+
+	// Engine wide themes
+	LoadThemesFromDirectory(GetEngineThemeDir());
+
+	// Project themes
+	LoadThemesFromDirectory(GetProjectThemeDir());
+
+	// User specific themes
+	LoadThemesFromDirectory(FPaths::EngineVersionAgnosticUserDir());
+
+	EnsureValidCurrentTheme();
+
+
+	ApplyTheme(CurrentThemeId);
+}
+
+void USlateThemeManager::SaveCurrentThemeAs(const FString& Filename)
+{
+	FStyleTheme& CurrentTheme = GetMutableCurrentTheme();
+	CurrentTheme.Filename = Filename;
+	{
+		FString Output;
+		TSharedRef<TJsonWriter<>> WriterRef = TJsonWriterFactory<>::Create(&Output);
+		TJsonWriter<>& Writer = WriterRef.Get();
+		Writer.WriteObjectStart();
+		Writer.WriteValue(TEXT("Version"), 1);
+		Writer.WriteValue(TEXT("Id"), CurrentTheme.Id.ToString());
+		Writer.WriteValue(TEXT("DisplayName"), CurrentTheme.DisplayName.ToString());
+
+		{
+			Writer.WriteObjectStart(TEXT("Colors"));
+			UEnum* Enum = StaticEnum<EStyleColor>();
+			check(Enum);
+			for (int32 ColorIndex = 0; ColorIndex < (int32)EStyleColor::MAX; ++ColorIndex)
+			{
+				FName EnumName = Enum->GetNameByIndex(ColorIndex);
+				Writer.WriteValue(EnumName.ToString(), *ActiveColors.StyleColors[ColorIndex].ToString());
+			}
+			Writer.WriteObjectEnd();
+		}
+		Writer.WriteObjectEnd();
+		Writer.Close();
+
+		if (FPlatformFileManager::Get().GetPlatformFile().FileExists(*Filename))
+		{
+			FPlatformFileManager::Get().GetPlatformFile().SetReadOnly(*Filename, false);
+		}
+		FFileHelper::SaveStringToFile(Output, *Filename);
+	}
+	
+}
+
+void USlateThemeManager::ApplyTheme(FGuid ThemeId)
+{
+	if (ThemeId.IsValid())
+	{
+		FStyleTheme* CurrentTheme = nullptr;
+		if (CurrentThemeId != ThemeId)
+		{
+			if (CurrentThemeId.IsValid())
+			{
+				CurrentTheme = &GetMutableCurrentTheme();
+				// Unload existing colors
+				CurrentTheme->LoadedDefaultColors.Empty();
+			}
+
+			FStyleTheme* Theme = Themes.FindByKey(ThemeId);
+			if (Theme)
+			{
+				CurrentThemeId = ThemeId;
+				SaveConfig();
+			}
+		}
+
+		//if (CurrentTheme->LoadedDefaultColors.IsEmpty())
+		{
+			CurrentTheme = &GetMutableCurrentTheme();
+			LoadThemeColors(*CurrentTheme);
+		}
+
+		// Apply the new colors. Note that if the incoming theme is the same as the current theme we still apply the colors as they may have been overwritten by defaults
+		FMemory::Memcpy(ActiveColors.StyleColors, CurrentTheme->LoadedDefaultColors.GetData(), sizeof(FLinearColor)*CurrentTheme->LoadedDefaultColors.Num());
+	}
+}
+
+void USlateThemeManager::RemoveTheme(FGuid ThemeId)
+{
+	// Current Theme cannot currently be removed.  Apply a new theme first
+	if (CurrentThemeId != ThemeId)
+	{
+		Themes.RemoveAll([&ThemeId](const FStyleTheme& TestTheme) { return TestTheme.Id == ThemeId; });
+	}
+}
+
+FGuid USlateThemeManager::DuplicateActiveTheme()
+{
+	const FStyleTheme& CurrentTheme = GetCurrentTheme();
+
+	FGuid NewThemeGuid = FGuid::NewGuid();
+	FStyleTheme NewTheme;
+	NewTheme.Id = NewThemeGuid;
+	NewTheme.DisplayName = FText::Format(NSLOCTEXT("StyleColors", "ThemeDuplicateCopyText", "{0} - Copy"), CurrentTheme.DisplayName);
+	NewTheme.LoadedDefaultColors = MakeArrayView<FLinearColor>(ActiveColors.StyleColors, (int32)EStyleColor::MAX);
+
+	Themes.Add(MoveTemp(NewTheme));
+
+	return NewThemeGuid;
+}
+
+void USlateThemeManager::SetCurrentThemeDisplayName(FText NewDisplayName)
+{
+	GetMutableCurrentTheme().DisplayName = NewDisplayName;
+}
+
+void USlateThemeManager::ResetActiveColorToDefault(EStyleColor Color)
+{
+	ActiveColors.StyleColors[(int32)Color] = GetCurrentTheme().LoadedDefaultColors[(int32)Color];
+}
+
+void USlateThemeManager::ValidateActiveTheme()
+{
+	// This is necessary because the core style loads the color table before ProcessNewlyLoadedUObjects is called which means none of the config properties are in the class property link at that time.
+	ReloadConfig();
+	EnsureValidCurrentTheme();
+	ApplyTheme(USlateThemeManager::Get().GetCurrentTheme().Id);
 }
 
+FString USlateThemeManager::GetEngineThemeDir() const
+{
+	return FPaths::EngineContentDir() / ThemesSubDir;
+}
+
+FString USlateThemeManager::GetProjectThemeDir() const
+{
+	return FPaths::ProjectContentDir() / ThemesSubDir;
+}
+
+FString USlateThemeManager::GetUserThemeDir() const
+{
+	return FPaths::EngineVersionAgnosticUserDir() / ThemesSubDir;
+}
+
+void USlateThemeManager::LoadThemesFromDirectory(const FString& Directory)
+{
+	TArray<FString> ThemeFiles;
+	IFileManager::Get().FindFiles(ThemeFiles, *Directory, TEXT(".json"));
+
+	for (const FString& ThemeFile : ThemeFiles)
+	{
+		bool bValidFile = false;
+		FString ThemeData;
+		FString ThemeFilename = Directory / ThemeFile;
+		if (FFileHelper::LoadFileToString(ThemeData, *ThemeFilename))
+		{
+			FStyleTheme Theme;
+			if (ReadTheme(ThemeData, Theme))
+			{
+				if (FStyleTheme* ExistingTheme = Themes.FindByKey(Theme.Id))
+				{
+					// Just update the existing theme.  Themes with the same id can override an existing one.  This behavior mimics config file hierarchies
+					ExistingTheme->Filename = MoveTemp(ThemeFilename);
+				}
+				else
+				{
+					// Theme not found, add a new one
+					Theme.Filename = MoveTemp(ThemeFilename);
+					Themes.Add(MoveTemp(Theme));
+				}
+			}
+		}
+	}
+}
+
+
+bool USlateThemeManager::ReadTheme(const FString& ThemeData, FStyleTheme& Theme)
+{
+	TSharedRef<TJsonReader<>> ReaderRef = TJsonReaderFactory<>::Create(ThemeData);
+	TJsonReader<>& Reader = ReaderRef.Get();
+
+	TSharedPtr<FJsonObject> ObjectPtr;
+	if (FJsonSerializer::Deserialize(Reader, ObjectPtr) && ObjectPtr.IsValid())
+	{
+		int32 Version = 0;
+		if (!ObjectPtr->TryGetNumberField("Version", Version))
+		{
+			// Invalid file
+			return false;
+		}
+
+		FString IdString;
+		if (!ObjectPtr->TryGetStringField(TEXT("Id"), IdString) || !FGuid::Parse(IdString, Theme.Id))
+		{
+			// Invalid Id;
+			return false;
+		}
+
+		FString DisplayStr;
+		if (!ObjectPtr->TryGetStringField(TEXT("DisplayName"), DisplayStr))
+		{
+			// Invalid file
+			return false;
+		}
+		Theme.DisplayName = FText::FromString(MoveTemp(DisplayStr));
+
+		// Just check that the theme has colors. We wont load them unless the theme is used
+		if (!ObjectPtr->HasField("Colors"))
+		{
+			// No colors
+			return false;
+		}
+	}
+	else
+	{
+		// Log invalid style file
+		return false;
+	}
+
+	return true;
+}
+
+void USlateThemeManager::EnsureValidCurrentTheme()
+{
+	FStyleTheme DefaultDarkTheme;
+	DefaultDarkTheme.DisplayName = NSLOCTEXT("StyleColors", "DefaultDarkTheme", "Dark");
+	// If you change this you invalidate the default dark theme forcing the default theme to be reset and the existing dark theme to become a user theme
+	// In general you should not do this. You should instead update the default "dark.json" file 
+	DefaultDarkTheme.Id = FGuid(0xE431F929, 0x452C4686, 0xAB5C0FE1, 0x2910653A);
+	DefaultDarkTheme.Filename = FPaths::EngineContentDir() / TEXT("Slate/Themes/Dark.json");
+
+	int32 ThemeIndex = Themes.AddUnique(DefaultDarkTheme);
+
+	if (!CurrentThemeId.IsValid() || !Themes.Contains(CurrentThemeId))
+	{
+		CurrentThemeId = DefaultDarkTheme.Id;
+	}
+}
+
+void USlateThemeManager::LoadThemeColors(FStyleTheme& Theme)
+{
+	FString ThemeData;
+	if (FFileHelper::LoadFileToString(ThemeData, *Theme.Filename))
+	{
+		Theme.LoadedDefaultColors.Empty();
+
+		// Start with the hard coded default colors. They are a fallback if the theme is incomplete 
+		Theme.LoadedDefaultColors = MakeArrayView<FLinearColor>(DefaultColors, (int32)EStyleColor::MAX);
+
+		TSharedRef<TJsonReader<>> ReaderRef = TJsonReaderFactory<>::Create(ThemeData);
+		TJsonReader<>& Reader = ReaderRef.Get();
+
+		TSharedPtr<FJsonObject> ObjectPtr;
+		if (FJsonSerializer::Deserialize(Reader, ObjectPtr) && ObjectPtr.IsValid())
+		{		
+			// Just check that the theme has colors. We wont load them unless the theme is used
+			const TSharedPtr<FJsonObject>* ColorsObject = nullptr;
+			if(ObjectPtr->TryGetObjectField(TEXT("Colors"), ColorsObject))
+			{
+				UEnum* Enum = StaticEnum<EStyleColor>();
+				check(Enum);
+				for (int32 ColorIndex = 0; ColorIndex < (int32)EStyleColor::MAX; ++ColorIndex)
+				{
+					FName EnumName = Enum->GetNameByIndex(ColorIndex);
+					FString ColorString;
+					if ((*ColorsObject)->TryGetStringField(EnumName.ToString(), ColorString))
+					{
+						Theme.LoadedDefaultColors[ColorIndex].InitFromString(ColorString);
+					}
+				}
+			}
+		}
+	}
+}
+
+/*
+const FStyleTheme* UStyleColorTable::GetTheme(const FGuid ThemeId) const
+{
+	return Themes.FindByKey(ThemeId);
+}*/
+
+#endif // ALLOW_THEMES
\ No newline at end of file
diff --git a/Engine/Source/Runtime/SlateCore/Public/Styling/StyleColors.h b/Engine/Source/Runtime/SlateCore/Public/Styling/StyleColors.h
--- a/Engine/Source/Runtime/SlateCore/Public/Styling/StyleColors.h
+++ b/Engine/Source/Runtime/SlateCore/Public/Styling/StyleColors.h
@@ -6,10 +6,15 @@
 #include "SlateColor.h"
 #include "StyleColors.generated.h"
 
+/** Themes are only allowed in the editor or standalone tools */
+#define ALLOW_THEMES WITH_EDITOR || IS_PROGRAM
 
 // HEX Colors from sRGB Color space
 #define COLOR( HexValue ) FLinearColor::FromSRGBColor(FColor::FromHex(HexValue))
 
+/**
+ * Note: If you add another color here, you should update the Dark.json theme file in Engine\Content\Slate\Themes for consistency
+ */
 UENUM()
 enum class EStyleColor : uint8
 {
@@ -55,7 +60,7 @@ enum class EStyleColor : uint8
 	/** Only user colors should be below this line
 	 * To use user colors:
 	 * 1. Set an unused user enum value below as the color value for an FSlateColor. E.g. FSlateColor MyCustomColor(EStyleColors::User1)
-	 * 2. Set the actual color. E.g UStyleColorTable::Get().SetColor(EStyleColor::User1, FLinearColor::White)
+	 * 2. Set the actual color. E.g USlateThemeManager::Get().SetDefaultColor(EStyleColor::User1, FLinearColor::White)
 	 * 3. Give it a display name if you want it to be configurable by editor users. E.g.  UStyleColorTable::Get().SetColorDisplayName(EUserStyleColor::User1, "My Color Name")
 	 */
 	User1,
@@ -89,55 +94,180 @@ struct FStyleColorList
 	FText DisplayNames[(int32)EStyleColor::MAX];
 };
 
+/**
+ * Represents a single theme
+ */
+USTRUCT()
+struct FStyleTheme
+{
+	GENERATED_BODY()
+
+	/** Unique Id for the theme */
+	FGuid Id;
+	/** Friendly, user customizable theme name */
+	FText DisplayName;
+	/** Filename where the theme is stored */
+	FString Filename;
+	/** The default colors for this theme. Used for resetting to default. Not the active colors*/
+	TArray<FLinearColor> LoadedDefaultColors;
+
+	bool operator==(const FStyleTheme& Other) const
+	{
+		return Id == Other.Id;
+	}
+
+	bool operator==(const FGuid& OtherId) const
+	{
+		return Id == OtherId;
+	}
+};
+
 UCLASS(Config=EditorSettings)
-class SLATECORE_API UStyleColorTable : public UObject 
+class SLATECORE_API USlateThemeManager : public UObject 
 {
 	GENERATED_BODY()
 public:
-	static UStyleColorTable& Get()
+	static USlateThemeManager& Get()
 	{
-		return *GetMutableDefault<UStyleColorTable>();
+		return *GetMutableDefault<USlateThemeManager>();
 	}
 
 	const FLinearColor& GetColor(EStyleColor Color)
 	{
-		return Colors.StyleColors[static_cast<int32>(Color)];
+		return ActiveColors.StyleColors[static_cast<int32>(Color)];
 	}
 
-	void SetColor(EStyleColor InColorId, FLinearColor InColor)
+
+	USlateThemeManager();
+
+	/**
+	 * Initializes default colors
+	 */
+	void InitalizeDefaults();
+
+	/**
+	 * Sets a default color to be used as a fallback if no theme is loaded
+	 */
+	void SetDefaultColor(EStyleColor InColorId, FLinearColor InColor)
 	{
-		Colors.StyleColors[static_cast<int32>(InColorId)] = InColor;
+#if ALLOW_THEMES
+		DefaultColors[static_cast<int32>(InColorId)] = InColor;
+#else
+		ActiveColors.StyleColors[static_cast<int32>(InColorId)] = InColor;
+#endif
 	}
 
+#if ALLOW_THEMES
+
+	/** Sets a custom display name for a style color */
 	void SetColorDisplayName(EStyleColor InColorId, FText DisplayName)
 	{
-		Colors.DisplayNames[static_cast<int32>(InColorId)] = DisplayName;
+		ActiveColors.DisplayNames[static_cast<int32>(InColorId)] = DisplayName;
 	}
 
+	/** Gets a custom display name for a style color. This will be empty if no custom name was chosen */
 	FText GetColorDisplayName(EStyleColor InColorId) const
 	{
-		return Colors.DisplayNames[static_cast<int32>(InColorId)];
+		return ActiveColors.DisplayNames[static_cast<int32>(InColorId)];
 	}
 
-	UStyleColorTable();
+	/**
+	 * Load all known themes from engine, project, and user directories
+	 */
+	void LoadThemes();
 
-#if WITH_EDITOR
-	virtual void PostEditChangeProperty(struct FPropertyChangedEvent& PropertyChangedEvent) override
-	{
-		Super::PostEditChangeProperty(PropertyChangedEvent);
+	/**
+	 * Saves the current theme
+	 * 
+	 * @param The filename to save the current theme as
+	 */
+	void SaveCurrentThemeAs(const FString& Filename);
 
-		if (PropertyChangedEvent.ChangeType != EPropertyChangeType::Interactive)
-		{
-			SaveConfig();
-		}
-	}
-#endif
+	/**
+	 * Applies a theme as the active theme
+	 */
+	void ApplyTheme(FGuid ThemeId);
 
-	void InitalizeDefaults();
+	/**
+	 * Removes a theme. 
+	 * Note: The active theme cannot be removed and there must always be an active theme.  Apply a new theme first before removing the current theme.
+	 */
+	void RemoveTheme(FGuid ThemeId);
+
+	/**
+	 * Duplicates the active theme
+	 * 
+	 * @return the id of the new theme
+	 */
+	FGuid DuplicateActiveTheme();
+
+	/**
+	 * Sets the display name for the current theme
+	 */
+	void SetCurrentThemeDisplayName(FText NewDisplayName);
+
+	/**
+	 * @return the current theme
+	 */
+	const FStyleTheme& GetCurrentTheme() const { return *Themes.FindByKey(CurrentThemeId); }
+
+	/**
+	 * @return All known themes
+	 */
+	const TArray<FStyleTheme>& GetThemes() const { return Themes; }
+
+	/**
+	 * Resets an active color to the default color for the curerent theme
+	 */
+	void ResetActiveColorToDefault(EStyleColor Color);
+
+	/**
+	 * Validate that there is an active loaded theme
+	 */
+	void ValidateActiveTheme();
+
+	/**
+	 * @return the engine theme dir.  Engine themes are project agnostic
+	 */
+	FString GetEngineThemeDir() const;
+
+	/**
+	 * @return the project theme dir. Project themes can override engine themes
+	 */
+	FString GetProjectThemeDir() const;
 
+	/**
+	 * @return the user theme dir. Themes in this dir are per-user and override engine and project themes
+	 */
+	FString GetUserThemeDir() const;
 private:
-	UPROPERTY(EditAnywhere, Config, Category = Colors)
-	FStyleColorList Colors;
+
+	FStyleTheme& GetMutableCurrentTheme() { return *Themes.FindByKey(CurrentThemeId); }
+	void LoadThemesFromDirectory(const FString& Directory);
+	bool ReadTheme(const FString& ThemeData, FStyleTheme& OutTheme);
+	void EnsureValidCurrentTheme();
+	void LoadThemeColors(FStyleTheme& Theme);
+
+
+private:
+	TArray<FStyleTheme> Themes;
+	FLinearColor DefaultColors[(int32)EStyleColor::MAX];
+#endif // ALLOW_THEMES
+
+	FLinearColor GetDefaultColor(EStyleColor InColorId)
+	{
+#if ALLOW_THEMES
+		return DefaultColors[static_cast<int32>(InColorId)];
+#else
+		return ActiveColors.StyleColors[static_cast<int32>(InColorId)];
+#endif
+	}
+
+	UPROPERTY(EditAnywhere, Config, Category=Colors)
+	FGuid CurrentThemeId;
+
+	UPROPERTY(EditAnywhere, transient, Category = Colors)
+	FStyleColorList ActiveColors;
 };
 
 /**
