Fix hitting cancel while choosing a theme file location, closing the dialog and creating an empty theme

diff --git a/Engine/Source/Editor/EditorStyle/Private/EditorStyleSettingsCustomization.cpp b/Engine/Source/Editor/EditorStyle/Private/EditorStyleSettingsCustomization.cpp
--- a/Engine/Source/Editor/EditorStyle/Private/EditorStyleSettingsCustomization.cpp
+++ b/Engine/Source/Editor/EditorStyle/Private/EditorStyleSettingsCustomization.cpp
@@ -166,8 +166,6 @@ private:
 
 	FReply OnSaveClicked()
 	{
-		static bool bIsRentrant = false;
-
 		FString Filename;
 
 		const FStyleTheme& Theme = USlateThemeManager::Get().GetCurrentTheme();
@@ -188,12 +186,15 @@ private:
 			Filename = Theme.Filename;
 		}
 
-		USlateThemeManager::Get().SaveCurrentThemeAs(Filename);
+		if(!Filename.IsEmpty())
+		{
+			USlateThemeManager::Get().SaveCurrentThemeAs(Filename);
 
-		ParentWindow.Pin()->SetOnWindowClosed(FOnWindowClosed());
-		ParentWindow.Pin()->RequestDestroyWindow();
+			ParentWindow.Pin()->SetOnWindowClosed(FOnWindowClosed());
+			ParentWindow.Pin()->RequestDestroyWindow();
 
-		OnThemeEditorClosed.ExecuteIfBound(true);
+			OnThemeEditorClosed.ExecuteIfBound(true);
+		}
 		return FReply::Handled();
 	}
 
