Fix hard to read console hint text highlight colors

diff --git a/Engine/Source/Developer/OutputLog/Private/SOutputLog.cpp b/Engine/Source/Developer/OutputLog/Private/SOutputLog.cpp
--- a/Engine/Source/Developer/OutputLog/Private/SOutputLog.cpp
+++ b/Engine/Source/Developer/OutputLog/Private/SOutputLog.cpp
@@ -219,6 +219,7 @@ TSharedRef<ITableRow> SConsoleInputBox::MakeSuggestionListItemWidget(TSharedPtr<
 			.Text(FText::FromString(SanitizedText))
 			.TextStyle(FEditorStyle::Get(), "Log.Normal")
 			.HighlightText(Suggestions.SuggestionsHighlight)
+			.ColorAndOpacity(FSlateColor::UseForeground())
 		];
 }
 
diff --git a/Engine/Source/Editor/EditorStyle/Private/StarshipStyle.cpp b/Engine/Source/Editor/EditorStyle/Private/StarshipStyle.cpp
--- a/Engine/Source/Editor/EditorStyle/Private/StarshipStyle.cpp
+++ b/Engine/Source/Editor/EditorStyle/Private/StarshipStyle.cpp
@@ -967,7 +967,7 @@ void FStarshipEditorStyle::FStyle::SetupGeneralStyles()
 			.SetFont(DEFAULT_FONT("Mono", LogFontSize))
 			.SetColorAndOpacity(LogColor_Normal)
 			.SetSelectedBackgroundColor(LogColor_SelectionBackground)
-			.SetHighlightColor(FLinearColor::White);
+			.SetHighlightColor(FStyleColors::Black);
 
 		Set("Log.Normal", NormalLogText );
 
