Pass over actor icons for content browser, placement, and outliners

diff --git a/Engine/Plugins/2D/Paper2D/Source/Paper2DEditor/Private/PaperStyle.cpp b/Engine/Plugins/2D/Paper2D/Source/Paper2DEditor/Private/PaperStyle.cpp
--- a/Engine/Plugins/2D/Paper2D/Source/Paper2DEditor/Private/PaperStyle.cpp
+++ b/Engine/Plugins/2D/Paper2D/Source/Paper2DEditor/Private/PaperStyle.cpp
@@ -10,6 +10,7 @@
 
 
 #define IMAGE_PLUGIN_BRUSH( RelativePath, ... ) FSlateImageBrush( FPaperStyle::InContent( RelativePath, ".png" ), __VA_ARGS__ )
+#define IMAGE_PLUGIN_BRUSH_SVG( RelativePath, ... ) FSlateVectorImageBrush( FPaperStyle::InContent( RelativePath, ".svg" ), __VA_ARGS__ )
 #define IMAGE_BRUSH(RelativePath, ...) FSlateImageBrush(StyleSet->RootToContentDir(RelativePath, TEXT(".png")), __VA_ARGS__)
 #define BOX_BRUSH(RelativePath, ...) FSlateBoxBrush(StyleSet->RootToContentDir(RelativePath, TEXT(".png")), __VA_ARGS__)
 #define DEFAULT_FONT(...) FCoreStyle::GetDefaultFontStyle(__VA_ARGS__)
@@ -262,28 +263,28 @@ void FPaperStyle::Initialize()
 
 	{
 		// Sprites (asset, component, actor)
-		StyleSet->Set("ClassIcon.PaperSprite", new IMAGE_PLUGIN_BRUSH("Icons/AssetIcons/PaperSprite_16x", Icon16x16));
-		StyleSet->Set("ClassThumbnail.PaperSprite", new IMAGE_PLUGIN_BRUSH("Icons/AssetIcons/PaperSprite_64x", Icon64x64));
-		StyleSet->Set("ClassIcon.PaperSpriteComponent", new IMAGE_PLUGIN_BRUSH("Icons/AssetIcons/PaperSpriteComponent_16x", Icon16x16));
-		StyleSet->Set("ClassThumbnail.PaperSpriteComponent", new IMAGE_PLUGIN_BRUSH("Icons/AssetIcons/PaperSpriteComponent_64x", Icon64x64));
-		StyleSet->Set("ClassIcon.PaperSpriteActor", new IMAGE_PLUGIN_BRUSH("Icons/AssetIcons/PaperSpriteActor_16x", Icon16x16));
-		StyleSet->Set("ClassThumbnail.PaperSpriteActor", new IMAGE_PLUGIN_BRUSH("Icons/AssetIcons/PaperSpriteActor_64x", Icon64x64));
+		StyleSet->Set("ClassIcon.PaperSprite", new IMAGE_PLUGIN_BRUSH_SVG("Icons/AssetIcons/PaperSpriteActor_16", Icon16x16));
+		StyleSet->Set("ClassThumbnail.PaperSprite", new IMAGE_PLUGIN_BRUSH_SVG("Icons/AssetIcons/PaperSpriteActor_64", Icon64x64));
+		StyleSet->Set("ClassIcon.PaperSpriteComponent", new IMAGE_PLUGIN_BRUSH_SVG("Icons/AssetIcons/PaperSpriteActor_16", Icon16x16));
+		StyleSet->Set("ClassThumbnail.PaperSpriteComponent", new IMAGE_PLUGIN_BRUSH_SVG("Icons/AssetIcons/PaperSpriteActor_64", Icon64x64));
+		StyleSet->Set("ClassIcon.PaperSpriteActor", new IMAGE_PLUGIN_BRUSH_SVG("Icons/AssetIcons/PaperSpriteActor_16", Icon16x16));
+		StyleSet->Set("ClassThumbnail.PaperSpriteActor", new IMAGE_PLUGIN_BRUSH_SVG("Icons/AssetIcons/PaperSpriteActor_64", Icon64x64));
 
 		// Flipbooks (asset, component, actor)
-		StyleSet->Set("ClassIcon.PaperFlipbook", new IMAGE_PLUGIN_BRUSH("Icons/AssetIcons/PaperFlipbook_16x", Icon16x16));
-		StyleSet->Set("ClassThumbnail.PaperFlipbook", new IMAGE_PLUGIN_BRUSH("Icons/AssetIcons/PaperFlipbook_64x", Icon64x64));
-		StyleSet->Set("ClassIcon.PaperFlipbookComponent", new IMAGE_PLUGIN_BRUSH("Icons/AssetIcons/PaperFlipbookComponent_16x", Icon16x16));
-		StyleSet->Set("ClassThumbnail.PaperFlipbookComponent", new IMAGE_PLUGIN_BRUSH("Icons/AssetIcons/PaperFlipbookComponent_64x", Icon64x64));
-		StyleSet->Set("ClassIcon.PaperFlipbookActor", new IMAGE_PLUGIN_BRUSH("Icons/AssetIcons/PaperFlipbookActor_16x", Icon16x16));
-		StyleSet->Set("ClassThumbnail.PaperFlipbookActor", new IMAGE_PLUGIN_BRUSH("Icons/AssetIcons/PaperFlipbookActor_64x", Icon64x64));
+		StyleSet->Set("ClassIcon.PaperFlipbook", new IMAGE_PLUGIN_BRUSH_SVG("Icons/AssetIcons/PaperFlipbookActor_16", Icon16x16));
+		StyleSet->Set("ClassThumbnail.PaperFlipbook", new IMAGE_PLUGIN_BRUSH_SVG("Icons/AssetIcons/PaperFlipbookActor_64", Icon64x64));
+		StyleSet->Set("ClassIcon.PaperFlipbookComponent", new IMAGE_PLUGIN_BRUSH_SVG("Icons/AssetIcons/PaperFlipbookActor_16", Icon16x16));
+		StyleSet->Set("ClassThumbnail.PaperFlipbookComponent", new IMAGE_PLUGIN_BRUSH_SVG("Icons/AssetIcons/PaperFlipbookActor_64", Icon64x64));
+		StyleSet->Set("ClassIcon.PaperFlipbookActor", new IMAGE_PLUGIN_BRUSH_SVG("Icons/AssetIcons/PaperFlipbookActor_16", Icon16x16));
+		StyleSet->Set("ClassThumbnail.PaperFlipbookActor", new IMAGE_PLUGIN_BRUSH_SVG("Icons/AssetIcons/PaperFlipbookActor_64", Icon64x64));
 
 		// Tile maps (asset, component, actor)
-		StyleSet->Set("ClassIcon.PaperTileMap", new IMAGE_PLUGIN_BRUSH("Icons/AssetIcons/PaperTileMap_16x", Icon16x16));
-		StyleSet->Set("ClassThumbnail.PaperTileMap", new IMAGE_PLUGIN_BRUSH("Icons/AssetIcons/PaperTileMap_64x", Icon64x64));
-		StyleSet->Set("ClassIcon.PaperTileMapComponent", new IMAGE_PLUGIN_BRUSH("Icons/AssetIcons/PaperTileMapComponent_16x", Icon16x16));
-		StyleSet->Set("ClassThumbnail.PaperTileMapComponent", new IMAGE_PLUGIN_BRUSH("Icons/AssetIcons/PaperTileMapComponent_64x", Icon64x64));
-		StyleSet->Set("ClassIcon.PaperTileMapActor", new IMAGE_PLUGIN_BRUSH("Icons/AssetIcons/PaperTileMapActor_16x", Icon16x16));
-		StyleSet->Set("ClassThumbnail.PaperTileMapActor", new IMAGE_PLUGIN_BRUSH("Icons/AssetIcons/PaperTileMapActor_64x", Icon64x64));
+		StyleSet->Set("ClassIcon.PaperTileMap", new IMAGE_PLUGIN_BRUSH_SVG("Icons/AssetIcons/PaperTileMapActor_16", Icon16x16));
+		StyleSet->Set("ClassThumbnail.PaperTileMap", new IMAGE_PLUGIN_BRUSH_SVG("Icons/AssetIcons/PaperTileMapActor_64", Icon64x64));
+		StyleSet->Set("ClassIcon.PaperTileMapComponent", new IMAGE_PLUGIN_BRUSH_SVG("Icons/AssetIcons/PaperTileMapActor_16", Icon16x16));
+		StyleSet->Set("ClassThumbnail.PaperTileMapComponent", new IMAGE_PLUGIN_BRUSH_SVG("Icons/AssetIcons/PaperTileMapActor_64", Icon64x64));
+		StyleSet->Set("ClassIcon.PaperTileMapActor", new IMAGE_PLUGIN_BRUSH_SVG("Icons/AssetIcons/PaperTileMapActor_16", Icon16x16));
+		StyleSet->Set("ClassThumbnail.PaperTileMapActor", new IMAGE_PLUGIN_BRUSH_SVG("Icons/AssetIcons/PaperTileMapActor_64", Icon64x64));
 
 		// UPaperSpriteAtlas assets
 		//@TODO: Paper2D: These icons don't match the naming scheme
@@ -293,22 +294,22 @@ void FPaperStyle::Initialize()
 		//@TODO: UPaperSpriteSheet icons?
 
 		// APaperCharacter icons
-		StyleSet->Set("ClassIcon.PaperCharacter", new IMAGE_PLUGIN_BRUSH("Icons/AssetIcons/PaperCharacter_16x", Icon16x16));
-		StyleSet->Set("ClassThumbnail.PaperCharacter", new IMAGE_PLUGIN_BRUSH("Icons/AssetIcons/PaperCharacter_64x", Icon64x64));
+		StyleSet->Set("ClassIcon.PaperCharacter", new IMAGE_PLUGIN_BRUSH_SVG("Icons/AssetIcons/PaperCharacter_16", Icon16x16));
+		StyleSet->Set("ClassThumbnail.PaperCharacter", new IMAGE_PLUGIN_BRUSH_SVG("Icons/AssetIcons/PaperCharacter_64", Icon64x64));
 
 		// UPaperTileSet icons
-		StyleSet->Set("ClassIcon.PaperTileSet", new IMAGE_PLUGIN_BRUSH("Icons/AssetIcons/PaperTileSet_16x", Icon16x16));
-		StyleSet->Set("ClassThumbnail.PaperTileSet", new IMAGE_PLUGIN_BRUSH("Icons/AssetIcons/PaperTileSet_64x", Icon64x64));
+		StyleSet->Set("ClassIcon.PaperTileSet", new IMAGE_PLUGIN_BRUSH_SVG("Icons/AssetIcons/PaperTileMapActor_16", Icon16x16));
+		StyleSet->Set("ClassThumbnail.PaperTileSet", new IMAGE_PLUGIN_BRUSH_SVG("Icons/AssetIcons/PaperTileMapActor_64", Icon64x64));
 
 		// UPaperTerrainMaterial icons
 		StyleSet->Set("ClassIcon.PaperTerrainMaterial", new IMAGE_PLUGIN_BRUSH("Icons/AssetIcons/PaperTerrainMaterial_16x", Icon16x16));
 		StyleSet->Set("ClassThumbnail.PaperTerrainMaterial", new IMAGE_PLUGIN_BRUSH("Icons/AssetIcons/PaperTerrainMaterial_64x", Icon64x64));
 
 		// Terrain splines (component, actor)
-		StyleSet->Set("ClassIcon.PaperTerrainComponent", new IMAGE_PLUGIN_BRUSH("Icons/AssetIcons/PaperTerrainComponent_16x", Icon16x16));
-		StyleSet->Set("ClassThumbnail.PaperTerrainComponent", new IMAGE_PLUGIN_BRUSH("Icons/AssetIcons/PaperTerrainComponent_64x", Icon64x64));
-		StyleSet->Set("ClassIcon.PaperTerrainActor", new IMAGE_PLUGIN_BRUSH("Icons/AssetIcons/PaperTerrainActor_16x", Icon16x16));
-		StyleSet->Set("ClassThumbnail.PaperTerrainActor", new IMAGE_PLUGIN_BRUSH("Icons/AssetIcons/PaperTerrainActor_16x", Icon64x64));
+		StyleSet->Set("ClassIcon.PaperTerrainComponent", new IMAGE_PLUGIN_BRUSH_SVG("Icons/AssetIcons/PaperTerrainActor_16", Icon16x16));
+		StyleSet->Set("ClassThumbnail.PaperTerrainComponent", new IMAGE_PLUGIN_BRUSH_SVG("Icons/AssetIcons/PaperTerrainActor_16", Icon64x64));
+		StyleSet->Set("ClassIcon.PaperTerrainActor", new IMAGE_PLUGIN_BRUSH_SVG("Icons/AssetIcons/PaperTerrainActor_16", Icon16x16));
+		StyleSet->Set("ClassThumbnail.PaperTerrainActor", new IMAGE_PLUGIN_BRUSH_SVG("Icons/AssetIcons/PaperTerrainActor_16", Icon64x64));
 	}
 
 	FSlateStyleRegistry::RegisterSlateStyle(*StyleSet.Get());
diff --git a/Engine/Source/Editor/AudioEditor/Private/AudioEditorModule.cpp b/Engine/Source/Editor/AudioEditor/Private/AudioEditorModule.cpp
--- a/Engine/Source/Editor/AudioEditor/Private/AudioEditorModule.cpp
+++ b/Engine/Source/Editor/AudioEditor/Private/AudioEditorModule.cpp
@@ -65,6 +65,11 @@ static const FVector2D Icon64 = FVector2D(64.0f, 64.0f);
 // Simpler version of SET_AUDIO_ICON, assumes same name of icon png and class name
 #define SET_AUDIO_ICON_SIMPLE(CLASS_NAME) SET_AUDIO_ICON(CLASS_NAME, CLASS_NAME)
 
+#define SET_AUDIO_ICON_SVG(CLASS_NAME, ICON_NAME) \
+		AudioStyleSet->Set( *FString::Printf(TEXT("ClassIcon.%s"), TEXT(#CLASS_NAME)), new FSlateVectorImageBrush(FPaths::EngineContentDir() / FString::Printf(TEXT("Editor/Slate/Starship/AssetIcons/%s_16.svg"), TEXT(#ICON_NAME)), Icon16)); \
+		AudioStyleSet->Set( *FString::Printf(TEXT("ClassThumbnail.%s"), TEXT(#CLASS_NAME)), new FSlateVectorImageBrush(FPaths::EngineContentDir() / FString::Printf(TEXT("Editor/Slate/Starship/AssetIcons/%s_64.svg"), TEXT(#ICON_NAME)), Icon64));
+
+#define SET_AUDIO_ICON_SVG_SIMPLE(CLASS_NAME) SET_AUDIO_ICON_SVG(CLASS_NAME, CLASS_NAME)
 
 class FAudioEditorModule : public IAudioEditorModule
 {
@@ -366,12 +371,12 @@ private:
 	void SetupIcons()
 	{
 		SET_AUDIO_ICON_SIMPLE(SoundAttenuation);
-		SET_AUDIO_ICON_SIMPLE(AmbientSound);
+		SET_AUDIO_ICON_SVG_SIMPLE(AmbientSound);
 		SET_AUDIO_ICON_SIMPLE(SoundClass);
 		SET_AUDIO_ICON_SIMPLE(SoundConcurrency);
 		SET_AUDIO_ICON_SIMPLE(SoundCue);
 		SET_AUDIO_ICON_SIMPLE(SoundMix);
-		SET_AUDIO_ICON_SIMPLE(AudioVolume);
+		SET_AUDIO_ICON_SVG_SIMPLE(AudioVolume);
 		SET_AUDIO_ICON_SIMPLE(SoundSourceBus);
 		SET_AUDIO_ICON_SIMPLE(SoundSubmix);
 		SET_AUDIO_ICON_SIMPLE(ReverbEffect);
diff --git a/Engine/Source/Editor/EditorStyle/Private/StarshipStyle.cpp b/Engine/Source/Editor/EditorStyle/Private/StarshipStyle.cpp
--- a/Engine/Source/Editor/EditorStyle/Private/StarshipStyle.cpp
+++ b/Engine/Source/Editor/EditorStyle/Private/StarshipStyle.cpp
@@ -5694,18 +5694,21 @@ void FStarshipEditorStyle::FStyle::SetupClassIconsAndThumbnails()
 #if WITH_EDITOR
 	// Actor Classes Outliner
 	{
-		Set("ClassIcon.Emitter", new IMAGE_BRUSH("Icons/ActorIcons/Emitter_16x", Icon16x16));
+		struct FClassIconInfo
+		{
+			FClassIconInfo(TCHAR* InType, bool bInHas64Size = true)
+				: Type(InType)
+				, bHas64Size(bInHas64Size)
+			{}
+
+			TCHAR* Type;
+			bool bHas64Size;
+		};
+
 		Set("ClassIcon.Light", new IMAGE_BRUSH("Icons/ActorIcons/LightActor_16x", Icon16x16));
-		Set("ClassIcon.Brush", new IMAGE_BRUSH("Icons/ActorIcons/Brush_16x", Icon16x16));
 		Set("ClassIcon.BrushAdditive", new IMAGE_BRUSH("Icons/ActorIcons/Brush_Add_16x", Icon16x16));
 		Set("ClassIcon.BrushSubtractive", new IMAGE_BRUSH("Icons/ActorIcons/Brush_Subtract_16x", Icon16x16));
-		Set("ClassIcon.Volume", new IMAGE_BRUSH("Icons/ActorIcons/Volume_16x", Icon16x16));
-		Set("ClassIcon.GroupActor", new IMAGE_BRUSH("Icons/ActorIcons/GroupActor_16x", Icon16x16));
-		Set("ClassIcon.VectorFieldVolume", new IMAGE_BRUSH("Icons/ActorIcons/VectorFieldVolume_16x", Icon16x16));
 		Set("ClassIcon.Deleted", new IMAGE_BRUSH("Icons/ActorIcons/DeletedActor_16px", Icon16x16));
-		Set("ClassIcon.StaticMeshActor", new IMAGE_BRUSH("Icons/AssetIcons/StaticMesh_16x", Icon16x16));
-		Set("ClassIcon.SkeletalMeshActor", new IMAGE_BRUSH("Icons/AssetIcons/SkeletalMesh_16x", Icon16x16));
-		Set("ClassIcon.LevelScriptActor", new IMAGE_BRUSH("Icons/AssetIcons/Blueprint_16x", Icon16x16));
 
 		// Component classes
 		Set("ClassIcon.AudioComponent", new IMAGE_BRUSH("Icons/ActorIcons/SoundActor_16x", Icon16x16));
@@ -5752,17 +5755,15 @@ void FStarshipEditorStyle::FStyle::SetupClassIconsAndThumbnails()
 		Set("ClassIcon.AnimInstance", new IMAGE_BRUSH("Icons/AssetIcons/AnimInstance_16x", Icon16x16));
 
 		// Factory classes
-		Set( "ClassIcon.ActorFactoryBoxVolume", new IMAGE_BRUSH( "Icons/icon_volume_Box_16x", Icon16x16 ) );
-		Set( "ClassIcon.ActorFactoryCylinderVolume", new IMAGE_BRUSH( "Icons/icon_volume_cylinder_16x", Icon16x16 ) );
-		Set( "ClassIcon.ActorFactorySphereVolume", new IMAGE_BRUSH( "Icons/icon_volume_sphere_16x", Icon16x16 ) );
+		//Set( "ClassIcon.ActorFactoryBoxVolume", new IMAGE_BRUSH( "Icons/icon_volume_Box_16x", Icon16x16 ) );
+		//Set( "ClassIcon.ActorFactoryCylinderVolume", new IMAGE_BRUSH( "Icons/icon_volume_cylinder_16x", Icon16x16 ) );
+		//Set( "ClassIcon.ActorFactorySphereVolume", new IMAGE_BRUSH( "Icons/icon_volume_sphere_16x", Icon16x16 ) );
 
 		// Asset Type Classes
 		const TCHAR* AssetTypes[] = {
 
 			TEXT("AbilitySystemComponent"),
-			TEXT("Actor"),
 			TEXT("ActorComponent"),
-			TEXT("AIController"),
 			TEXT("AimOffsetBlendSpace"),
 			TEXT("AimOffsetBlendSpace1D"),
 			TEXT("AIPerceptionComponent"),
@@ -5775,120 +5776,74 @@ void FStarshipEditorStyle::FStyle::SetupClassIconsAndThumbnails()
 			TEXT("AnimSequence"),
 			TEXT("AnimationSharingSetup"),
 			TEXT("ApplicationLifecycleComponent"),
-			TEXT("AtmosphericFog"),
 			TEXT("BehaviorTree"),
 			TEXT("BlackboardData"),
 			TEXT("BlendSpace"),
 			TEXT("BlendSpace1D"),
-			TEXT("BlockingVolume"),
 			TEXT("Blueprint"),
 			TEXT("BlueprintFunctionLibrary"),
 			TEXT("BlueprintInterface"),
 			TEXT("BlueprintMacroLibrary"),
-			TEXT("BoxReflectionCapture"),
 			TEXT("ButtonStyleAsset"),
-			TEXT("CableActor"),
 			TEXT("CableComponent"),
-			TEXT("CameraActor"),
 			TEXT("CameraAnim"),
-			TEXT("CameraBlockingVolume"),
-			TEXT("Character"),
 			TEXT("CharacterMovementComponent"),
 			TEXT("Class"),
 			TEXT("CullDistanceVolume"),
 			TEXT("CurveBase"),
 			TEXT("DataAsset"),
 			TEXT("DataTable"),
-			TEXT("DecalActor"),
 			TEXT("Default"),
-			TEXT("DefaultPawn"),
 			TEXT("DialogueWave"),
 			TEXT("DialogueVoice"),
-			TEXT("DirectionalLight"),
 			TEXT("DirectionalLightMovable"),
 			TEXT("DirectionalLightStatic"),
 			TEXT("DirectionalLightStationary"),
-			TEXT("DocumentationActor"),
 			TEXT("EditorTutorial"),
 			TEXT("EnvQuery"),
-			TEXT("ExponentialHeightFog"),
 			TEXT("FileMediaSource"),
 			TEXT("Font"),
 			TEXT("FontFace"),
 			TEXT("ForceFeedbackEffect"),
 			TEXT("GameModeBase"),
 			TEXT("GameStateBase"),
-			TEXT("HUD"),
 			TEXT("Interface"),
 			TEXT("InterpData"),
-			TEXT("KillZVolume"),
-			TEXT("Landscape"),
-			TEXT("LevelBounds"),
 			TEXT("LevelSequence"),
-			TEXT("LevelStreamingVolume"),
 			TEXT("LightmassCharacterIndirectDetailVolume"),
 			TEXT("VolumetricLightmapDensityVolume"),
-			TEXT("LightmassImportanceVolume"),
 			TEXT("MassiveLODOverrideVolume"),
 			TEXT("Material"),
 			TEXT("MaterialFunction"),
-			TEXT("MaterialInstanceActor"),
 			TEXT("MaterialInstanceConstant"),
 			TEXT("MaterialParameterCollection"),
 			TEXT("MatineeActor"),
 			TEXT("MediaPlayer"),
 			TEXT("MediaTexture"),
-			TEXT("LevelSequenceActor"),
 			TEXT("MultiFont"),
-			TEXT("NavLinkProxy"),
-			TEXT("NavMeshBoundsVolume"),
 			TEXT("NavModifierComponent"),
-			TEXT("NavModifierVolume"),
-			TEXT("Note"),
 			TEXT("ObjectLibrary"),
-			TEXT("PainCausingVolume"),
 			TEXT("ParticleSystem"),
-			TEXT("Pawn"),
 			TEXT("PawnNoiseEmitterComponent"),
 			TEXT("PawnSensingComponent"),
 			TEXT("PhysicalMaterial"),
 			TEXT("PhysicsAsset"),
-			TEXT("PhysicsConstraintActor"),
 			TEXT("PhysicsConstraintComponent"),
 			TEXT("PhysicsHandleComponent"),
-			TEXT("PhysicsThruster"),
 			TEXT("PhysicsThrusterComponent"),
-			TEXT("PhysicsVolume"),
-			TEXT("PlaneReflectionCapture"),
 			TEXT("PlatformMediaSource"),
-			TEXT("PlayerController"),
-			TEXT("PlayerStart"),
-			TEXT("PointLight"),
 			TEXT("PoseAsset"),
-			TEXT("PostProcessVolume"),
-			TEXT("PrecomputedVisibilityOverrideVolume"),
-			TEXT("PrecomputedVisibilityVolume"),
-			TEXT("ProceduralFoliageVolume"),
-			TEXT("ProceduralFoliageBlockingVolume"),
 			TEXT("ProjectileMovementComponent"),
-			TEXT("RadialForceActor"),
 			TEXT("RadialForceComponent"),
-			TEXT("ReflectionCapture"),
 			TEXT("RotatingMovementComponent"),
-			TEXT("SceneCapture2D"),
 			TEXT("SceneCaptureCube"),
 			TEXT("SceneComponent"),
-			TEXT("SkyLight"),
 			TEXT("SkyLightComponent"),
-			TEXT("SkyAtmosphere"),
-			TEXT("VolumetricCloud"),
 			TEXT("SkeletalMesh"),
 			TEXT("Skeleton"),
 			TEXT("SlateBrushAsset"),
 			TEXT("SlateWidgetStyleAsset"),
 			TEXT("StringTable"),
-			TEXT("SphereReflectionCapture"),
-			TEXT("SpotLight"),
 			TEXT("SpotLightMovable"),
 			TEXT("SpotLightStatic"),
 			TEXT("SpotLightStationary"),
@@ -5896,32 +5851,21 @@ void FStarshipEditorStyle::FStyle::SetupClassIconsAndThumbnails()
 			TEXT("StaticMesh"),
 			TEXT("StreamMediaSource"),
 			TEXT("SubsurfaceProfile"),
-			TEXT("TargetPoint"),
-			TEXT("TextRenderActor"),
 			TEXT("Texture2D"),
 			TEXT("TextureRenderTarget2D"),
 			TEXT("TextureRenderTargetCube"),
-			TEXT("TriggerBase"),
-			TEXT("TriggerBox"),
-			TEXT("TriggerCapsule"),
-			TEXT("TriggerSphere"),
-			TEXT("TriggerVolume"),
 			TEXT("TouchInterface"),
 			TEXT("UserDefinedCaptureProtocol"),
 			TEXT("UserDefinedEnum"),
 			TEXT("UserDefinedStruct"),
 			TEXT("UserWidget"),
 			TEXT("WidgetBlueprint"),
-			TEXT("WindDirectionalSource"),
 			TEXT("World"),
 			TEXT("Cube"),
 			TEXT("Sphere"),
 			TEXT("Cylinder"),
 			TEXT("Cone"),
 			TEXT("Plane"),
-			TEXT("CineCameraActor"),
-			TEXT("CameraRig_Crane"),
-			TEXT("CameraRig_Rail"),
 			TEXT("FoliageType_Actor"),
 			TEXT("LevelInstance")
 		};
@@ -5932,6 +5876,109 @@ void FStarshipEditorStyle::FStyle::SetupClassIconsAndThumbnails()
 			Set( *FString::Printf(TEXT("ClassIcon.%s"), Type),		new IMAGE_BRUSH(FString::Printf(TEXT("Icons/AssetIcons/%s_%dx"), Type, 16), Icon16x16 ) );
 			Set( *FString::Printf(TEXT("ClassThumbnail.%s"), Type),	new IMAGE_BRUSH(FString::Printf(TEXT("Icons/AssetIcons/%s_%dx"), Type, 64), Icon64x64 ) );
 		}
+
+		const FClassIconInfo AssetTypesSVG[] = {
+			{TEXT("Actor")},
+			{TEXT("AIController")},
+			{TEXT("AtmosphericFog")},
+			{TEXT("BlockingVolume")},
+			{TEXT("BoxReflectionCapture")},
+			{TEXT("Brush")},
+			{TEXT("CableActor")},
+			{TEXT("CameraActor")},
+			{TEXT("CameraBlockingVolume")},
+			{TEXT("CameraRig_Crane")},
+			{TEXT("CameraRig_Rail")},
+			{TEXT("Character")},
+			{TEXT("CineCameraActor")},
+			{TEXT("CompositingElement")},
+			{TEXT("DecalActor")},
+			{TEXT("DefaultPawn")},
+			{TEXT("DirectionalLight")},
+			{TEXT("DocumentationActor")},
+			{TEXT("Emitter")},
+			{TEXT("EmptyActor")},
+			{TEXT("ExponentialHeightFog")},
+			{TEXT("ForceFeedback")},
+			{TEXT("GroupActor")},
+			{TEXT("HUD")},
+			{TEXT("ImagePlate")},
+			{TEXT("KillZVolume")},
+			{TEXT("Landscape")},
+			{TEXT("LevelBounds")},
+			{TEXT("LevelScriptActor")},
+			{TEXT("LevelSequenceActor")},
+			{TEXT("LevelStreamingVolume"), false},
+			{TEXT("LightmassImportanceVolume"), false},
+			{TEXT("LightmassVolume")},
+			{TEXT("MaterialInstanceActor")},
+			{TEXT("NavLink")},
+			{TEXT("NavLinkProxy"), false},
+			{TEXT("NavMeshBoundsVolume"), false},
+			{TEXT("NavModifierVolume"), false},
+			{TEXT("Note")},
+			{TEXT("PainCausingVolume")},
+			{TEXT("Pawn")},
+			{TEXT("PhysicsConstraintActor"), false},
+			{TEXT("PhysicsThruster")},
+			{TEXT("PhysicsVolume")},
+			{TEXT("PlaneReflectionCapture")},
+			{TEXT("PlayerController")},
+			{TEXT("PlayerStart")},
+			{TEXT("PointLight")},
+			{TEXT("PostProcessVolume")},
+			{TEXT("PrecomputedVisibilityOverrideVolume")},
+			{TEXT("PrecomputedVisibilityVolume")},
+			{TEXT("ProceduralFoliageBlockingVolume")},
+			{TEXT("ProceduralFoliageVolume")},
+			{TEXT("RadialForceActor")},
+			{TEXT("RectLight"), false},
+			{TEXT("ReflectionCapture")},
+			{TEXT("SceneCapture2D")},
+			{TEXT("SkeletalMeshActor")},
+			{TEXT("SkyAtmosphere")},
+			{TEXT("SkyLight")},
+			{TEXT("SphereReflectionCapture")},
+			{TEXT("SpotLight")},
+			{TEXT("StaticMeshActor")},
+			{TEXT("TargetPoint")},
+			{TEXT("TextRenderActor")},
+			{TEXT("TriggerBase")},
+			{TEXT("TriggerBox")},
+			{TEXT("TriggerCapsule")},
+			{TEXT("TriggerSphere")},
+			{TEXT("TriggerVolume")},
+			{TEXT("VectorField")},
+			{TEXT("VectorFieldVolume")},
+			{TEXT("Volume")},
+			{TEXT("VolumetricCloud"), false},
+			{TEXT("WindDirectionalSource")}
+		};
+	
+		// SVG Asset icons
+		{
+			for (int32 TypeIndex = 0; TypeIndex < UE_ARRAY_COUNT(AssetTypesSVG); ++TypeIndex)
+			{
+				const FClassIconInfo& Info = AssetTypesSVG[TypeIndex];
+
+				FString ClassIconName = FString::Printf(TEXT("ClassIcon.%s"), Info.Type);
+				if (GetOptionalBrush(*ClassIconName, nullptr, nullptr))
+				{
+					UE_LOG(LogSlate, Log, TEXT("%s already found"), *ClassIconName);
+				}
+
+				Set(*FString::Printf(TEXT("ClassIcon.%s"), Info.Type), new IMAGE_BRUSH_SVG(FString::Printf(TEXT("Starship/AssetIcons/%s_%d"), Info.Type, 16), Icon16x16));
+				if (Info.bHas64Size)
+				{
+					Set(*FString::Printf(TEXT("ClassThumbnail.%s"), Info.Type), new IMAGE_BRUSH_SVG(FString::Printf(TEXT("Starship/AssetIcons/%s_%d"), Info.Type, 64), Icon64x64));
+				}
+				else
+				{
+					// Temp to avoid missing icons while in progress. use the 64 variant for 16 for now.  
+					Set(*FString::Printf(TEXT("ClassThumbnail.%s"), Info.Type), new IMAGE_BRUSH_SVG(FString::Printf(TEXT("Starship/AssetIcons/%s_%d"), Info.Type, 16), Icon64x64));
+				}
+			}
+		}
 	}
 #endif
 }
