Fixed missing const on type storing string literal values

#rnx

diff --git a/Engine/Source/Editor/EditorStyle/Private/StarshipStyle.cpp b/Engine/Source/Editor/EditorStyle/Private/StarshipStyle.cpp
--- a/Engine/Source/Editor/EditorStyle/Private/StarshipStyle.cpp
+++ b/Engine/Source/Editor/EditorStyle/Private/StarshipStyle.cpp
@@ -5696,12 +5696,12 @@ void FStarshipEditorStyle::FStyle::SetupClassIconsAndThumbnails()
 	{
 		struct FClassIconInfo
 		{
-			FClassIconInfo(TCHAR* InType, bool bInHas64Size = true)
+			FClassIconInfo(const TCHAR* InType, bool bInHas64Size = true)
 				: Type(InType)
 				, bHas64Size(bInHas64Size)
 			{}
 
-			TCHAR* Type;
+			const TCHAR* Type;
 			bool bHas64Size;
 		};
 
