Legacy Mode UI removed.

#rnx

diff --git a/Engine/Plugins/Editor/GeometryMode/Source/GeometryMode/Private/GeometryMode.cpp b/Engine/Plugins/Editor/GeometryMode/Source/GeometryMode/Private/GeometryMode.cpp
--- a/Engine/Plugins/Editor/GeometryMode/Source/GeometryMode/Private/GeometryMode.cpp
+++ b/Engine/Plugins/Editor/GeometryMode/Source/GeometryMode/Private/GeometryMode.cpp
@@ -67,9 +67,6 @@ private:
 	/** Creates controls for the modifiers section */
 	TSharedRef<SVerticalBox> CreateTopModifierButtons(TSharedRef<FGeometryModeToolkit>& GeometryMode);
 
-	/** Creates controls for the actions section */
-	TSharedRef<SUniformGridPanel> CreateBottomModifierButtons(TSharedRef<FGeometryModeToolkit>& GeometryMode);
-
 	/** Creates controls for the modifier properties section */
 	TSharedRef<class IDetailsView> CreateModifierProperties(TSharedRef<FGeometryModeToolkit>& GeometryMode);
 
@@ -177,12 +174,6 @@ void SGeometryModeControls::CreateLayout(TSharedRef<FGeometryModeToolkit>& Geome
 				.Visibility(this, &SGeometryModeControls::GetPropertyVisibility)
 				.OnClicked(GeometryMode, &FGeometryModeToolkit::OnApplyClicked)
 			]
-			+SVerticalBox::Slot()
-			.AutoHeight()
-			.HAlign(HAlign_Center)
-			[
-				CreateBottomModifierButtons(GeometryMode)
-			]
 		]
 	];
 }
@@ -195,24 +186,6 @@ TSharedRef<SVerticalBox> SGeometryModeControls::CreateTopModifierButtons(TShared
 
 	// Loop through all geometry modifiers and create radio buttons for ones with the bPushButton set to false
 	int32 CurrentModifierButtonCount = 0;
-	
-	if(GetDefault<UEditorStyleSettings>()->bEnableLegacyEditorModeUI)
-	{
-		for (FModeTool_GeometryModify::TModifierIterator Itor(GeometryModeTool->ModifierIterator()); Itor; ++Itor)
-		{
-			UGeomModifier* Modifier = *Itor;
-			if (!Modifier->bPushButton)
-			{
-				RadioButtonPanel->AddSlot(CurrentModifierButtonCount % 2, CurrentModifierButtonCount / 2)
-					.Padding(FMargin(20.0f, 5.0f))
-					[
-						CreateSingleModifierRadioButton(Modifier, GeometryMode)
-					];
-
-				++CurrentModifierButtonCount;
-			}
-		}
-	}
 
 	// Add the Apply button
 	SAssignNew(Vbox, SVerticalBox)
@@ -225,47 +198,6 @@ TSharedRef<SVerticalBox> SGeometryModeControls::CreateTopModifierButtons(TShared
 	return Vbox.ToSharedRef();
 }
 
-TSharedRef<SUniformGridPanel> SGeometryModeControls::CreateBottomModifierButtons(TSharedRef<FGeometryModeToolkit>& GeometryMode)
-{
-	FModeTool_GeometryModify* GeometryModeTool = GeometryMode->GetGeometryModeTool();
-
-	TSharedRef<SUniformGridPanel> ButtonGrid = SNew(SUniformGridPanel).SlotPadding(5.0f);
-
-	if (GetDefault<UEditorStyleSettings>()->bEnableLegacyEditorModeUI)
-	{
-		// The IDs of the buttons created in this function need to sequentially follow the IDs of the modifier radio buttons
-			// So, this loop simply counts the number of radio buttons so we can use that as an offset
-		int32 CurrentModifierButtonCount = 0;
-		for (FModeTool_GeometryModify::TModifierConstIterator Itor(GeometryModeTool->ModifierConstIterator()); Itor; ++Itor)
-		{
-			const UGeomModifier* Modifier = *Itor;
-			if (!Modifier->bPushButton)
-			{
-				++CurrentModifierButtonCount;
-			}
-		}
-
-		// Loop through all geometry modifiers and create buttons for ones with the bPushButton set to true
-		int32 PushButtonId = 0;
-		for (FModeTool_GeometryModify::TModifierIterator Itor(GeometryModeTool->ModifierIterator()); Itor; ++Itor)
-		{
-			UGeomModifier* Modifier = *Itor;
-			if (Modifier->bPushButton)
-			{
-				ButtonGrid->AddSlot(PushButtonId % 2, PushButtonId / 2)
-					[
-						CreateSingleModifierButton(Modifier, GeometryMode)
-					];
-
-				++CurrentModifierButtonCount;
-				++PushButtonId;
-			}
-		}
-	}
-
-	return ButtonGrid;
-}
-
 TSharedRef<IDetailsView> SGeometryModeControls::CreateModifierProperties(TSharedRef<FGeometryModeToolkit>& GeometryMode)
 {
 	FDetailsViewArgs Args;
@@ -438,10 +370,7 @@ TSharedPtr<SWidget> FGeometryModeToolkit::GetInlineContent() const
 
 void FGeometryModeToolkit::GetToolPaletteNames(TArray<FName>& PaletteNames) const
 {
-	if(!GetDefault<UEditorStyleSettings>()->bEnableLegacyEditorModeUI)
-	{
-		PaletteNames.Add(GeometryModePaletteNames::ToolsPalette);
-	}
+	PaletteNames.Add(GeometryModePaletteNames::ToolsPalette);
 }
 
 FText FGeometryModeToolkit::GetToolPaletteDisplayName(FName Palette) const
diff --git a/Engine/Source/Editor/EditorStyle/Private/EditorStyleClasses.cpp b/Engine/Source/Editor/EditorStyle/Private/EditorStyleClasses.cpp
--- a/Engine/Source/Editor/EditorStyle/Private/EditorStyleClasses.cpp
+++ b/Engine/Source/Editor/EditorStyle/Private/EditorStyleClasses.cpp
@@ -35,8 +35,6 @@ UEditorStyleSettings::UEditorStyleSettings( const FObjectInitializer& ObjectInit
 
 	bShowFriendlyNames = true;
 	LogTimestampMode = ELogTimes::None;
-	
-	bEnableLegacyEditorModeUI = false;
 }
 
 void UEditorStyleSettings::Init()
diff --git a/Engine/Source/Editor/EditorStyle/Private/SlateEditorStyle.cpp b/Engine/Source/Editor/EditorStyle/Private/SlateEditorStyle.cpp
--- a/Engine/Source/Editor/EditorStyle/Private/SlateEditorStyle.cpp
+++ b/Engine/Source/Editor/EditorStyle/Private/SlateEditorStyle.cpp
@@ -1783,22 +1783,6 @@ void FSlateEditorStyle::FStyle::SetupGeneralStyles()
 		Set("FoliageEditMode.MoveToCurrentLevel",        new IMAGE_BRUSH("Icons/GeneralTools/MoveToLevel_40x", Icon20x20));
 		Set("FoliageEditMode.MoveToCurrentLevel.Small",  new IMAGE_BRUSH("Icons/GeneralTools/MoveToLevel_40x", Icon20x20));
 
-
-		if (GetDefault<UEditorStyleSettings>()->bEnableLegacyEditorModeUI)
-		{
-			Set("FoliageEditMode.SetPaint", new IMAGE_BRUSH("Icons/FoliageEditMode/icon_FoliageEdMode_Paint_40x", Icon20x20));
-			Set("FoliageEditMode.SetReapplySettings", new IMAGE_BRUSH("Icons/FoliageEditMode/icon_FoliageEdMode_Reapply_40x", Icon20x20));
-			Set("FoliageEditMode.SetSelect", new IMAGE_BRUSH("Icons/FoliageEditMode/icon_FoliageEdMode_Select_40x", Icon20x20));
-			Set("FoliageEditMode.SetLassoSelect", new IMAGE_BRUSH("Icons/FoliageEditMode/icon_FoliageEdMode_Lasso_40x", Icon20x20));
-			Set("FoliageEditMode.SetPaintBucket", new IMAGE_BRUSH("Icons/FoliageEditMode/icon_FoliageEdMode_PaintBucket_40x", Icon20x20));
-
-			Set( "FoliageEditMode.SetPaint.Small", new IMAGE_BRUSH( "Icons/FoliageEditMode/icon_FoliageEdMode_Paint_40x", Icon20x20 ) );
-			Set( "FoliageEditMode.SetReapplySettings.Small", new IMAGE_BRUSH( "Icons/FoliageEditMode/icon_FoliageEdMode_Reapply_40x", Icon20x20 ) );
-			Set( "FoliageEditMode.SetSelect.Small", new IMAGE_BRUSH( "Icons/FoliageEditMode/icon_FoliageEdMode_Select_40x", Icon20x20 ) );
-			Set( "FoliageEditMode.SetLassoSelect.Small", new IMAGE_BRUSH( "Icons/FoliageEditMode/icon_FoliageEdMode_Lasso_40x", Icon20x20 ) );
-			Set( "FoliageEditMode.SetPaintBucket.Small", new IMAGE_BRUSH( "Icons/FoliageEditMode/icon_FoliageEdMode_PaintBucket_40x", Icon20x20 ) );
-		}
-
 		Set( "FoliageEditMode.SetNoSettings", new IMAGE_BRUSH( "Icons/FoliageEditMode/icon_FoliageEdMode_NoSettings_20x", Icon20x20 ) );
 		Set( "FoliageEditMode.SetPaintSettings", new IMAGE_BRUSH( "Icons/FoliageEditMode/icon_FoliageEdMode_PaintingSettings_20x", Icon20x20 ) );
 		Set( "FoliageEditMode.SetClusterSettings", new IMAGE_BRUSH( "Icons/FoliageEditMode/icon_FoliageEdMode_ClusterSettings_20x", Icon20x20 ) );
@@ -6863,83 +6847,6 @@ void FSlateEditorStyle::FStyle::SetupLandscapeEditorStyle()
 		Set("LandscapeEditor.SculptMode.Small", new IMAGE_BRUSH("Icons/icon_Landscape_Mode_Sculpt_20x", Icon20x20));
 		Set("LandscapeEditor.PaintMode.Small",  new IMAGE_BRUSH("Icons/icon_Landscape_Mode_Paint_20x",  Icon20x20));
 
-		if (GetDefault<UEditorStyleSettings>()->bEnableLegacyEditorModeUI)
-		{
-			// Tools
-			Set("LandscapeEditor.NewLandscape",       new IMAGE_BRUSH("Icons/icon_Landscape_New_Landscape_40x", Icon40x40));
-			Set("LandscapeEditor.NewLandscape.Small", new IMAGE_BRUSH("Icons/icon_Landscape_New_Landscape_20x", Icon20x20));
-			Set("LandscapeEditor.ResizeLandscape",       new IMAGE_BRUSH("Icons/icon_Landscape_Tool_Copy_40x", Icon40x40));
-			Set("LandscapeEditor.ResizeLandscape.Small", new IMAGE_BRUSH("Icons/icon_Landscape_Tool_Copy_20x", Icon20x20));
-
-			Set("LandscapeEditor.SculptTool",       new IMAGE_BRUSH("Icons/icon_Landscape_Tool_Sculpt_40x",           Icon40x40));
-			Set("LandscapeEditor.EraseTool",		new IMAGE_BRUSH("Icons/icon_Landscape_Tool_Erase_40x", Icon40x40));
-			Set("LandscapeEditor.PaintTool",        new IMAGE_BRUSH("Icons/icon_Landscape_Tool_Paint_40x",            Icon40x40));
-			Set("LandscapeEditor.SmoothTool",       new IMAGE_BRUSH("Icons/icon_Landscape_Tool_Smooth_40x",           Icon40x40));
-			Set("LandscapeEditor.FlattenTool",      new IMAGE_BRUSH("Icons/icon_Landscape_Tool_Flatten_40x",          Icon40x40));
-			Set("LandscapeEditor.RampTool",         new IMAGE_BRUSH("Icons/icon_Landscape_Tool_Ramp_40x",             Icon40x40));
-			Set("LandscapeEditor.ErosionTool",      new IMAGE_BRUSH("Icons/icon_Landscape_Tool_Erosion_40x",          Icon40x40));
-			Set("LandscapeEditor.HydroErosionTool", new IMAGE_BRUSH("Icons/icon_Landscape_Tool_HydraulicErosion_40x", Icon40x40));
-			Set("LandscapeEditor.NoiseTool",        new IMAGE_BRUSH("Icons/icon_Landscape_Tool_Noise_40x",            Icon40x40));
-			Set("LandscapeEditor.RetopologizeTool", new IMAGE_BRUSH("Icons/icon_Landscape_Tool_Retopologize_40x",     Icon40x40));
-			Set("LandscapeEditor.VisibilityTool",   new IMAGE_BRUSH("Icons/icon_Landscape_Tool_Visibility_40x",       Icon40x40));
-			Set("LandscapeEditor.BlueprintBrushTool", new IMAGE_BRUSH("Icons/icon_Landscape_Tool_BlueprintBrush_40x", Icon40x40));
-			Set("LandscapeEditor.SculptTool.Small",       new IMAGE_BRUSH("Icons/icon_Landscape_Tool_Sculpt_20x",           Icon20x20));
-			Set("LandscapeEditor.EraseTool.Small", new IMAGE_BRUSH("Icons/icon_Landscape_Tool_Erase_20x", Icon20x20));
-			Set("LandscapeEditor.PaintTool.Small",        new IMAGE_BRUSH("Icons/icon_Landscape_Tool_Paint_20x",            Icon20x20));
-			Set("LandscapeEditor.SmoothTool.Small",       new IMAGE_BRUSH("Icons/icon_Landscape_Tool_Smooth_20x",           Icon20x20));
-			Set("LandscapeEditor.FlattenTool.Small",      new IMAGE_BRUSH("Icons/icon_Landscape_Tool_Flatten_20x",          Icon20x20));
-			Set("LandscapeEditor.RampTool.Small",         new IMAGE_BRUSH("Icons/icon_Landscape_Tool_Ramp_20x",             Icon20x20));
-			Set("LandscapeEditor.ErosionTool.Small",      new IMAGE_BRUSH("Icons/icon_Landscape_Tool_Erosion_20x",          Icon20x20));
-			Set("LandscapeEditor.HydroErosionTool.Small", new IMAGE_BRUSH("Icons/icon_Landscape_Tool_HydraulicErosion_20x", Icon20x20));
-			Set("LandscapeEditor.NoiseTool.Small",        new IMAGE_BRUSH("Icons/icon_Landscape_Tool_Noise_20x",            Icon20x20));
-			Set("LandscapeEditor.RetopologizeTool.Small", new IMAGE_BRUSH("Icons/icon_Landscape_Tool_Retopologize_20x",     Icon20x20));
-			Set("LandscapeEditor.VisibilityTool.Small",   new IMAGE_BRUSH("Icons/icon_Landscape_Tool_Visibility_20x",       Icon20x20));
-			Set("LandscapeEditor.BlueprintBrushTool.Small", new IMAGE_BRUSH("Icons/icon_Landscape_Tool_BlueprintBrush_20x", Icon20x20));
-
-			Set("LandscapeEditor.SelectComponentTool", new IMAGE_BRUSH("Icons/icon_Landscape_Tool_Selection_40x",       Icon40x40));
-			Set("LandscapeEditor.AddComponentTool",    new IMAGE_BRUSH("Icons/icon_Landscape_Tool_AddComponent_40x",    Icon40x40));
-			Set("LandscapeEditor.DeleteComponentTool", new IMAGE_BRUSH("Icons/icon_Landscape_Tool_DeleteComponent_40x", Icon40x40));
-			Set("LandscapeEditor.MoveToLevelTool",     new IMAGE_BRUSH("Icons/icon_Landscape_Tool_MoveToLevel_40x",     Icon40x40));
-			Set("LandscapeEditor.SelectComponentTool.Small", new IMAGE_BRUSH("Icons/icon_Landscape_Tool_Selection_20x",       Icon20x20));
-			Set("LandscapeEditor.AddComponentTool.Small",    new IMAGE_BRUSH("Icons/icon_Landscape_Tool_AddComponent_20x",    Icon20x20));
-			Set("LandscapeEditor.DeleteComponentTool.Small", new IMAGE_BRUSH("Icons/icon_Landscape_Tool_DeleteComponent_20x", Icon20x20));
-			Set("LandscapeEditor.MoveToLevelTool.Small",     new IMAGE_BRUSH("Icons/icon_Landscape_Tool_MoveToLevel_20x",     Icon20x20));
-
-			Set("LandscapeEditor.RegionSelectTool",    new IMAGE_BRUSH("Icons/icon_Landscape_Tool_Mask_40x", Icon40x40));
-			Set("LandscapeEditor.RegionCopyPasteTool", new IMAGE_BRUSH("Icons/icon_Landscape_Tool_Copy_40x", Icon40x40));
-			Set("LandscapeEditor.RegionSelectTool.Small",    new IMAGE_BRUSH("Icons/icon_Landscape_Tool_Mask_20x", Icon20x20));
-			Set("LandscapeEditor.RegionCopyPasteTool.Small", new IMAGE_BRUSH("Icons/icon_Landscape_Tool_Copy_20x", Icon20x20));
-
-			Set("LandscapeEditor.MirrorTool",       new IMAGE_BRUSH("Icons/icon_Landscape_Tool_Mirror_40x", Icon40x40));
-			Set("LandscapeEditor.MirrorTool.Small", new IMAGE_BRUSH("Icons/icon_Landscape_Tool_Mirror_20x", Icon20x20));
-
-			Set("LandscapeEditor.SplineTool",       new IMAGE_BRUSH("Icons/icon_Landscape_Tool_Splines_40x", Icon40x40));
-			Set("LandscapeEditor.SplineTool.Small", new IMAGE_BRUSH("Icons/icon_Landscape_Tool_Splines_20x", Icon20x20));
-
-			// Brush Sets
-			Set("LandscapeEditor.CircleBrush",        new IMAGE_BRUSH("Icons/icon_Landscape_Brush_Circle_smooth_40x", Icon40x40));
-			Set("LandscapeEditor.AlphaBrush",         new IMAGE_BRUSH("Icons/icon_Landscape_Brush_Alpha_40x",     Icon40x40));
-			Set("LandscapeEditor.AlphaBrush_Pattern", new IMAGE_BRUSH("Icons/icon_Landscape_Brush_Pattern_40x",   Icon40x40));
-			Set("LandscapeEditor.ComponentBrush",     new IMAGE_BRUSH("Icons/icon_Landscape_Brush_Component_40x", Icon40x40));
-			Set("LandscapeEditor.GizmoBrush",         new IMAGE_BRUSH("Icons/icon_Landscape_Brush_Gizmo_40x",     Icon40x40));
-			Set("LandscapeEditor.CircleBrush.Small",        new IMAGE_BRUSH("Icons/icon_Landscape_Brush_Circle_smooth_20x", Icon20x20));
-			Set("LandscapeEditor.AlphaBrush.Small",         new IMAGE_BRUSH("Icons/icon_Landscape_Brush_Alpha_20x",     Icon20x20));
-			Set("LandscapeEditor.AlphaBrush_Pattern.Small", new IMAGE_BRUSH("Icons/icon_Landscape_Brush_Pattern_20x",   Icon20x20));
-			Set("LandscapeEditor.ComponentBrush.Small",     new IMAGE_BRUSH("Icons/icon_Landscape_Brush_Component_20x", Icon20x20));
-			Set("LandscapeEditor.GizmoBrush.Small",         new IMAGE_BRUSH("Icons/icon_Landscape_Brush_Gizmo_20x",     Icon20x20));
-
-			// Brushes
-			Set("LandscapeEditor.CircleBrush_Smooth",    new IMAGE_BRUSH("Icons/icon_Landscape_Brush_Circle_smooth_40x",    Icon40x40));
-			Set("LandscapeEditor.CircleBrush_Linear",    new IMAGE_BRUSH("Icons/icon_Landscape_Brush_Circle_linear_40x",    Icon40x40));
-			Set("LandscapeEditor.CircleBrush_Spherical", new IMAGE_BRUSH("Icons/icon_Landscape_Brush_Circle_spherical_40x", Icon40x40));
-			Set("LandscapeEditor.CircleBrush_Tip",       new IMAGE_BRUSH("Icons/icon_Landscape_Brush_Circle_tip_40x",       Icon40x40));
-			Set("LandscapeEditor.CircleBrush_Smooth.Small",    new IMAGE_BRUSH("Icons/icon_Landscape_Brush_Circle_smooth_20x",    Icon20x20));
-			Set("LandscapeEditor.CircleBrush_Linear.Small",    new IMAGE_BRUSH("Icons/icon_Landscape_Brush_Circle_linear_20x",    Icon20x20));
-			Set("LandscapeEditor.CircleBrush_Spherical.Small", new IMAGE_BRUSH("Icons/icon_Landscape_Brush_Circle_spherical_20x", Icon20x20));
-			Set("LandscapeEditor.CircleBrush_Tip.Small",       new IMAGE_BRUSH("Icons/icon_Landscape_Brush_Circle_tip_20x",       Icon20x20));
-		}
-
-		else
 		{
 			// Tools
 			Set("LandscapeEditor.NewLandscape",					new IMAGE_BRUSH("Icons/Landscape/Landscape_NewLandscape_x40",		Icon20x20));
diff --git a/Engine/Source/Editor/EditorStyle/Private/StarshipStyle.cpp b/Engine/Source/Editor/EditorStyle/Private/StarshipStyle.cpp
--- a/Engine/Source/Editor/EditorStyle/Private/StarshipStyle.cpp
+++ b/Engine/Source/Editor/EditorStyle/Private/StarshipStyle.cpp
@@ -1539,22 +1539,6 @@ void FStarshipEditorStyle::FStyle::SetupGeneralStyles()
 		Set("FoliageEditMode.MoveToCurrentLevel",        new IMAGE_BRUSH("Icons/GeneralTools/MoveToLevel_40x", Icon20x20));
 		Set("FoliageEditMode.MoveToCurrentLevel.Small",  new IMAGE_BRUSH("Icons/GeneralTools/MoveToLevel_40x", Icon20x20));
 
-
-		if (GetDefault<UEditorStyleSettings>()->bEnableLegacyEditorModeUI)
-		{
-			Set("FoliageEditMode.SetPaint", new IMAGE_BRUSH("Icons/FoliageEditMode/icon_FoliageEdMode_Paint_40x", Icon20x20));
-			Set("FoliageEditMode.SetReapplySettings", new IMAGE_BRUSH("Icons/FoliageEditMode/icon_FoliageEdMode_Reapply_40x", Icon20x20));
-			Set("FoliageEditMode.SetSelect", new IMAGE_BRUSH("Icons/FoliageEditMode/icon_FoliageEdMode_Select_40x", Icon20x20));
-			Set("FoliageEditMode.SetLassoSelect", new IMAGE_BRUSH("Icons/FoliageEditMode/icon_FoliageEdMode_Lasso_40x", Icon20x20));
-			Set("FoliageEditMode.SetPaintBucket", new IMAGE_BRUSH("Icons/FoliageEditMode/icon_FoliageEdMode_PaintBucket_40x", Icon20x20));
-
-			Set( "FoliageEditMode.SetPaint.Small", new IMAGE_BRUSH( "Icons/FoliageEditMode/icon_FoliageEdMode_Paint_40x", Icon20x20 ) );
-			Set( "FoliageEditMode.SetReapplySettings.Small", new IMAGE_BRUSH( "Icons/FoliageEditMode/icon_FoliageEdMode_Reapply_40x", Icon20x20 ) );
-			Set( "FoliageEditMode.SetSelect.Small", new IMAGE_BRUSH( "Icons/FoliageEditMode/icon_FoliageEdMode_Select_40x", Icon20x20 ) );
-			Set( "FoliageEditMode.SetLassoSelect.Small", new IMAGE_BRUSH( "Icons/FoliageEditMode/icon_FoliageEdMode_Lasso_40x", Icon20x20 ) );
-			Set( "FoliageEditMode.SetPaintBucket.Small", new IMAGE_BRUSH( "Icons/FoliageEditMode/icon_FoliageEdMode_PaintBucket_40x", Icon20x20 ) );
-		}
-
 		Set( "FoliageEditMode.SetNoSettings", new IMAGE_BRUSH( "Icons/FoliageEditMode/icon_FoliageEdMode_NoSettings_20x", Icon20x20 ) );
 		Set( "FoliageEditMode.SetPaintSettings", new IMAGE_BRUSH( "Icons/FoliageEditMode/icon_FoliageEdMode_PaintingSettings_20x", Icon20x20 ) );
 		Set( "FoliageEditMode.SetClusterSettings", new IMAGE_BRUSH( "Icons/FoliageEditMode/icon_FoliageEdMode_ClusterSettings_20x", Icon20x20 ) );
@@ -6276,83 +6260,6 @@ void FStarshipEditorStyle::FStyle::SetupLandscapeEditorStyle()
 		Set("LandscapeEditor.SculptMode.Small", new IMAGE_BRUSH("Icons/icon_Landscape_Mode_Sculpt_20x", Icon20x20));
 		Set("LandscapeEditor.PaintMode.Small",  new IMAGE_BRUSH("Icons/icon_Landscape_Mode_Paint_20x",  Icon20x20));
 
-		if (GetDefault<UEditorStyleSettings>()->bEnableLegacyEditorModeUI)
-		{
-			// Tools
-			Set("LandscapeEditor.NewLandscape",       new IMAGE_BRUSH("Icons/icon_Landscape_New_Landscape_40x", Icon40x40));
-			Set("LandscapeEditor.NewLandscape.Small", new IMAGE_BRUSH("Icons/icon_Landscape_New_Landscape_20x", Icon20x20));
-			Set("LandscapeEditor.ResizeLandscape",       new IMAGE_BRUSH("Icons/icon_Landscape_Tool_Copy_40x", Icon40x40));
-			Set("LandscapeEditor.ResizeLandscape.Small", new IMAGE_BRUSH("Icons/icon_Landscape_Tool_Copy_20x", Icon20x20));
-
-			Set("LandscapeEditor.SculptTool",       new IMAGE_BRUSH("Icons/icon_Landscape_Tool_Sculpt_40x",           Icon40x40));
-			Set("LandscapeEditor.EraseTool",		new IMAGE_BRUSH("Icons/icon_Landscape_Tool_Erase_40x", Icon40x40));
-			Set("LandscapeEditor.PaintTool",        new IMAGE_BRUSH("Icons/icon_Landscape_Tool_Paint_40x",            Icon40x40));
-			Set("LandscapeEditor.SmoothTool",       new IMAGE_BRUSH("Icons/icon_Landscape_Tool_Smooth_40x",           Icon40x40));
-			Set("LandscapeEditor.FlattenTool",      new IMAGE_BRUSH("Icons/icon_Landscape_Tool_Flatten_40x",          Icon40x40));
-			Set("LandscapeEditor.RampTool",         new IMAGE_BRUSH("Icons/icon_Landscape_Tool_Ramp_40x",             Icon40x40));
-			Set("LandscapeEditor.ErosionTool",      new IMAGE_BRUSH("Icons/icon_Landscape_Tool_Erosion_40x",          Icon40x40));
-			Set("LandscapeEditor.HydroErosionTool", new IMAGE_BRUSH("Icons/icon_Landscape_Tool_HydraulicErosion_40x", Icon40x40));
-			Set("LandscapeEditor.NoiseTool",        new IMAGE_BRUSH("Icons/icon_Landscape_Tool_Noise_40x",            Icon40x40));
-			Set("LandscapeEditor.RetopologizeTool", new IMAGE_BRUSH("Icons/icon_Landscape_Tool_Retopologize_40x",     Icon40x40));
-			Set("LandscapeEditor.VisibilityTool",   new IMAGE_BRUSH("Icons/icon_Landscape_Tool_Visibility_40x",       Icon40x40));
-			Set("LandscapeEditor.BlueprintBrushTool", new IMAGE_BRUSH("Icons/icon_Landscape_Tool_BlueprintBrush_40x", Icon40x40));
-			Set("LandscapeEditor.SculptTool.Small",       new IMAGE_BRUSH("Icons/icon_Landscape_Tool_Sculpt_20x",           Icon20x20));
-			Set("LandscapeEditor.EraseTool.Small", new IMAGE_BRUSH("Icons/icon_Landscape_Tool_Erase_20x", Icon20x20));
-			Set("LandscapeEditor.PaintTool.Small",        new IMAGE_BRUSH("Icons/icon_Landscape_Tool_Paint_20x",            Icon20x20));
-			Set("LandscapeEditor.SmoothTool.Small",       new IMAGE_BRUSH("Icons/icon_Landscape_Tool_Smooth_20x",           Icon20x20));
-			Set("LandscapeEditor.FlattenTool.Small",      new IMAGE_BRUSH("Icons/icon_Landscape_Tool_Flatten_20x",          Icon20x20));
-			Set("LandscapeEditor.RampTool.Small",         new IMAGE_BRUSH("Icons/icon_Landscape_Tool_Ramp_20x",             Icon20x20));
-			Set("LandscapeEditor.ErosionTool.Small",      new IMAGE_BRUSH("Icons/icon_Landscape_Tool_Erosion_20x",          Icon20x20));
-			Set("LandscapeEditor.HydroErosionTool.Small", new IMAGE_BRUSH("Icons/icon_Landscape_Tool_HydraulicErosion_20x", Icon20x20));
-			Set("LandscapeEditor.NoiseTool.Small",        new IMAGE_BRUSH("Icons/icon_Landscape_Tool_Noise_20x",            Icon20x20));
-			Set("LandscapeEditor.RetopologizeTool.Small", new IMAGE_BRUSH("Icons/icon_Landscape_Tool_Retopologize_20x",     Icon20x20));
-			Set("LandscapeEditor.VisibilityTool.Small",   new IMAGE_BRUSH("Icons/icon_Landscape_Tool_Visibility_20x",       Icon20x20));
-			Set("LandscapeEditor.BlueprintBrushTool.Small", new IMAGE_BRUSH("Icons/icon_Landscape_Tool_BlueprintBrush_20x", Icon20x20));
-
-			Set("LandscapeEditor.SelectComponentTool", new IMAGE_BRUSH("Icons/icon_Landscape_Tool_Selection_40x",       Icon40x40));
-			Set("LandscapeEditor.AddComponentTool",    new IMAGE_BRUSH("Icons/icon_Landscape_Tool_AddComponent_40x",    Icon40x40));
-			Set("LandscapeEditor.DeleteComponentTool", new IMAGE_BRUSH("Icons/icon_Landscape_Tool_DeleteComponent_40x", Icon40x40));
-			Set("LandscapeEditor.MoveToLevelTool",     new IMAGE_BRUSH("Icons/icon_Landscape_Tool_MoveToLevel_40x",     Icon40x40));
-			Set("LandscapeEditor.SelectComponentTool.Small", new IMAGE_BRUSH("Icons/icon_Landscape_Tool_Selection_20x",       Icon20x20));
-			Set("LandscapeEditor.AddComponentTool.Small",    new IMAGE_BRUSH("Icons/icon_Landscape_Tool_AddComponent_20x",    Icon20x20));
-			Set("LandscapeEditor.DeleteComponentTool.Small", new IMAGE_BRUSH("Icons/icon_Landscape_Tool_DeleteComponent_20x", Icon20x20));
-			Set("LandscapeEditor.MoveToLevelTool.Small",     new IMAGE_BRUSH("Icons/icon_Landscape_Tool_MoveToLevel_20x",     Icon20x20));
-
-			Set("LandscapeEditor.RegionSelectTool",    new IMAGE_BRUSH("Icons/icon_Landscape_Tool_Mask_40x", Icon40x40));
-			Set("LandscapeEditor.RegionCopyPasteTool", new IMAGE_BRUSH("Icons/icon_Landscape_Tool_Copy_40x", Icon40x40));
-			Set("LandscapeEditor.RegionSelectTool.Small",    new IMAGE_BRUSH("Icons/icon_Landscape_Tool_Mask_20x", Icon20x20));
-			Set("LandscapeEditor.RegionCopyPasteTool.Small", new IMAGE_BRUSH("Icons/icon_Landscape_Tool_Copy_20x", Icon20x20));
-
-			Set("LandscapeEditor.MirrorTool",       new IMAGE_BRUSH("Icons/icon_Landscape_Tool_Mirror_40x", Icon40x40));
-			Set("LandscapeEditor.MirrorTool.Small", new IMAGE_BRUSH("Icons/icon_Landscape_Tool_Mirror_20x", Icon20x20));
-
-			Set("LandscapeEditor.SplineTool",       new IMAGE_BRUSH("Icons/icon_Landscape_Tool_Splines_40x", Icon40x40));
-			Set("LandscapeEditor.SplineTool.Small", new IMAGE_BRUSH("Icons/icon_Landscape_Tool_Splines_20x", Icon20x20));
-
-			// Brush Sets
-			Set("LandscapeEditor.CircleBrush",        new IMAGE_BRUSH("Icons/icon_Landscape_Brush_Circle_smooth_40x", Icon40x40));
-			Set("LandscapeEditor.AlphaBrush",         new IMAGE_BRUSH("Icons/icon_Landscape_Brush_Alpha_40x",     Icon40x40));
-			Set("LandscapeEditor.AlphaBrush_Pattern", new IMAGE_BRUSH("Icons/icon_Landscape_Brush_Pattern_40x",   Icon40x40));
-			Set("LandscapeEditor.ComponentBrush",     new IMAGE_BRUSH("Icons/icon_Landscape_Brush_Component_40x", Icon40x40));
-			Set("LandscapeEditor.GizmoBrush",         new IMAGE_BRUSH("Icons/icon_Landscape_Brush_Gizmo_40x",     Icon40x40));
-			Set("LandscapeEditor.CircleBrush.Small",        new IMAGE_BRUSH("Icons/icon_Landscape_Brush_Circle_smooth_20x", Icon20x20));
-			Set("LandscapeEditor.AlphaBrush.Small",         new IMAGE_BRUSH("Icons/icon_Landscape_Brush_Alpha_20x",     Icon20x20));
-			Set("LandscapeEditor.AlphaBrush_Pattern.Small", new IMAGE_BRUSH("Icons/icon_Landscape_Brush_Pattern_20x",   Icon20x20));
-			Set("LandscapeEditor.ComponentBrush.Small",     new IMAGE_BRUSH("Icons/icon_Landscape_Brush_Component_20x", Icon20x20));
-			Set("LandscapeEditor.GizmoBrush.Small",         new IMAGE_BRUSH("Icons/icon_Landscape_Brush_Gizmo_20x",     Icon20x20));
-
-			// Brushes
-			Set("LandscapeEditor.CircleBrush_Smooth",    new IMAGE_BRUSH("Icons/icon_Landscape_Brush_Circle_smooth_40x",    Icon40x40));
-			Set("LandscapeEditor.CircleBrush_Linear",    new IMAGE_BRUSH("Icons/icon_Landscape_Brush_Circle_linear_40x",    Icon40x40));
-			Set("LandscapeEditor.CircleBrush_Spherical", new IMAGE_BRUSH("Icons/icon_Landscape_Brush_Circle_spherical_40x", Icon40x40));
-			Set("LandscapeEditor.CircleBrush_Tip",       new IMAGE_BRUSH("Icons/icon_Landscape_Brush_Circle_tip_40x",       Icon40x40));
-			Set("LandscapeEditor.CircleBrush_Smooth.Small",    new IMAGE_BRUSH("Icons/icon_Landscape_Brush_Circle_smooth_20x",    Icon20x20));
-			Set("LandscapeEditor.CircleBrush_Linear.Small",    new IMAGE_BRUSH("Icons/icon_Landscape_Brush_Circle_linear_20x",    Icon20x20));
-			Set("LandscapeEditor.CircleBrush_Spherical.Small", new IMAGE_BRUSH("Icons/icon_Landscape_Brush_Circle_spherical_20x", Icon20x20));
-			Set("LandscapeEditor.CircleBrush_Tip.Small",       new IMAGE_BRUSH("Icons/icon_Landscape_Brush_Circle_tip_20x",       Icon20x20));
-		}
-
-		else
 		{
 			// Tools
 			Set("LandscapeEditor.NewLandscape",					new IMAGE_BRUSH("Icons/Landscape/Landscape_NewLandscape_x40",		Icon20x20));
diff --git a/Engine/Source/Editor/EditorStyle/Public/Classes/EditorStyleSettings.h b/Engine/Source/Editor/EditorStyle/Public/Classes/EditorStyleSettings.h
--- a/Engine/Source/Editor/EditorStyle/Public/Classes/EditorStyleSettings.h
+++ b/Engine/Source/Editor/EditorStyle/Public/Classes/EditorStyleSettings.h
@@ -153,10 +153,6 @@ public:
 	UPROPERTY(EditAnywhere, config, Category=UserInterface)
 	uint32 bEnableColorizedEditorTabs : 1;
 
-	/** If enabled, the modes tab will revert to the previous mode UI */
-	UPROPERTY(EditAnywhere, config, AdvancedDisplay, Category = UserInterface, meta=(ConfigRestartRequired = true))
-	uint32 bEnableLegacyEditorModeUI : 1;
-
 public:
 
 	/**
diff --git a/Engine/Source/Editor/FoliageEdit/Private/FoliageEdModeToolkit.cpp b/Engine/Source/Editor/FoliageEdit/Private/FoliageEdModeToolkit.cpp
--- a/Engine/Source/Editor/FoliageEdit/Private/FoliageEdModeToolkit.cpp
+++ b/Engine/Source/Editor/FoliageEdit/Private/FoliageEdModeToolkit.cpp
@@ -67,10 +67,7 @@ void FFoliageEdModeToolkit::NotifyFoliageTypeMeshChanged(class UFoliageType* Fol
 
 void FFoliageEdModeToolkit::GetToolPaletteNames(TArray<FName>& InPaletteName) const
 {
-	if (!GetDefault<UEditorStyleSettings>()->bEnableLegacyEditorModeUI)
-	{
-		InPaletteName = FoliagePaletteNames;
-	}
+	InPaletteName = FoliagePaletteNames;
 }
 
 FText FFoliageEdModeToolkit::GetToolPaletteDisplayName(FName PaletteName) const
diff --git a/Engine/Source/Editor/LandscapeEditor/Private/SLandscapeEditor.cpp b/Engine/Source/Editor/LandscapeEditor/Private/SLandscapeEditor.cpp
--- a/Engine/Source/Editor/LandscapeEditor/Private/SLandscapeEditor.cpp
+++ b/Engine/Source/Editor/LandscapeEditor/Private/SLandscapeEditor.cpp
@@ -9,7 +9,6 @@
 #include "SlateOptMacros.h"
 #include "Widgets/Layout/SBox.h"
 #include "Widgets/Notifications/SErrorText.h"
-#include "Classes/EditorStyleSettings.h"
 #include "EditorStyleSet.h"
 #include "EditorModeManager.h"
 #include "EditorModes.h"
@@ -135,11 +134,7 @@ void FLandscapeToolKit::Init(const TSharedPtr<IToolkitHost>& InitToolkitHost)
 #undef MAP_BRUSH
 
 	LandscapeEditorWidgets = SNew(SLandscapeEditor, SharedThis(this));
-
-	if (!GetDefault<UEditorStyleSettings>()->bEnableLegacyEditorModeUI)
-	{
-		BrushesWidgets = StaticCastSharedRef<FLandscapeEditorDetails>(FLandscapeEditorDetails::MakeInstance());
-	}
+	BrushesWidgets = StaticCastSharedRef<FLandscapeEditorDetails>(FLandscapeEditorDetails::MakeInstance());
 
 	FModeToolkit::Init(InitToolkitHost);
 }
@@ -168,10 +163,7 @@ const TArray<FName> FLandscapeToolKit::PaletteNames = { LandscapeEditorNames::Ma
 
 void FLandscapeToolKit::GetToolPaletteNames(TArray<FName>& InPaletteName) const
 {
-	if (!GetDefault<UEditorStyleSettings>()->bEnableLegacyEditorModeUI)
-	{
-		InPaletteName = PaletteNames;
-	}
+	InPaletteName = PaletteNames;
 }
 
 FText FLandscapeToolKit::GetToolPaletteDisplayName(FName PaletteName) const
@@ -444,15 +436,6 @@ BEGIN_SLATE_FUNCTION_BUILD_OPTIMIZATION
 void SLandscapeEditor::Construct(const FArguments& InArgs, TSharedRef<FLandscapeToolKit> InParentToolkit)
 {
 	ParentToolkit = InParentToolkit;
-	TSharedRef<FUICommandList> CommandList = InParentToolkit->GetToolkitCommands();
-
-	// Modes:
-	FToolBarBuilder ModeSwitchButtons(CommandList, FMultiBoxCustomization::None);
-	{
-		ModeSwitchButtons.AddToolBarButton(FLandscapeEditorCommands::Get().ManageMode, NAME_None, LOCTEXT("Mode.Manage", "Manage"), LOCTEXT("Mode.Manage.Tooltip", "Contains tools to add a new landscape, import/export landscape, add/remove components and manage streaming"));
-		ModeSwitchButtons.AddToolBarButton(FLandscapeEditorCommands::Get().SculptMode, NAME_None, LOCTEXT("Mode.Sculpt", "Sculpt"), LOCTEXT("Mode.Sculpt.Tooltip", "Contains tools that modify the shape of a landscape"));
-		ModeSwitchButtons.AddToolBarButton(FLandscapeEditorCommands::Get().PaintMode,  NAME_None, LOCTEXT("Mode.Paint",  "Paint"),  LOCTEXT("Mode.Paint.Tooltip",  "Contains tools that paint materials on to a landscape"));
-	}
 
 	FPropertyEditorModule& PropertyEditorModule = FModuleManager::LoadModuleChecked<FPropertyEditorModule>("PropertyEditor");
 	FDetailsViewArgs DetailsViewArgs(false, false, false,FDetailsViewArgs::HideNameArea);
@@ -487,18 +470,6 @@ void SLandscapeEditor::Construct(const FArguments& InArgs, TSharedRef<FLandscape
 			.Padding(4, 0, 4, 5)
 			[
 				SNew(SOverlay)
-				+ SOverlay::Slot()
-				[
-					SNew(SBorder)
-					.Visibility_Lambda( [] () -> EVisibility { return GetDefault<UEditorStyleSettings>()->bEnableLegacyEditorModeUI ? EVisibility::Visible : EVisibility::Collapsed; } )
-					.BorderImage(FEditorStyle::GetBrush("ToolPanel.GroupBorder"))
-					.HAlign(HAlign_Center)
-					[
-
-						ModeSwitchButtons.MakeWidget()
-					]
-				]
-
 				// Tutorial link
 				+ SOverlay::Slot()
 				.HAlign(HAlign_Right)
diff --git a/Engine/Source/Editor/LevelEditor/Private/LevelEditor.cpp b/Engine/Source/Editor/LevelEditor/Private/LevelEditor.cpp
--- a/Engine/Source/Editor/LevelEditor/Private/LevelEditor.cpp
+++ b/Engine/Source/Editor/LevelEditor/Private/LevelEditor.cpp
@@ -210,13 +210,6 @@ TSharedRef<SDockTab> FLevelEditorModule::SpawnLevelEditor( const FSpawnTabArgs&
 		GLevelEditorModeTools().DeactivateAllModes();
 		GLevelEditorModeTools().ActivateDefaultMode();
 
-		if (GetDefault<UEditorStyleSettings>()->bEnableLegacyEditorModeUI)
-		{
-			// In legacy mode this toolbox should always be open
-			static const FTabId ToolboxTabId("LevelEditorToolBox");
-			LevelEditorTabManager->TryInvokeTab(ToolboxTabId);
-		}
-
 		LevelEditorCreatedEvent.Broadcast(LevelEditorTmp);
 
 		TSharedRef<SProjectBadge> ProjectBadge = SNew(SProjectBadge);
diff --git a/Engine/Source/Editor/LevelEditor/Private/SLevelEditor.cpp b/Engine/Source/Editor/LevelEditor/Private/SLevelEditor.cpp
--- a/Engine/Source/Editor/LevelEditor/Private/SLevelEditor.cpp
+++ b/Engine/Source/Editor/LevelEditor/Private/SLevelEditor.cpp
@@ -690,17 +690,14 @@ TSharedRef<SDockTab> SLevelEditor::SpawnLevelEditorTab( const FSpawnTabArgs& Arg
 	}
 	else if (TabIdentifier == LevelEditorTabIds::PlacementBrowser)
 	{
-		if(!GetDefault<UEditorStyleSettings>()->bEnableLegacyEditorModeUI)
-		{
-			return
-				SNew(SDockTab)
-				.Icon(FEditorStyle::GetBrush("LevelEditor.Tabs.PlacementBrowser"))
-				.Label(NSLOCTEXT("LevelEditor", "PlacementBrowserTitle", "Place Actors"))
-				.AddMetaData<FTutorialMetaData>(FTutorialMetaData(TEXT("PlacementBrowser"), TEXT("PlacementBrowser")))
-				[
-					IPlacementModeModule::Get().CreatePlacementModeBrowser()
-				];
-		}
+		return
+			SNew(SDockTab)
+			.Icon(FEditorStyle::GetBrush("LevelEditor.Tabs.PlacementBrowser"))
+			.Label(NSLOCTEXT("LevelEditor", "PlacementBrowserTitle", "Place Actors"))
+			.AddMetaData<FTutorialMetaData>(FTutorialMetaData(TEXT("PlacementBrowser"), TEXT("PlacementBrowser")))
+			[
+				IPlacementModeModule::Get().CreatePlacementModeBrowser()
+			];
 	}
 	else if( TabIdentifier == LevelEditorTabIds::LevelEditorBuildAndSubmit )
 	{
@@ -1490,7 +1487,7 @@ void SLevelEditor::OnEditorModeIdChanged(const FEditorModeID& ModeChangedID, boo
 				ToolboxTab->RequestCloseTab();
 			}
 		}
-		else if (!GetDefault<UEditorStyleSettings>()->bEnableLegacyEditorModeUI)
+		else
 		{
 			LevelEditorTabManager->TryInvokeTab(LevelEditorTabIds::LevelEditorToolBox);
 		}
@@ -1533,7 +1530,7 @@ void SLevelEditor::RefreshEditorModeCommands()
 				Commands.EditorModeCommands[CommandIndex],
 				FExecuteAction::CreateStatic( &SLevelEditor::ToggleEditorMode, Mode.ID ),
 				FCanExecuteAction(),
-				FIsActionChecked::CreateStatic( &SLevelEditor::IsModeActive, Mode.ID ));
+				FIsActionChecked::CreateSP( GLevelEditorModeTools().AsShared(), &FEditorModeTools::IsModeActive, Mode.ID ));
 		}
 
 		CommandIndex++;
diff --git a/Engine/Source/Editor/LevelEditor/Private/SLevelEditorToolBox.cpp b/Engine/Source/Editor/LevelEditor/Private/SLevelEditorToolBox.cpp
--- a/Engine/Source/Editor/LevelEditor/Private/SLevelEditorToolBox.cpp
+++ b/Engine/Source/Editor/LevelEditor/Private/SLevelEditorToolBox.cpp
@@ -141,14 +141,7 @@ void SLevelEditorToolBox::UpdateModeLegacyToolBar()
 		}
 	}
 
-	if (GetDefault<UEditorStyleSettings>()->bEnableLegacyEditorModeUI)
-	{
-		ModeToolBarContainer->SetContent(EditorModeTools.MakeWidget());
-	}
-	else
-	{
-		ModeToolBarContainer->SetVisibility(EVisibility::Collapsed);
-	}
+	ModeToolBarContainer->SetVisibility(EVisibility::Collapsed);
 
 	const TArray<TSharedPtr<IToolkit>>& HostedToolkits = LevelEditor.Pin()->GetHostedToolkits();
 	for(const TSharedPtr<IToolkit>& HostedToolkitIt : HostedToolkits)
diff --git a/Engine/Source/Editor/VREditor/VREditorModeManager.cpp b/Engine/Source/Editor/VREditor/VREditorModeManager.cpp
--- a/Engine/Source/Editor/VREditor/VREditorModeManager.cpp
+++ b/Engine/Source/Editor/VREditor/VREditorModeManager.cpp
@@ -103,27 +103,6 @@ void FVREditorModeManager::EnableVREditor( const bool bEnable, const bool bForce
 					bUsingDefaultInteractors = (InteractorClassSoft.Get() == UVREditorInteractor::StaticClass());
 				}
 			}
-			if ( bUsingDefaultInteractors && StyleSettings && !StyleSettings->bEnableLegacyEditorModeUI)
-			{
-				FSuppressableWarningDialog::FSetupInfo SetupInfo(LOCTEXT("VRModeLegacyModeUIEntry_Message", "VR Mode currently requires that legacy editor mode UI be enabled.  Without this, modes like mesh paint, landscape, and foliage will not function.  Enable Legacy editor mode UI (Requires restart)?"),
-					LOCTEXT("VRModeEntry_Title", "Entering VR Mode - Experimental"), "Warning_VRModeLegacyModeUIEntry", GEditorSettingsIni);
-
-				SetupInfo.ConfirmText = LOCTEXT("VRModeLegacyModeUIEntry_ConfirmText", "Enable and Restart");
-				SetupInfo.CancelText = LOCTEXT("VRModeLegacyModeUIEntry_CancelText", "Don't Enable");
-				SetupInfo.bDefaultToSuppressInTheFuture = false;
-
-				FSuppressableWarningDialog VRModeVRModeLegacyModeUIWarning(SetupInfo);
-
-				if (VRModeVRModeLegacyModeUIWarning.ShowModal() != FSuppressableWarningDialog::Cancel)
-				{
-					StyleSettings->bEnableLegacyEditorModeUI = true;
-					StyleSettings->SaveConfig();
-					FUnrealEdMisc::Get().RestartEditor(true);
-					return;
-				}
-
-			}
-
 
 			{
 				FSuppressableWarningDialog::FSetupInfo SetupInfo(LOCTEXT("VRModeEntry_Message", "VR Mode enables you to work on your project in virtual reality using motion controllers. This feature is still under development, so you may experience bugs or crashes while using it."),
