Updated Modes Toolbox Palette Style

diff --git a/Engine/Source/Editor/EditorStyle/Private/StarshipStyle.cpp b/Engine/Source/Editor/EditorStyle/Private/StarshipStyle.cpp
--- a/Engine/Source/Editor/EditorStyle/Private/StarshipStyle.cpp
+++ b/Engine/Source/Editor/EditorStyle/Private/StarshipStyle.cpp
@@ -731,9 +731,10 @@ void FStarshipEditorStyle::FStyle::SetupGeneralStyles()
 #if WITH_EDITOR || (IS_PROGRAM && WITH_UNREAL_DEVELOPER_TOOLS)
 
 	{
-		Set("SystemWideCommands.FindInContentBrowser", new CORE_IMAGE_BRUSH_SVG("Starship/Common/search", Icon20x20));
-		Set("SystemWideCommands.FindInContentBrowser.Small", new CORE_IMAGE_BRUSH_SVG("Starship/Common/search", Icon16x16));
+		Set("SystemWideCommands.FindInContentBrowser", new IMAGE_BRUSH_SVG("Starship/Common/ContentBrowser", Icon20x20));
+		Set("SystemWideCommands.FindInContentBrowser.Small", new IMAGE_BRUSH_SVG("Starship/Common/ContentBrowser", Icon16x16));
 	}
+	
 
 	// PList Editor
 	{
@@ -1775,12 +1776,15 @@ void FStarshipEditorStyle::FStyle::SetupGeneralStyles()
 
 		Set( "PaletteToolBar.Tab",  FCheckBoxStyle()
 			.SetCheckBoxType(            ESlateCheckBoxType::ToggleButton)
-			.SetCheckedImage(            FSlateNoResource())
-			.SetCheckedHoveredImage(     FSlateNoResource())
-			.SetCheckedPressedImage(     FSlateNoResource())
-			.SetUncheckedImage(          FSlateNoResource())
-			.SetUncheckedHoveredImage(   FSlateNoResource())
-			.SetUncheckedPressedImage(   FSlateNoResource())
+
+			.SetCheckedImage(            FSlateRoundedBoxBrush(FStyleColors::Input, 2.0f))
+			.SetCheckedHoveredImage(     FSlateRoundedBoxBrush(FStyleColors::Input, 2.0f))
+			.SetCheckedPressedImage(     FSlateRoundedBoxBrush(FStyleColors::Input, 2.0f))
+
+			.SetUncheckedImage(          FSlateRoundedBoxBrush(FStyleColors::Secondary, 2.0f))
+			.SetUncheckedHoveredImage(   FSlateRoundedBoxBrush(FStyleColors::Hover, 2.0f))
+			.SetUncheckedPressedImage(   FSlateRoundedBoxBrush(FStyleColors::Secondary, 2.0f))
+
 			.SetForegroundColor(         FStyleColors::Foreground)
 			.SetHoveredForegroundColor(  FStyleColors::ForegroundHover)
 			.SetPressedForegroundColor(  FStyleColors::ForegroundHover)
@@ -1797,12 +1801,12 @@ void FStarshipEditorStyle::FStyle::SetupGeneralStyles()
 			/* Create style for "ToolBar.ToggleButton" widget ... */
 		const FCheckBoxStyle ToolBarToggleButtonCheckBoxStyle = FCheckBoxStyle()
 				.SetCheckBoxType( ESlateCheckBoxType::ToggleButton )
-				.SetUncheckedImage(       FSlateRoundedBoxBrush(FStyleColors::Background, 2.0))
-				.SetUncheckedPressedImage(FSlateRoundedBoxBrush(FStyleColors::Background, 2.0))
-				.SetUncheckedHoveredImage(FSlateRoundedBoxBrush(FStyleColors::Background, 2.0))
-				.SetCheckedImage(         FSlateRoundedBoxBrush(FStyleColors::Recessed, 2.0))
-				.SetCheckedHoveredImage(  FSlateRoundedBoxBrush(FStyleColors::Recessed, 2.0))
-				.SetCheckedPressedImage(  FSlateRoundedBoxBrush(FStyleColors::Recessed, 2.0))
+				.SetUncheckedImage(        FSlateNoResource() )
+				.SetUncheckedPressedImage( FSlateNoResource() )
+				.SetUncheckedHoveredImage( FSlateNoResource() )
+				.SetCheckedImage(          FSlateNoResource() )
+				.SetCheckedHoveredImage(   FSlateNoResource() )
+				.SetCheckedPressedImage(   FSlateNoResource() )
 				.SetForegroundColor(       FStyleColors::Foreground)
 				.SetHoveredForegroundColor(FStyleColors::ForegroundHover)
 				.SetPressedForegroundColor(FStyleColors::ForegroundHover)
@@ -1812,9 +1816,9 @@ void FStarshipEditorStyle::FStyle::SetupGeneralStyles()
 		
 		PaletteToolBarStyle.SetToggleButtonStyle(ToolBarToggleButtonCheckBoxStyle);
 		PaletteToolBarStyle.SetButtonStyle(FButtonStyle(Button)
-			.SetNormal( FSlateRoundedBoxBrush(FStyleColors::Background, 2.0))
-			.SetPressed(FSlateRoundedBoxBrush(FStyleColors::Background, 2.0))
-			.SetHovered(FSlateRoundedBoxBrush(FStyleColors::Background, 2.0))
+			.SetNormal(  FSlateNoResource() )
+			.SetPressed( FSlateNoResource() )
+			.SetHovered( FSlateNoResource() )
 			.SetNormalForeground(         FStyleColors::Foreground)
 			.SetHoveredForeground(        FStyleColors::ForegroundHover)
 			.SetPressedForeground(        FStyleColors::ForegroundHover)
