New round of class icons.

diff --git a/Engine/Source/Editor/EditorStyle/Private/StarshipStyle.cpp b/Engine/Source/Editor/EditorStyle/Private/StarshipStyle.cpp
--- a/Engine/Source/Editor/EditorStyle/Private/StarshipStyle.cpp
+++ b/Engine/Source/Editor/EditorStyle/Private/StarshipStyle.cpp
@@ -5722,7 +5722,6 @@ void FStarshipEditorStyle::FStyle::SetupClassIconsAndThumbnails()
 		Set("ClassIcon.MaterialBillboardComponent", new IMAGE_BRUSH("Icons/ActorIcons/MaterialSprite_16px", Icon16x16));
 		Set("ClassIcon.SceneCaptureComponent2D", new IMAGE_BRUSH("Icons/AssetIcons/SceneCapture2D_16x", Icon16x16));
 		Set("ClassIcon.SceneCaptureComponent", new IMAGE_BRUSH("Icons/ActorIcons/SceneCapture_16px", Icon16x16));
-		Set("ClassIcon.SceneComponent", new IMAGE_BRUSH("Icons/ActorIcons/Scene_16px", Icon16x16));
 		Set("ClassIcon.SphereComponent", new IMAGE_BRUSH("Icons/ActorIcons/Sphere_16px", Icon16x16));
 		Set("ClassIcon.SplineComponent", new IMAGE_BRUSH("Icons/ActorIcons/Spline_16px", Icon16x16));
 		Set("ClassIcon.BillboardComponent", new IMAGE_BRUSH("Icons/ActorIcons/SpriteComponent_16px", Icon16x16));
@@ -5746,9 +5745,6 @@ void FStarshipEditorStyle::FStyle::SetupClassIconsAndThumbnails()
 		const TCHAR* AssetTypes[] = {
 
 			TEXT("AbilitySystemComponent"),
-			TEXT("ActorComponent"),
-			TEXT("AimOffsetBlendSpace"),
-			TEXT("AimOffsetBlendSpace1D"),
 			TEXT("AIPerceptionComponent"),
 			TEXT("AnimationModifier"),		
 			TEXT("AnimBlueprint"),
@@ -5759,86 +5755,53 @@ void FStarshipEditorStyle::FStyle::SetupClassIconsAndThumbnails()
 			TEXT("AnimSequence"),
 			TEXT("AnimationSharingSetup"),
 			TEXT("ApplicationLifecycleComponent"),
-			TEXT("BehaviorTree"),
 			TEXT("BlackboardData"),
 			TEXT("BlendSpace"),
 			TEXT("BlendSpace1D"),
-			TEXT("Blueprint"),
 			TEXT("BlueprintFunctionLibrary"),
 			TEXT("BlueprintInterface"),
 			TEXT("BlueprintMacroLibrary"),
-			TEXT("ButtonStyleAsset"),
 			TEXT("CableComponent"),
 			TEXT("CameraAnim"),
-			TEXT("CharacterMovementComponent"),
 			TEXT("Class"),
-			TEXT("CullDistanceVolume"),
-			TEXT("CurveBase"),
 			TEXT("DataAsset"),
-			TEXT("DataTable"),
 			TEXT("Default"),
-			TEXT("DialogueWave"),
-			TEXT("DialogueVoice"),
 			TEXT("DirectionalLightMovable"),
 			TEXT("DirectionalLightStatic"),
 			TEXT("DirectionalLightStationary"),
-			TEXT("EditorTutorial"),
 			TEXT("EnvQuery"),
-			TEXT("FileMediaSource"),
 			TEXT("Font"),
 			TEXT("FontFace"),
 			TEXT("ForceFeedbackEffect"),
-			TEXT("GameModeBase"),
-			TEXT("GameStateBase"),
 			TEXT("Interface"),
 			TEXT("InterpData"),
 			TEXT("LevelSequence"),
 			TEXT("LightmassCharacterIndirectDetailVolume"),
-			TEXT("VolumetricLightmapDensityVolume"),
 			TEXT("MassiveLODOverrideVolume"),
 			TEXT("Material"),
 			TEXT("MaterialFunction"),
 			TEXT("MaterialInstanceConstant"),
 			TEXT("MaterialParameterCollection"),
 			TEXT("MatineeActor"),
-			TEXT("MediaPlayer"),
-			TEXT("MediaTexture"),
 			TEXT("MultiFont"),
-			TEXT("NavModifierComponent"),
 			TEXT("ObjectLibrary"),
 			TEXT("ParticleSystem"),
-			TEXT("PawnNoiseEmitterComponent"),
-			TEXT("PawnSensingComponent"),
 			TEXT("PhysicalMaterial"),
 			TEXT("PhysicsAsset"),
 			TEXT("PhysicsConstraintComponent"),
-			TEXT("PhysicsHandleComponent"),
 			TEXT("PhysicsThrusterComponent"),
 			TEXT("PlatformMediaSource"),
 			TEXT("PoseAsset"),
-			TEXT("ProjectileMovementComponent"),
-			TEXT("RadialForceComponent"),
-			TEXT("RotatingMovementComponent"),
-			TEXT("SceneCaptureCube"),
-			TEXT("SceneComponent"),
 			TEXT("SkyLightComponent"),
-			TEXT("SkeletalMesh"),
-			TEXT("Skeleton"),
-			TEXT("SlateBrushAsset"),
 			TEXT("SlateWidgetStyleAsset"),
 			TEXT("StringTable"),
 			TEXT("SpotLightMovable"),
 			TEXT("SpotLightStatic"),
 			TEXT("SpotLightStationary"),
-			TEXT("SpringArmComponent"),
 			TEXT("StaticMesh"),
-			TEXT("StreamMediaSource"),
 			TEXT("SubsurfaceProfile"),
 			TEXT("Texture2D"),
 			TEXT("TextureRenderTarget2D"),
-			TEXT("TextureRenderTargetCube"),
-			TEXT("TouchInterface"),
-			TEXT("UserDefinedCaptureProtocol"),
 			TEXT("UserDefinedEnum"),
 			TEXT("UserDefinedStruct"),
 			TEXT("UserWidget"),
@@ -5850,7 +5813,6 @@ void FStarshipEditorStyle::FStyle::SetupClassIconsAndThumbnails()
 			TEXT("Cone"),
 			TEXT("Plane"),
 			TEXT("FoliageType_Actor"),
-			TEXT("LevelInstance")
 		};
 
 		for (int32 TypeIndex = 0; TypeIndex < UE_ARRAY_COUNT(AssetTypes); ++TypeIndex)
@@ -5862,47 +5824,77 @@ void FStarshipEditorStyle::FStyle::SetupClassIconsAndThumbnails()
 
 		const FClassIconInfo AssetTypesSVG[] = {
 			{TEXT("Actor")},
+			{TEXT("ActorComponent")},
 			{TEXT("AIController")},
+			{TEXT("AimOffsetBlendSpace")},
+			{TEXT("AimOffsetBlendSpace1D")},
+			{TEXT("AmbientSound")},
 			{TEXT("AtmosphericFog")},
+			{TEXT("AudioVolume")},
+			{TEXT("BehaviorTree")},
 			{TEXT("BlockingVolume")},
+			{TEXT("Blueprint")},
 			{TEXT("BoxReflectionCapture")},
 			{TEXT("Brush")},
+			{TEXT("ButtonStyleAsset")},
 			{TEXT("CableActor")},
 			{TEXT("CameraActor")},
 			{TEXT("CameraBlockingVolume")},
 			{TEXT("CameraRig_Crane")},
 			{TEXT("CameraRig_Rail")},
 			{TEXT("Character")},
+			{TEXT("CharacterMovementComponent")},
 			{TEXT("CineCameraActor")},
 			{TEXT("CompositingElement")},
+			{TEXT("CullDistanceVolume")},
+			{TEXT("CurveBase")},
+			{TEXT("DataTable")},
 			{TEXT("DecalActor")},
 			{TEXT("DefaultPawn")},
+			{TEXT("DialogueVoice")},
+			{TEXT("DialogueWave")},
 			{TEXT("DirectionalLight")},
 			{TEXT("DocumentationActor")},
+			{TEXT("EditorTutorial")},
 			{TEXT("Emitter")},
 			{TEXT("EmptyActor")},
 			{TEXT("ExponentialHeightFog")},
+			{TEXT("FileMediaOutput")},
+			{TEXT("FileMediaSource")},
 			{TEXT("ForceFeedback")},
+			{TEXT("GameModeBase")},
+			{TEXT("GameStateBase")},
+			{TEXT("GeometryCollection")},
 			{TEXT("GroupActor")},
 			{TEXT("HUD")},
 			{TEXT("ImagePlate")},
 			{TEXT("KillZVolume")},
 			{TEXT("Landscape")},
 			{TEXT("LevelBounds")},
+			{TEXT("LevelInstance")},
 			{TEXT("LevelScriptActor")},
 			{TEXT("LevelSequenceActor")},
-			{TEXT("LevelStreamingVolume"), false},
-			{TEXT("LightmassImportanceVolume"), false},
+			{TEXT("LevelStreamingVolume")},
+			{TEXT("LightmassCharacterDetailIndirectVolume")},
+			{TEXT("LightmassImportanceVolume")},
 			{TEXT("LightmassVolume")},
+			{TEXT("LiveLinkPreset")},
 			{TEXT("MaterialInstanceActor")},
+			{TEXT("MediaPlayer")},
+			{TEXT("MediaTexture")},
+			{TEXT("ModularSynthPresetBank")},
 			{TEXT("NavLink")},
-			{TEXT("NavLinkProxy"), false},
-			{TEXT("NavMeshBoundsVolume"), false},
-			{TEXT("NavModifierVolume"), false},
+			{TEXT("NavLinkProxy")},
+			{TEXT("NavMeshBoundsVolume")},
+			{TEXT("NavModifierComponent")},
+			{TEXT("NavModifierVolume")},
 			{TEXT("Note")},
 			{TEXT("PainCausingVolume")},
 			{TEXT("Pawn")},
-			{TEXT("PhysicsConstraintActor"), false},
+			{TEXT("PawnNoiseEmitterComponent")},
+			{TEXT("PawnSensingComponent")},
+			{TEXT("PhysicsConstraintActor")},
+			{TEXT("PhysicsHandleComponent")},
 			{TEXT("PhysicsThruster")},
 			{TEXT("PhysicsVolume")},
 			{TEXT("PlaneReflectionCapture")},
@@ -5914,27 +5906,55 @@ void FStarshipEditorStyle::FStyle::SetupClassIconsAndThumbnails()
 			{TEXT("PrecomputedVisibilityVolume")},
 			{TEXT("ProceduralFoliageBlockingVolume")},
 			{TEXT("ProceduralFoliageVolume")},
+			{TEXT("ProjectileMovementComponent")},
 			{TEXT("RadialForceActor")},
-			{TEXT("RectLight"), false},
+			{TEXT("RadialForceComponent")},
+			{TEXT("RectLight")},
 			{TEXT("ReflectionCapture")},
+			{TEXT("ReverbEffect")},
+			{TEXT("RotatingMovementComponent")},
 			{TEXT("SceneCapture2D")},
+			{TEXT("SceneCaptureCube")},
+			{TEXT("SceneComponent")},
 			{TEXT("SkeletalMeshActor")},
+			{TEXT("Skeleton")},
 			{TEXT("SkyAtmosphere")},
 			{TEXT("SkyLight")},
+			{TEXT("SlateBrushAsset")},
+			{TEXT("SoundAttenuation")},
+			{TEXT("SoundClass")},
+			{TEXT("SoundConcurrency")},
+			{TEXT("SoundCue")},
+			{TEXT("SoundEffectSourcePreset")},
+			{TEXT("SoundMix")},
 			{TEXT("SphereReflectionCapture")},
 			{TEXT("SpotLight")},
+			{TEXT("SpringArmComponent")},
 			{TEXT("StaticMeshActor")},
+			{TEXT("StreamMediaSource")},
+			{TEXT("TakePreset")},
+			{TEXT("TakeRecorderLevelSequenceSource")},
+			{TEXT("TakeRecorderLiveLinkSource")},
+			{TEXT("TakeRecorderMicrophoneAudioSource")},
+			{TEXT("TakeRecorderNearbySpawnedActorSource")},
+			{TEXT("TakeRecorderWorldSource")},
 			{TEXT("TargetPoint")},
+			{TEXT("TemplateSequence")},
 			{TEXT("TextRenderActor")},
+			{TEXT("TextureRenderTargetCube")},
+			{TEXT("TimeCodeSynchronizer")},
+			{TEXT("TouchInterface")},
 			{TEXT("TriggerBase")},
 			{TEXT("TriggerBox")},
 			{TEXT("TriggerCapsule")},
 			{TEXT("TriggerSphere")},
 			{TEXT("TriggerVolume")},
+			{TEXT("UserDefinedCaptureProtocol")},
 			{TEXT("VectorField")},
 			{TEXT("VectorFieldVolume")},
 			{TEXT("Volume")},
 			{TEXT("VolumetricCloud"), false},
+			{TEXT("VolumetricLightmapDensityVolume")},
 			{TEXT("WindDirectionalSource")}
 		};
 	
