Unify lock icons

diff --git a/Engine/Source/Developer/Merge/Private/SMergeGraphView.cpp b/Engine/Source/Developer/Merge/Private/SMergeGraphView.cpp
--- a/Engine/Source/Developer/Merge/Private/SMergeGraphView.cpp
+++ b/Engine/Source/Developer/Merge/Private/SMergeGraphView.cpp
@@ -604,7 +604,7 @@ FReply SMergeGraphView::OnToggleLockView()
 
 const FSlateBrush* SMergeGraphView::GetLockViewImage() const
 {
-	return bViewsAreLocked ? FEditorStyle::GetBrush("GenericLock") : FEditorStyle::GetBrush("GenericUnlock");
+	return bViewsAreLocked ? FEditorStyle::GetBrush("Icons.Lock") : FEditorStyle::GetBrush("Icons.Unlock");
 }
 
 #undef LOCTEXT_NAMESPACE
diff --git a/Engine/Source/Editor/DeviceProfileEditor/Private/SDeviceProfileEditor.cpp b/Engine/Source/Editor/DeviceProfileEditor/Private/SDeviceProfileEditor.cpp
--- a/Engine/Source/Editor/DeviceProfileEditor/Private/SDeviceProfileEditor.cpp
+++ b/Engine/Source/Editor/DeviceProfileEditor/Private/SDeviceProfileEditor.cpp
@@ -181,7 +181,7 @@ void SDeviceProfileSourceControl::Construct(const FArguments& InArgs)
 					.AutoWidth()
 					[
 						SNew(SImage)
-						.Image(FEditorStyle::GetBrush("GenericUnlock"))
+						.Image(FEditorStyle::GetBrush("Icons.Unlock"))
 					]
 					+ SHorizontalBox::Slot()
 					.VAlign(VAlign_Center)
@@ -208,7 +208,7 @@ void SDeviceProfileSourceControl::Construct(const FArguments& InArgs)
 					.AutoWidth()
 					[
 						SNew(SImage)
-						.Image(FEditorStyle::GetBrush("GenericLock"))
+						.Image(FEditorStyle::GetBrush("Icons.Lock"))
 					]
 					+ SHorizontalBox::Slot()
 					.VAlign(VAlign_Center)
diff --git a/Engine/Source/Editor/EditorStyle/Private/StarshipStyle.cpp b/Engine/Source/Editor/EditorStyle/Private/StarshipStyle.cpp
--- a/Engine/Source/Editor/EditorStyle/Private/StarshipStyle.cpp
+++ b/Engine/Source/Editor/EditorStyle/Private/StarshipStyle.cpp
@@ -651,11 +651,6 @@ void FStarshipEditorStyle::FStyle::SetupGeneralStyles()
 
 	Set( "MarqueeSelection", new BORDER_BRUSH( "Old/DashedBorder", FMargin(6.0f/32.0f) ) );
 
-	Set( "GenericLock", new IMAGE_BRUSH( "Icons/padlock_locked_16x", Icon16x16 ) );
-	Set( "GenericLock.Small", new IMAGE_BRUSH( "Icons/padlock_locked_16x", Icon16x16 ) );
-	Set( "GenericUnlock", new IMAGE_BRUSH( "Icons/padlock_unlocked_16x", Icon16x16 ) );
-	Set( "GenericUnlock.Small", new IMAGE_BRUSH( "Icons/padlock_unlocked_16x", Icon16x16 ) );
-
 	Set( "GenericPlay", new IMAGE_BRUSH( "Icons/generic_play_16x", Icon16x16 ) );
 	Set( "GenericPause", new IMAGE_BRUSH( "Icons/generic_pause_16x", Icon16x16 ) );
 	Set( "GenericStop", new IMAGE_BRUSH( "Icons/generic_stop_16x", Icon16x16 ) );
@@ -3107,10 +3102,10 @@ void FStarshipEditorStyle::FStyle::SetupPropertyEditorStyles()
 		Set( "PropertyWindow.Favorites_Disabled", new IMAGE_BRUSH( "Icons/EmptyStar_16x", Icon16x16 ) );
 		Set( "PropertyWindow.Locked", new CORE_IMAGE_BRUSH_SVG( "Starship/Common/lock", Icon16x16 ) );
 		Set( "PropertyWindow.Unlocked", new CORE_IMAGE_BRUSH_SVG( "Starship/Common/lock-unlocked", Icon16x16 ) );
-		Set( "PropertyWindow.DiffersFromDefault", new IMAGE_BRUSH( "/PropertyView/DiffersFromDefault_8x8", FVector2D(8,8) ) ) ;
+		Set( "PropertyWindow.DiffersFromDefault", new IMAGE_BRUSH_SVG( "Starship/Common/ResetToDefault", Icon16x16) ) ;
 		
-		Set( "PropertyWindow.NormalFont", DEFAULT_FONT("Regular", 8) );
-		Set( "PropertyWindow.BoldFont", DEFAULT_FONT( "Bold", 8 ) );
+		Set( "PropertyWindow.NormalFont", FStyleFonts::Get().Small);
+		Set( "PropertyWindow.BoldFont",FStyleFonts::Get().SmallBold);
 		Set( "PropertyWindow.ItalicFont", DEFAULT_FONT( "Italic", 8 ) );
 		Set( "PropertyWindow.FilterFont", DEFAULT_FONT( "Regular", 10 ) );
 		Set( "PropertyWindow.NoOverlayColor", new FSlateNoResource() );
@@ -5279,9 +5274,6 @@ void FStarshipEditorStyle::FStyle::SetupPersonaStyle()
 			.SetFont(ICON_FONT(10))
 			.SetColorAndOpacity(FLinearColor(0.f, 0.f, 0.f))
 		);
-
-		Set("FindResults.LockButton_Locked", new IMAGE_BRUSH("Icons/padlock_locked_16x", Icon16x16));
-		Set("FindResults.LockButton_Unlocked", new IMAGE_BRUSH("Icons/padlock_unlocked_16x", Icon16x16));
 	}
 
 	//Bookmarks
diff --git a/Engine/Source/Editor/Kismet/Private/FindInBlueprints.cpp b/Engine/Source/Editor/Kismet/Private/FindInBlueprints.cpp
--- a/Engine/Source/Editor/Kismet/Private/FindInBlueprints.cpp
+++ b/Engine/Source/Editor/Kismet/Private/FindInBlueprints.cpp
@@ -1664,11 +1664,11 @@ const FSlateBrush* SFindInBlueprints::OnGetLockButtonImage() const
 {
 	if (bIsLocked)
 	{
-		return FEditorStyle::GetBrush("FindResults.LockButton_Locked");
+		return FEditorStyle::GetBrush("Icons.Lock");
 	}
 	else
 	{
-		return FEditorStyle::GetBrush("FindResults.LockButton_Unlocked");
+		return FEditorStyle::GetBrush("Icons.Unlock");
 	}
 }
 
diff --git a/Engine/Source/Editor/Kismet/Private/SBlueprintDiff.cpp b/Engine/Source/Editor/Kismet/Private/SBlueprintDiff.cpp
--- a/Engine/Source/Editor/Kismet/Private/SBlueprintDiff.cpp
+++ b/Engine/Source/Editor/Kismet/Private/SBlueprintDiff.cpp
@@ -1282,7 +1282,7 @@ void SBlueprintDiff::OnToggleSplitViewMode()
 
 FSlateIcon SBlueprintDiff::GetLockViewImage() const
 {
-	return FSlateIcon(FEditorStyle::GetStyleSetName(), bLockViews ? "GenericLock" : "GenericUnlock");
+	return FSlateIcon(FEditorStyle::GetStyleSetName(), bLockViews ? "Icons.Lock" : "Icons.Unlock");
 }
 
 FSlateIcon SBlueprintDiff::GetSplitViewModeImage() const
diff --git a/Engine/Source/Editor/LevelEditor/Private/SSurfaceProperties.cpp b/Engine/Source/Editor/LevelEditor/Private/SSurfaceProperties.cpp
--- a/Engine/Source/Editor/LevelEditor/Private/SSurfaceProperties.cpp
+++ b/Engine/Source/Editor/LevelEditor/Private/SSurfaceProperties.cpp
@@ -953,7 +953,7 @@ FReply SSurfaceProperties::OnApplyScaling()
 
 const FSlateBrush* SSurfaceProperties::GetPreserveScaleRatioImage() const
 {
-	return bPreserveScaleRatio ? FEditorStyle::GetBrush( TEXT("GenericLock") ) : FEditorStyle::GetBrush( TEXT("GenericUnlock") ) ;
+	return bPreserveScaleRatio ? FEditorStyle::GetBrush( TEXT("Icons.Lock") ) : FEditorStyle::GetBrush( TEXT("Icons.Unlock") ) ;
 }
 
 ECheckBoxState SSurfaceProperties::IsPreserveScaleRatioChecked() const
