Large dump of new asset icons

diff --git a/Engine/Plugins/2D/Paper2D/Source/Paper2DEditor/Private/PaperStyle.cpp b/Engine/Plugins/2D/Paper2D/Source/Paper2DEditor/Private/PaperStyle.cpp
--- a/Engine/Plugins/2D/Paper2D/Source/Paper2DEditor/Private/PaperStyle.cpp
+++ b/Engine/Plugins/2D/Paper2D/Source/Paper2DEditor/Private/PaperStyle.cpp
@@ -288,8 +288,8 @@ void FPaperStyle::Initialize()
 
 		// UPaperSpriteAtlas assets
 		//@TODO: Paper2D: These icons don't match the naming scheme
-		StyleSet->Set("ClassIcon.PaperSpriteAtlas", new IMAGE_PLUGIN_BRUSH("Icons/AssetIcons/Paper2DSpriteAtlasGroup_16x", Icon16x16));
-		StyleSet->Set("ClassThumbnail.PaperSpriteAtlas", new IMAGE_PLUGIN_BRUSH("Icons/AssetIcons/Paper2DSpriteAtlasGroup_64x", Icon64x64));
+		StyleSet->Set("ClassIcon.PaperSpriteAtlas", new IMAGE_PLUGIN_BRUSH_SVG("Icons/AssetIcons/PaperSpriteAtlas_16", Icon16x16));
+		StyleSet->Set("ClassThumbnail.PaperSpriteAtlas", new IMAGE_PLUGIN_BRUSH_SVG("Icons/AssetIcons/PaperSpriteAtlas_64", Icon64x64));
 
 		//@TODO: UPaperSpriteSheet icons?
 
@@ -309,7 +309,7 @@ void FPaperStyle::Initialize()
 		StyleSet->Set("ClassIcon.PaperTerrainComponent", new IMAGE_PLUGIN_BRUSH_SVG("Icons/AssetIcons/PaperTerrainActor_16", Icon16x16));
 		StyleSet->Set("ClassThumbnail.PaperTerrainComponent", new IMAGE_PLUGIN_BRUSH_SVG("Icons/AssetIcons/PaperTerrainActor_16", Icon64x64));
 		StyleSet->Set("ClassIcon.PaperTerrainActor", new IMAGE_PLUGIN_BRUSH_SVG("Icons/AssetIcons/PaperTerrainActor_16", Icon16x16));
-		StyleSet->Set("ClassThumbnail.PaperTerrainActor", new IMAGE_PLUGIN_BRUSH_SVG("Icons/AssetIcons/PaperTerrainActor_16", Icon64x64));
+		StyleSet->Set("ClassThumbnail.PaperTerrainActor", new IMAGE_PLUGIN_BRUSH_SVG("Icons/AssetIcons/PaperTerrainActor_64", Icon64x64));
 	}
 
 	FSlateStyleRegistry::RegisterSlateStyle(*StyleSet.Get());
diff --git a/Engine/Plugins/Compositing/OpenColorIO/Source/OpenColorIOEditor/Private/OpenColorIOEditorModule.cpp b/Engine/Plugins/Compositing/OpenColorIO/Source/OpenColorIOEditor/Private/OpenColorIOEditorModule.cpp
--- a/Engine/Plugins/Compositing/OpenColorIO/Source/OpenColorIOEditor/Private/OpenColorIOEditorModule.cpp
+++ b/Engine/Plugins/Compositing/OpenColorIO/Source/OpenColorIOEditor/Private/OpenColorIOEditorModule.cpp
@@ -118,7 +118,7 @@ void FOpenColorIOEditorModule::CleanFeatureLevelDelegate()
 
 void FOpenColorIOEditorModule::RegisterStyle()
 {
-#define IMAGE_BRUSH(RelativePath, ...) FSlateImageBrush(StyleInstance->RootToContentDir(RelativePath, TEXT(".png")), __VA_ARGS__)
+#define IMAGE_BRUSH_SVG(RelativePath, ...) FSlateVectorImageBrush(StyleInstance->RootToContentDir(RelativePath, TEXT(".svg")), __VA_ARGS__)
 
 	StyleInstance = MakeUnique<FSlateStyleSet>("OpenColorIOStyle");
 
@@ -128,11 +128,11 @@ void FOpenColorIOEditorModule::RegisterStyle()
 		StyleInstance->SetContentRoot(FPaths::Combine(Plugin->GetContentDir(), TEXT("Editor/Icons")));
 	}
 
-	const FVector2D Icon20x20(20.0f, 20.0f);
+	const FVector2D Icon16x16(16.0f, 16.0f);
 	const FVector2D Icon64x64(64.0f, 64.0f);
 
-	StyleInstance->Set("ClassThumbnail.OpenColorIOConfiguration", new IMAGE_BRUSH("OpenColorIOConfigIcon_64x", Icon64x64));
-	StyleInstance->Set("ClassIcon.OpenColorIOConfiguration", new IMAGE_BRUSH("OpenColorIOConfigIcon_20x", Icon20x20));
+	StyleInstance->Set("ClassThumbnail.OpenColorIOConfiguration", new IMAGE_BRUSH_SVG("OpenColorIOConfiguration_64", Icon64x64));
+	StyleInstance->Set("ClassIcon.OpenColorIOConfiguration", new IMAGE_BRUSH_SVG("OpenColorIOConfiguration_16", Icon16x16));
 
 	FSlateStyleRegistry::RegisterSlateStyle(*StyleInstance.Get());
 
diff --git a/Engine/Plugins/Experimental/ChaosSolverPlugin/Source/ChaosSolverEditor/Private/Chaos/ChaosSolverEditorStyle.h b/Engine/Plugins/Experimental/ChaosSolverPlugin/Source/ChaosSolverEditor/Private/Chaos/ChaosSolverEditorStyle.h
--- a/Engine/Plugins/Experimental/ChaosSolverPlugin/Source/ChaosSolverEditor/Private/Chaos/ChaosSolverEditorStyle.h
+++ b/Engine/Plugins/Experimental/ChaosSolverPlugin/Source/ChaosSolverEditor/Private/Chaos/ChaosSolverEditorStyle.h
@@ -16,8 +16,8 @@ public:
 		SetContentRoot(FPaths::EnginePluginsDir() / TEXT("Experimental/ChaosSolverPlugin/Resources"));
 #endif
 
-		Set("ClassIcon.ChaosSolver", new FSlateImageBrush(RootToContentDir(TEXT("ChaosSolver_16x.png")), Icon16x16));
-		Set("ClassThumbnail.ChaosSolver", new FSlateImageBrush(RootToContentDir(TEXT("ChaosSolver_64x.png")), Icon64x64));
+		Set("ClassIcon.ChaosSolver", new FSlateVectorImageBrush(RootToContentDir(TEXT("ChaosSolver_16.svg")), Icon16x16));
+		Set("ClassThumbnail.ChaosSolver", new FSlateVectorImageBrush(RootToContentDir(TEXT("ChaosSolver_64.svg")), Icon64x64));
 
 		FSlateStyleRegistry::RegisterSlateStyle(*this);
 	}
diff --git a/Engine/Plugins/Experimental/ChaosVehiclesPlugin/Source/ChaosVehiclesEditor/Private/ChaosVehiclesEditorStyle.h b/Engine/Plugins/Experimental/ChaosVehiclesPlugin/Source/ChaosVehiclesEditor/Private/ChaosVehiclesEditorStyle.h
--- a/Engine/Plugins/Experimental/ChaosVehiclesPlugin/Source/ChaosVehiclesEditor/Private/ChaosVehiclesEditorStyle.h
+++ b/Engine/Plugins/Experimental/ChaosVehiclesPlugin/Source/ChaosVehiclesEditor/Private/ChaosVehiclesEditorStyle.h
@@ -16,8 +16,8 @@ public:
 		SetContentRoot(FPaths::EnginePluginsDir() / TEXT("Experimental/ChaosVehiclesPlugin/Resources"));
 #endif
 
-		Set("ClassIcon.ChaosVehicles", new FSlateImageBrush(RootToContentDir(TEXT("ChaosVehicles_16x.png")), Icon16x16));
-		Set("ClassThumbnail.ChaosVehicles", new FSlateImageBrush(RootToContentDir(TEXT("ChaosVehicles_64x.png")), Icon64x64));
+		Set("ClassIcon.ChaosVehicles", new FSlateVectorImageBrush(RootToContentDir(TEXT("ChaosVehicles_16.svg")), Icon16x16));
+		Set("ClassThumbnail.ChaosVehicles", new FSlateVectorImageBrush(RootToContentDir(TEXT("ChaosVehicles_64.svg")), Icon64x64));
 
 		FSlateStyleRegistry::RegisterSlateStyle(*this);
 	}
diff --git a/Engine/Plugins/VirtualProduction/Takes/Source/TakeRecorder/Private/TakeRecorderStyle.cpp b/Engine/Plugins/VirtualProduction/Takes/Source/TakeRecorder/Private/TakeRecorderStyle.cpp
--- a/Engine/Plugins/VirtualProduction/Takes/Source/TakeRecorder/Private/TakeRecorderStyle.cpp
+++ b/Engine/Plugins/VirtualProduction/Takes/Source/TakeRecorder/Private/TakeRecorderStyle.cpp
@@ -134,32 +134,36 @@ FTakeRecorderStyle::FTakeRecorderStyle()
 
 	Set("TakeRecorder.TabIcon", new FSlateImageBrush(RootToContentDir(TEXT("TabIcon_16x.png")), Icon16x16));
 
+
 	Set("ClassIcon.TakePreset",      new FSlateImageBrush(RootToContentDir(TEXT("TakePreset_16x.png")), Icon16x16));
-	Set("ClassThumbnail.TakePreset", new FSlateImageBrush(RootToContentDir(TEXT("TakePreset_64x.png")), Icon64x64));
+	Set("ClassThumbnail.TakePreset", new FSlateVectorImageBrush(RootToContentDir(TEXT("AssetIcons/TakePreset_64.svg")), Icon64x64));
 
 	Set("ClassIcon.TakeRecorderActorSource", new FSlateImageBrush(RootToContentDir(TEXT("TakeRecorderActorSource_16x.png")), Icon16x16));
-	Set("ClassThumbnail.TakeRecorderActorSource", new FSlateImageBrush(RootToContentDir(TEXT("TakeRecorderActorSource_64x.png")), Icon64x64));
+	Set("ClassThumbnail.TakeRecorderActorSource", new FSlateVectorImageBrush(RootToContentDir(TEXT("AssetIcons/TakeRecorderActorSource_64.svg")), Icon64x64));
 
 	Set("ClassIcon.TakeRecorderCameraCutSource", new FSlateImageBrush(RootToContentDir(TEXT("TakeRecorderCameraCutSource_16x.png")), Icon16x16));
 	Set("ClassThumbnail.TakeRecorderCameraCutSource", new FSlateImageBrush(RootToContentDir(TEXT("TakeRecorderCameraCutSource_64x.png")), Icon64x64));
 
 	Set("ClassIcon.TakeRecorderLevelSequenceSource", new FSlateImageBrush(RootToContentDir(TEXT("TakeRecorderLevelSequenceSource_16x.png")), Icon16x16));
-	Set("ClassThumbnail.TakeRecorderLevelSequenceSource", new FSlateImageBrush(RootToContentDir(TEXT("TakeRecorderLevelSequenceSource_64x.png")), Icon64x64));
+	Set("ClassThumbnail.TakeRecorderLevelSequenceSource", new FSlateVectorImageBrush(RootToContentDir(TEXT("AssetIcons/TakeRecorderLevelSequenceSource_64.svg")), Icon64x64));
+
+	Set("ClassIcon.TakeRecorderLiveLinkSource", new FSlateVectorImageBrush(RootToContentDir(TEXT("AssetIcons/TakeRecorderLiveLinkSource_64.svg")), Icon16x16));
+	Set("ClassThumbnail.TakeRecorderLiveLinkSource", new FSlateVectorImageBrush(RootToContentDir(TEXT("AssetIcons/TakeRecorderLiveLinkSource_64.svg")), Icon64x64));
 
 	Set("ClassIcon.TakeRecorderMicrophoneAudioSource", new FSlateImageBrush(RootToContentDir(TEXT("TakeRecorderMicrophoneAudioSource_16x.png")), Icon16x16));
-	Set("ClassThumbnail.TakeRecorderMicrophoneAudioSource", new FSlateImageBrush(RootToContentDir(TEXT("TakeRecorderMicrophoneAudioSource_64x.png")), Icon64x64));
+	Set("ClassThumbnail.TakeRecorderMicrophoneAudioSource", new FSlateVectorImageBrush(RootToContentDir(TEXT("AssetIcons/TakeRecorderMicrophoneAudioSource_64.svg")), Icon64x64));
 
 	Set("ClassIcon.TakeRecorderLevelVisibilitySource", new FSlateImageBrush(RootToContentDir(TEXT("TakeRecorderLevelVisibilitySource_16x.png")), Icon16x16));
 	Set("ClassThumbnail.TakeRecorderLevelVisibilitySource", new FSlateImageBrush(RootToContentDir(TEXT("TakeRecorderLevelVisibilitySource_64x.png")), Icon64x64));
 
 	Set("ClassIcon.TakeRecorderNearbySpawnedActorSource", new FSlateImageBrush(RootToContentDir(TEXT("TakeRecorderNearbySpawnedActorSource_16x.png")), Icon16x16));
-	Set("ClassThumbnail.TakeRecorderNearbySpawnedActorSource", new FSlateImageBrush(RootToContentDir(TEXT("TakeRecorderNearbySpawnedActorSource_64x.png")), Icon64x64));
+	Set("ClassThumbnail.TakeRecorderNearbySpawnedActorSource", new FSlateVectorImageBrush(RootToContentDir(TEXT("AssetIcons/TakeRecorderNearbySpawnedActorSource_64.svg")), Icon64x64));
 
 	Set("ClassIcon.TakeRecorderPlayerSource", new FSlateImageBrush(RootToContentDir(TEXT("TakeRecorderPlayerSource_16x.png")), Icon16x16));
 	Set("ClassThumbnail.TakeRecorderPlayerSource", new FSlateImageBrush(RootToContentDir(TEXT("TakeRecorderPlayerSource_64x.png")), Icon64x64));
 
 	Set("ClassIcon.TakeRecorderWorldSource", new FSlateImageBrush(RootToContentDir(TEXT("TakeRecorderWorldSource_16x.png")), Icon16x16));
-	Set("ClassThumbnail.TakeRecorderWorldSource", new FSlateImageBrush(RootToContentDir(TEXT("TakeRecorderWorldSource_64x.png")), Icon64x64));
+	Set("ClassThumbnail.TakeRecorderWorldSource", new FSlateVectorImageBrush(RootToContentDir(TEXT("AssetIcons/TakeRecorderWorldSource_64.svg")), Icon64x64));
 
 	Set("TakeRecorder.TabIcon", new FSlateImageBrush(RootToContentDir(TEXT("TabIcon_16x.png")), Icon16x16));
 
diff --git a/Engine/Source/Editor/EditorStyle/Private/StarshipStyle.cpp b/Engine/Source/Editor/EditorStyle/Private/StarshipStyle.cpp
--- a/Engine/Source/Editor/EditorStyle/Private/StarshipStyle.cpp
+++ b/Engine/Source/Editor/EditorStyle/Private/StarshipStyle.cpp
@@ -5746,7 +5746,6 @@ void FStarshipEditorStyle::FStyle::SetupClassIconsAndThumbnails()
 		Set("ClassIcon.MovableMobilityIcon", new IMAGE_BRUSH("Icons/ActorIcons/Light_Movable_16x", Icon16x16));
 		Set("ClassIcon.StationaryMobilityIcon", new IMAGE_BRUSH("Icons/ActorIcons/Light_Adjustable_16x", Icon16x16));
 		Set("ClassIcon.ComponentMobilityHeaderIcon", new IMAGE_BRUSH("Icons/ActorIcons/ComponentMobilityHeader_7x16", Icon7x16));
-		Set("ClassIcon.AnimInstance", new IMAGE_BRUSH("Icons/AssetIcons/AnimInstance_16x", Icon16x16));
 
 		// Factory classes
 		//Set( "ClassIcon.ActorFactoryBoxVolume", new IMAGE_BRUSH( "Icons/icon_volume_Box_16x", Icon16x16 ) );
@@ -5758,52 +5757,28 @@ void FStarshipEditorStyle::FStyle::SetupClassIconsAndThumbnails()
 
 			TEXT("AbilitySystemComponent"),
 			TEXT("AIPerceptionComponent"),
-			TEXT("AnimationModifier"),		
-			TEXT("AnimBlueprint"),
-			TEXT("AnimLayerInterface"),
-			TEXT("AnimComposite"),
-			TEXT("AnimInstance"),
-			TEXT("AnimMontage"),
-			TEXT("AnimSequence"),
-			TEXT("AnimationSharingSetup"),
-			TEXT("ApplicationLifecycleComponent"),
-			TEXT("BlackboardData"),
-			TEXT("BlendSpace"),
-			TEXT("BlendSpace1D"),
-			TEXT("BlueprintFunctionLibrary"),
 			TEXT("BlueprintInterface"),
 			TEXT("BlueprintMacroLibrary"),
 			TEXT("CableComponent"),
 			TEXT("CameraAnim"),
-			TEXT("Class"),
-			TEXT("DataAsset"),
 			TEXT("Default"),
 			TEXT("DirectionalLightMovable"),
 			TEXT("DirectionalLightStatic"),
 			TEXT("DirectionalLightStationary"),
-			TEXT("EnvQuery"),
 			TEXT("Font"),
 			TEXT("FontFace"),
 			TEXT("ForceFeedbackEffect"),
-			TEXT("Interface"),
 			TEXT("InterpData"),
 			TEXT("LevelSequence"),
 			TEXT("LightmassCharacterIndirectDetailVolume"),
 			TEXT("MassiveLODOverrideVolume"),
 			TEXT("Material"),
-			TEXT("MaterialFunction"),
-			TEXT("MaterialInstanceConstant"),
 			TEXT("MaterialParameterCollection"),
 			TEXT("MatineeActor"),
 			TEXT("MultiFont"),
-			TEXT("ObjectLibrary"),
 			TEXT("ParticleSystem"),
-			TEXT("PhysicalMaterial"),
-			TEXT("PhysicsAsset"),
 			TEXT("PhysicsConstraintComponent"),
 			TEXT("PhysicsThrusterComponent"),
-			TEXT("PlatformMediaSource"),
-			TEXT("PoseAsset"),
 			TEXT("SkyLightComponent"),
 			TEXT("SlateWidgetStyleAsset"),
 			TEXT("StringTable"),
@@ -5811,14 +5786,7 @@ void FStarshipEditorStyle::FStyle::SetupClassIconsAndThumbnails()
 			TEXT("SpotLightStatic"),
 			TEXT("SpotLightStationary"),
 			TEXT("StaticMesh"),
-			TEXT("SubsurfaceProfile"),
-			TEXT("Texture2D"),
-			TEXT("TextureRenderTarget2D"),
 			TEXT("UserDefinedEnum"),
-			TEXT("UserDefinedStruct"),
-			TEXT("UserWidget"),
-			TEXT("WidgetBlueprint"),
-			TEXT("World"),
 			TEXT("Cube"),
 			TEXT("Sphere"),
 			TEXT("Cylinder"),
@@ -5841,11 +5809,24 @@ void FStarshipEditorStyle::FStyle::SetupClassIconsAndThumbnails()
 			{TEXT("AimOffsetBlendSpace")},
 			{TEXT("AimOffsetBlendSpace1D")},
 			{TEXT("AmbientSound")},
+			{TEXT("AnimationModifier")},
+			{TEXT("AnimationSharingSetup")},
+			{TEXT("AnimBlueprint")},
+			{TEXT("AnimComposite")},
+			{TEXT("AnimInstance")},
+			{TEXT("AnimLayerInterface")},
+			{TEXT("AnimMontage")},
+			{TEXT("AnimSequence")},
+			{TEXT("ApplicationLifecycleComponent")},
 			{TEXT("AtmosphericFog")},
 			{TEXT("AudioVolume")},
 			{TEXT("BehaviorTree")},
+			{TEXT("BlackboardData")},
+			{TEXT("BlendSpace")},
+			{TEXT("BlendSpace1D")},
 			{TEXT("BlockingVolume")},
 			{TEXT("Blueprint")},
+			{TEXT("BlueprintFunctionLibrary")},
 			{TEXT("BoxReflectionCapture")},
 			{TEXT("Brush")},
 			{TEXT("ButtonStyleAsset")},
@@ -5857,9 +5838,11 @@ void FStarshipEditorStyle::FStyle::SetupClassIconsAndThumbnails()
 			{TEXT("Character")},
 			{TEXT("CharacterMovementComponent")},
 			{TEXT("CineCameraActor")},
+			{TEXT("Class")},
 			{TEXT("CompositingElement")},
 			{TEXT("CullDistanceVolume")},
 			{TEXT("CurveBase")},
+			{TEXT("DataAsset")},
 			{TEXT("DataTable")},
 			{TEXT("DecalActor")},
 			{TEXT("DefaultPawn")},
@@ -5868,6 +5851,7 @@ void FStarshipEditorStyle::FStyle::SetupClassIconsAndThumbnails()
 			{TEXT("DirectionalLight")},
 			{TEXT("DocumentationActor")},
 			{TEXT("EditorTutorial")},
+			{TEXT("EnvQuery")},
 			{TEXT("Emitter")},
 			{TEXT("EmptyActor")},
 			{TEXT("ExponentialHeightFog")},
@@ -5880,6 +5864,7 @@ void FStarshipEditorStyle::FStyle::SetupClassIconsAndThumbnails()
 			{TEXT("GroupActor")},
 			{TEXT("HUD")},
 			{TEXT("ImagePlate")},
+			{TEXT("Interface")},
 			{TEXT("KillZVolume")},
 			{TEXT("Landscape")},
 			{TEXT("LevelBounds")},
@@ -5891,7 +5876,9 @@ void FStarshipEditorStyle::FStyle::SetupClassIconsAndThumbnails()
 			{TEXT("LightmassImportanceVolume")},
 			{TEXT("LightmassVolume")},
 			{TEXT("LiveLinkPreset")},
+			{TEXT("MaterialFunction")},
 			{TEXT("MaterialInstanceActor")},
+			{TEXT("MaterialInstanceConstant")},
 			{TEXT("MediaPlayer")},
 			{TEXT("MediaTexture")},
 			{TEXT("ModularSynthPresetBank")},
@@ -5901,18 +5888,24 @@ void FStarshipEditorStyle::FStyle::SetupClassIconsAndThumbnails()
 			{TEXT("NavModifierComponent")},
 			{TEXT("NavModifierVolume")},
 			{TEXT("Note")},
+			{TEXT("Object")},
+			{TEXT("ObjectLibrary")},
 			{TEXT("PainCausingVolume")},
 			{TEXT("Pawn")},
 			{TEXT("PawnNoiseEmitterComponent")},
 			{TEXT("PawnSensingComponent")},
+			{TEXT("PhysicalMaterial")},
+			{TEXT("PhysicsAsset")},
 			{TEXT("PhysicsConstraintActor")},
 			{TEXT("PhysicsHandleComponent")},
 			{TEXT("PhysicsThruster")},
 			{TEXT("PhysicsVolume")},
 			{TEXT("PlaneReflectionCapture")},
+			{TEXT("PlatformMediaSource")},
 			{TEXT("PlayerController")},
 			{TEXT("PlayerStart")},
 			{TEXT("PointLight")},
+			{TEXT("PoseAsset")},
 			{TEXT("PostProcessVolume")},
 			{TEXT("PrecomputedVisibilityOverrideVolume")},
 			{TEXT("PrecomputedVisibilityVolume")},
@@ -5939,20 +5932,18 @@ void FStarshipEditorStyle::FStyle::SetupClassIconsAndThumbnails()
 			{TEXT("SoundCue")},
 			{TEXT("SoundEffectSourcePreset")},
 			{TEXT("SoundMix")},
+			{TEXT("SoundSubmix") },
 			{TEXT("SphereReflectionCapture")},
 			{TEXT("SpotLight")},
 			{TEXT("SpringArmComponent")},
 			{TEXT("StaticMeshActor")},
 			{TEXT("StreamMediaSource")},
-			{TEXT("TakePreset")},
-			{TEXT("TakeRecorderLevelSequenceSource")},
-			{TEXT("TakeRecorderLiveLinkSource")},
-			{TEXT("TakeRecorderMicrophoneAudioSource")},
-			{TEXT("TakeRecorderNearbySpawnedActorSource")},
-			{TEXT("TakeRecorderWorldSource")},
+			{TEXT("SubsurfaceProfile")},
 			{TEXT("TargetPoint")},
 			{TEXT("TemplateSequence")},
 			{TEXT("TextRenderActor")},
+			{TEXT("Texture2D")},
+			{TEXT("TextureRenderTarget2D")},
 			{TEXT("TextureRenderTargetCube")},
 			{TEXT("TimeCodeSynchronizer")},
 			{TEXT("TouchInterface")},
@@ -5962,12 +5953,16 @@ void FStarshipEditorStyle::FStyle::SetupClassIconsAndThumbnails()
 			{TEXT("TriggerSphere")},
 			{TEXT("TriggerVolume")},
 			{TEXT("UserDefinedCaptureProtocol")},
+			{TEXT("UserDefinedStruct") },
+			{TEXT("UserWidget")},
 			{TEXT("VectorField")},
 			{TEXT("VectorFieldVolume")},
 			{TEXT("Volume")},
 			{TEXT("VolumetricCloud"), false},
 			{TEXT("VolumetricLightmapDensityVolume")},
-			{TEXT("WindDirectionalSource")}
+			{TEXT("WidgetBlueprint")},
+			{TEXT("WindDirectionalSource")},
+			{TEXT("World")},
 		};
 	
 		// SVG Asset icons
@@ -5976,6 +5971,7 @@ void FStarshipEditorStyle::FStyle::SetupClassIconsAndThumbnails()
 			{
 				const FClassIconInfo& Info = AssetTypesSVG[TypeIndex];
 
+				// Look up if the brush already exists to audit old vs new icons during starship development.
 				FString ClassIconName = FString::Printf(TEXT("ClassIcon.%s"), Info.Type);
 				if (GetOptionalBrush(*ClassIconName, nullptr, nullptr))
 				{
