deprecate Slates EnableMenuAnimations api. This feature has been disabled and has not worked since 4.0. Turning it on by accient causes massive issues with menus.

diff --git a/Engine/Config/BaseEditorPerProjectUserSettings.ini b/Engine/Config/BaseEditorPerProjectUserSettings.ini
--- a/Engine/Config/BaseEditorPerProjectUserSettings.ini
+++ b/Engine/Config/BaseEditorPerProjectUserSettings.ini
@@ -93,8 +93,6 @@ LogFontSize=9
 
 ; If true, all toolbars will use small icons without labels
 bUseSmallToolBarIcons=False
-; Enables animated transitions for certain menus and pop-up windows.  Note that animations may be automatically disabled at low frame rates in order to improve responsiveness.
-bEnableWindowAnimations=False
 ;If true, Variables names in the script editor are displayed in a sanitized format
 bShowFriendlyNames=True
 ;If true, the Editor Preferences and Project Settings menu entries in the main menu will be expanded with sub-menus for each settings section
diff --git a/Engine/Source/Editor/EditorStyle/Private/EditorStyleClasses.cpp b/Engine/Source/Editor/EditorStyle/Private/EditorStyleClasses.cpp
--- a/Engine/Source/Editor/EditorStyle/Private/EditorStyleClasses.cpp
+++ b/Engine/Source/Editor/EditorStyle/Private/EditorStyleClasses.cpp
@@ -63,10 +63,6 @@ void UEditorStyleSettings::PostEditChangeProperty(struct FPropertyChangedEvent&
 
 	const FName PropertyName = (PropertyChangedEvent.Property != nullptr) ? PropertyChangedEvent.Property->GetFName() : NAME_None;
 
-	if (PropertyName == GET_MEMBER_NAME_CHECKED(UEditorStyleSettings, bEnableWindowAnimations))
-	{
-		FSlateApplication::Get().EnableMenuAnimations(bEnableWindowAnimations);
-	}
 
 	// This property is intentionally not per project so it must be manually written to the correct config file
 	if (PropertyName == GET_MEMBER_NAME_CHECKED(UEditorStyleSettings, bEnableHighDPIAwareness))
diff --git a/Engine/Source/Editor/EditorStyle/Public/Classes/EditorStyleSettings.h b/Engine/Source/Editor/EditorStyle/Public/Classes/EditorStyleSettings.h
--- a/Engine/Source/Editor/EditorStyle/Public/Classes/EditorStyleSettings.h
+++ b/Engine/Source/Editor/EditorStyle/Public/Classes/EditorStyleSettings.h
@@ -109,10 +109,6 @@ public:
 	UPROPERTY(EditAnywhere, config, Category = Graphs, meta = (DisplayName = "Background Brush"))
 	FSlateBrush GraphBackgroundBrush;
 
-	/** Enables animated transitions for certain menus and pop-up windows.  Note that animations may be automatically disabled at low frame rates in order to improve responsiveness. */
-	UPROPERTY(EditAnywhere, config, Category=UserInterface)
-	uint32 bEnableWindowAnimations:1;
-
 	/** When enabled, the C++ names for properties and functions will be displayed in a format that is easier to read */
 	UPROPERTY(EditAnywhere, config, Category=UserInterface, meta=(DisplayName="Show Friendly Variable Names"))
 	uint32 bShowFriendlyNames:1;
diff --git a/Engine/Source/Programs/ChaosVisualDebugger/Private/ChaosVisualDebuggerMain.cpp b/Engine/Source/Programs/ChaosVisualDebugger/Private/ChaosVisualDebuggerMain.cpp
--- a/Engine/Source/Programs/ChaosVisualDebugger/Private/ChaosVisualDebuggerMain.cpp
+++ b/Engine/Source/Programs/ChaosVisualDebugger/Private/ChaosVisualDebuggerMain.cpp
@@ -19,8 +19,6 @@ void InitializedSlateApplication()
 	// crank up a normal Slate application using the platform's standalone renderer
 	FSlateApplication::InitializeAsStandaloneApplication(GetStandardStandaloneRenderer());
 
-	FSlateApplication::Get().EnableMenuAnimations(false);
-
 	// Set the application name.
 	const FText ApplicationTitle = NSLOCTEXT("ChaosVisualDebugger", "AppTitle", "ChaosVisualDebugger");
 	FGlobalTabmanager::Get()->SetApplicationTitle(ApplicationTitle);
diff --git a/Engine/Source/Programs/SlateViewer/Private/SlateViewerApp.cpp b/Engine/Source/Programs/SlateViewer/Private/SlateViewerApp.cpp
--- a/Engine/Source/Programs/SlateViewer/Private/SlateViewerApp.cpp
+++ b/Engine/Source/Programs/SlateViewer/Private/SlateViewerApp.cpp
@@ -42,7 +42,6 @@ int RunSlateViewer( const TCHAR* CommandLine )
 
 	// crank up a normal Slate application using the platform's standalone renderer
 	FSlateApplication::InitializeAsStandaloneApplication(GetStandardStandaloneRenderer());
-	FSlateApplication::Get().EnableMenuAnimations(false);
 
 	FSlateApplication::InitHighDPI(true);
 
diff --git a/Engine/Source/Programs/UnrealInsights/Private/UserInterfaceCommand.cpp b/Engine/Source/Programs/UnrealInsights/Private/UserInterfaceCommand.cpp
--- a/Engine/Source/Programs/UnrealInsights/Private/UserInterfaceCommand.cpp
+++ b/Engine/Source/Programs/UnrealInsights/Private/UserInterfaceCommand.cpp
@@ -220,8 +220,6 @@ void FUserInterfaceCommand::InitializeSlateApplication(bool bOpenTraceFile, cons
 	//const FSlateBrush* AppIcon = new FSlateImageBrush(FPaths::EngineContentDir() / "Editor/Slate/Icons/Insights/AppIcon_24x.png", FVector2D(24.0f, 24.0f));
 	//FSlateApplication::Get().SetAppIcon(AppIcon);
 
-	// Menu anims aren't supported. See Runtime\Slate\Private\Framework\Application\MenuStack.cpp.
-	FSlateApplication::Get().EnableMenuAnimations(false);
 
 	// Set the application name.
 	const FText ApplicationTitle = FText::Format(NSLOCTEXT("UnrealInsights", "AppTitle", "Unreal Insights {0}"), FText::FromString(TEXT(UNREAL_INSIGHTS_VERSION_STRING_EX)));
diff --git a/Engine/Source/Runtime/Slate/Public/Framework/Application/SlateApplication.h b/Engine/Source/Runtime/Slate/Public/Framework/Application/SlateApplication.h
--- a/Engine/Source/Runtime/Slate/Public/Framework/Application/SlateApplication.h
+++ b/Engine/Source/Runtime/Slate/Public/Framework/Application/SlateApplication.h
@@ -830,6 +830,7 @@ public:
 	 *
 	 * @param	bEnableAnimations	True if animations should be used, otherwise false.
 	 */
+	UE_DEPRECATED(5.0, "Enable Window Animations is no longer used and is a no-op so calling this function is no longer necessary.")
 	void EnableMenuAnimations( const bool bEnableAnimations );
 
 	void SetPlatformApplication(const TSharedRef<class GenericApplication>& InPlatformApplication);
