Blueprint icon and style update.

This change includes one core toolbar change. FSlateIcon now supports a status  overlay. This is used to show compile status in the blueprint editor

diff --git a/Engine/Source/Editor/EditorStyle/Private/StarshipStyle.cpp b/Engine/Source/Editor/EditorStyle/Private/StarshipStyle.cpp
--- a/Engine/Source/Editor/EditorStyle/Private/StarshipStyle.cpp
+++ b/Engine/Source/Editor/EditorStyle/Private/StarshipStyle.cpp
@@ -1013,7 +1013,7 @@ void FStarshipEditorStyle::FStyle::SetupGeneralStyles()
 
 	// Blueprint Debugger Window
 	{
-		Set("BlueprintDebugger.TabIcon", new IMAGE_BRUSH( "Icons/icon_tab_BlueprintDebugger_16x", Icon16x16 ) );
+		Set("BlueprintDebugger.TabIcon", new IMAGE_BRUSH_SVG( "Starship/Common/Debug", Icon16x16 ) );
 	}
 
 	// Collision Analyzer Window
@@ -1212,8 +1212,8 @@ void FStarshipEditorStyle::FStyle::SetupGeneralStyles()
 			.SetShadowOffset( FVector2D::ZeroVector )
 		);
 
-		Set( "GraphBreadcrumb.BrowseBack", new IMAGE_BRUSH( "Icons/icon_BlueprintBrowserL_24x", Icon24x24) );
-		Set( "GraphBreadcrumb.BrowseForward", new IMAGE_BRUSH( "Icons/icon_BlueprintBrowserR_24x", Icon24x24) );
+		Set("GraphBreadcrumb.BrowseBack", new IMAGE_BRUSH_SVG( "Starship/GraphEditors/PreviousArrow", Icon20x20));
+		Set("GraphBreadcrumb.BrowseForward", new IMAGE_BRUSH_SVG( "Starship/GraphEditors/NextArrow", Icon20x20));
 
 		const FComboButtonStyle FastJumpComboBoxComboButton = FComboButtonStyle()
 			.SetButtonStyle(GetWidgetStyle<FButtonStyle>("GraphBreadcrumbButton"));
@@ -1223,17 +1223,6 @@ void FStarshipEditorStyle::FStyle::SetupGeneralStyles()
 
 	// Graph bookmark button
 	{
-		Set("GraphBookmarkButton", FButtonStyle()
-			.SetNormal(FSlateNoResource())
-			.SetPressed(BOX_BRUSH("Common/Button_Pressed", 8.0f / 32.0f, SelectionColor_Pressed))
-			.SetHovered(BOX_BRUSH("Common/Button_Hovered", 8.0f / 32.0f, SelectionColor))
-		);
-
-		Set("GraphBookmarkButtonText", FTextBlockStyle(NormalText)
-			.SetFont(ICON_FONT(16))
-			.SetColorAndOpacity(FLinearColor(0.5f, 0.5f, 0.5f))
-			.SetShadowOffset(FVector2D::ZeroVector)
-		);
 
 		Set("GraphBookmarkMenuImage.Button_Add", new IMAGE_BRUSH("Icons/PlusSymbol_12x", Icon12x12));
 		Set("GraphBookmarkMenuText.EmptyListItem", FTextBlockStyle(NormalText)
@@ -3441,6 +3430,7 @@ void FStarshipEditorStyle::FStyle::SetupGraphEditorStyles()
 		Set( "Graph.Node.CommentFont", DEFAULT_FONT( "Regular", 10 ) );
 		Set( "Graph.Node.Comment.BubbleOffset", FMargin(8,0,0,0) );
 		Set( "Graph.Node.Comment.PinIconPadding", FMargin(0,2,0,0) );
+		Set("Graph.Node.Comment.Handle", new IMAGE_BRUSH_SVG("Starship/GraphEditors/Comment_Handle", FVector2D(16.0f, 16.0f)));
 		Set("Graph.Node.Comment.BubblePadding", FVector2D(3, 3));
 		Set("Graph.Node.Comment.BubbleWidgetMargin", FMargin(4, 4));
 
@@ -3632,8 +3622,8 @@ void FStarshipEditorStyle::FStyle::SetupGraphEditorStyles()
 			Set( "BTEditor.Debugger.SavedValues", new IMAGE_BRUSH( "BehaviorTree/Debugger_Saved_40x", Icon40x40 ) );
 			Set( "BTEditor.Debugger.SavedValues.Small", new IMAGE_BRUSH( "BehaviorTree/Debugger_Saved_40x", Icon20x20 ) );
 
-			Set( "BTEditor.DebuggerOverlay.Breakpoint.Disabled", new IMAGE_BRUSH( "Old/Kismet2/Breakpoint_Disabled", Icon32x32 ) );
-			Set( "BTEditor.DebuggerOverlay.Breakpoint.Enabled", new IMAGE_BRUSH( "Old/Kismet2/Breakpoint_Valid", Icon32x32 ) );
+			Set( "BTEditor.DebuggerOverlay.Breakpoint.Disabled", new IMAGE_BRUSH_SVG( "Starship/Blueprints/Breakpoint_Disabled", Icon32x32 ) );
+			Set( "BTEditor.DebuggerOverlay.Breakpoint.Enabled", new IMAGE_BRUSH_SVG( "Starship/Blueprints/Breakpoint_Valid", Icon32x32 ) );
 			Set( "BTEditor.DebuggerOverlay.ActiveNodePointer", new IMAGE_BRUSH( "Old/Kismet2/IP_Normal", FVector2D(128,96)) );
 			Set( "BTEditor.DebuggerOverlay.SearchTriggerPointer", new IMAGE_BRUSH( "/BehaviorTree/SearchTriggerPointer", FVector2D(48,64)) );
 			Set( "BTEditor.DebuggerOverlay.FailedTriggerPointer", new IMAGE_BRUSH( "/BehaviorTree/FailedTriggerPointer", FVector2D(48,64)) );
@@ -3928,8 +3918,8 @@ void FStarshipEditorStyle::FStyle::SetupGraphEditorStyles()
 			);
 		}
 
-		Set("ClassIcon.K2Node_CallFunction",	new IMAGE_BRUSH("Icons/icon_Blueprint_NewFunction_16x", Icon16x16));
-		Set("ClassIcon.K2Node_FunctionEntry",	new IMAGE_BRUSH("Icons/icon_Blueprint_NewFunction_16x", Icon16x16));
+		Set("ClassIcon.K2Node_CallFunction",	new IMAGE_BRUSH_SVG("Starship/Blueprints/icon_Blueprint_Function", Icon16x16));
+		Set("ClassIcon.K2Node_FunctionEntry",	new IMAGE_BRUSH_SVG("Starship/Blueprints/icon_Blueprint_Function", Icon16x16));
 		Set("ClassIcon.K2Node_CustomEvent",		new IMAGE_BRUSH("Icons/icon_Blueprint_Event_16x", Icon16x16));
 		Set("ClassIcon.K2Node_Event",			new IMAGE_BRUSH("Icons/icon_Blueprint_Event_16x", Icon16x16));
 		Set("ClassIcon.K2Node_Variable",		new IMAGE_BRUSH("Graph/Icons/Node", Icon16x16, FLinearColor::White));
@@ -3940,10 +3930,10 @@ void FStarshipEditorStyle::FStyle::SetupGraphEditorStyles()
 
 
 		Set( "GraphEditor.Default_16x", new IMAGE_BRUSH("Icons/icon_Blueprint_Node_16x", Icon16x16));
-		Set( "GraphEditor.EventGraph_16x", new IMAGE_BRUSH( "Icons/icon_Blueprint_EventGraph_16x", Icon16x16 ) );
-		Set( "GraphEditor.InterfaceFunction_16x", new IMAGE_BRUSH( "Icons/icon_Blueprint_Interfacefunction_16x", Icon16x16 ) );
-		Set( "GraphEditor.Macro_16x", new IMAGE_BRUSH( "Icons/icon_Blueprint_Macro_16x", Icon16x16 ) );
-		Set( "GraphEditor.Function_16x", new IMAGE_BRUSH( "Icons/icon_Blueprint_NewFunction_16x", Icon16x16 ) );
+		Set( "GraphEditor.EventGraph_16x", new IMAGE_BRUSH_SVG( "Starship/Blueprints/icon_BlueprintEditor_EventGraph", Icon16x16 ) );
+		Set( "GraphEditor.InterfaceFunction_16x", new IMAGE_BRUSH( "Icons/icon_Blueprint_Interfacefunction_16x", Icon16x16));
+		Set( "GraphEditor.Macro_16x", new IMAGE_BRUSH_SVG( "Starship/Blueprints/icon_Blueprint_Macro", Icon16x16));
+		Set( "GraphEditor.Function_16x", new IMAGE_BRUSH_SVG( "Starship/Blueprints/icon_Blueprint_Function", Icon16x16));
 		Set( "GraphEditor.PureFunction_16x", new IMAGE_BRUSH( "Icons/icon_Blueprint_NewPureFunction_16x", Icon16x16 ) );
 		Set( "GraphEditor.PotentialOverrideFunction_16x", new IMAGE_BRUSH( "Icons/icon_Blueprint_OverrideableFunction_16x", Icon16x16 ) );
 		Set( "GraphEditor.OverrideFunction_16x", new IMAGE_BRUSH( "Icons/icon_Blueprint_OverrideFunction_16x", Icon16x16 ) );
@@ -3988,10 +3978,10 @@ void FStarshipEditorStyle::FStyle::SetupGraphEditorStyles()
 		Set( "GraphEditor.Delegate_16x", new IMAGE_BRUSH( "Icons/icon_Blueprint_Delegate_16x", Icon16x16 ) );
 		Set( "GraphEditor.Delegate_24x", new IMAGE_BRUSH( "Icons/icon_Blueprint_Delegate_24x", Icon24x24 ) );
 
-		Set( "GraphEditor.EventGraph_24x", new IMAGE_BRUSH( "Icons/icon_Blueprint_EventGraph_24x", Icon24x24 ) );
+		Set( "GraphEditor.EventGraph_24x", new IMAGE_BRUSH_SVG( "Starship/Blueprints/icon_BlueprintEditor_EventGraph", Icon24x24 ) );
 		Set( "GraphEditor.InterfaceFunction_24x", new IMAGE_BRUSH( "Icons/icon_Blueprint_InterfaceFunction_24x", Icon24x24 ) );
-		Set( "GraphEditor.Macro_24x", new IMAGE_BRUSH( "Icons/icon_Blueprint_Macro_24x", Icon24x24 ) );
-		Set( "GraphEditor.Function_24x", new IMAGE_BRUSH( "Icons/icon_Blueprint_NewFunction_24x", Icon24x24 ) );
+		Set( "GraphEditor.Macro_24x", new IMAGE_BRUSH_SVG( "Starship/Blueprints/icon_Blueprint_Macro", Icon24x24 ) );
+		Set( "GraphEditor.Function_24x", new IMAGE_BRUSH_SVG( "Starship/Blueprints/icon_Blueprint_Function", Icon24x24 ) );
 		Set( "GraphEditor.PotentialOverrideFunction_24x", new IMAGE_BRUSH( "Icons/icon_Blueprint_OverrideableFunction_24x", Icon24x24 ) );
 		Set( "GraphEditor.OverrideFunction_24x", new IMAGE_BRUSH( "Icons/icon_Blueprint_OverrideFunction_24x", Icon24x24 ) );
 		Set( "GraphEditor.SubGraph_24x", new IMAGE_BRUSH( "Icons/icon_Blueprint_SubgraphComposite_24x", Icon24x24 ) );
@@ -4047,8 +4037,8 @@ void FStarshipEditorStyle::FStyle::SetupGraphEditorStyles()
 		Set( "GraphEditor.DistributeNodesHorizontally", new IMAGE_BRUSH( "Icons/GraphEditor/icon_DistributeNodesHorizontally_20px", Icon20x20 ) );
 		Set( "GraphEditor.DistributeNodesVertically", new IMAGE_BRUSH( "Icons/GraphEditor/icon_DistributeNodesVertically_20px", Icon20x20 ) );
 		
-		Set( "GraphEditor.ToggleHideUnrelatedNodes", new IMAGE_BRUSH( "Icons/icon_HideUnrelatedNodes_40x", Icon40x40 ) );
-		Set( "GraphEditor.ToggleHideUnrelatedNodes.Small", new IMAGE_BRUSH( "Icons/icon_HideUnrelatedNodes_40x", Icon20x20 ) );
+		Set( "GraphEditor.ToggleHideUnrelatedNodes", new IMAGE_BRUSH_SVG( "Starship/GraphEditors/HideUnrelated", Icon20x20 ) );
+		Set("GraphEditor.Bookmark", new IMAGE_BRUSH_SVG("Starship/Common/Bookmarks", Icon16x16));
 
 		// Graph editor widgets
 		{
@@ -4108,11 +4098,9 @@ void FStarshipEditorStyle::FStyle::SetupGraphEditorStyles()
 
 		// SCSEditor
 	{
-		Set("SCSEditor.ToggleComponentEditing", new IMAGE_BRUSH("Icons/icon_translate_40x", Icon40x40));
-		Set("SCSEditor.ToggleComponentEditing.Small", new IMAGE_BRUSH("Icons/icon_translate_40x", Icon20x20));
 		Set("SCSEditor.TileViewTooltip.NonContentBorder", new BOX_BRUSH("/Docking/TabContentArea", FMargin(4 / 16.0f)));
 
-		Set("SCSEditor.PromoteToBlueprintIcon", new IMAGE_BRUSH("Icons/AssetIcons/Blueprint_16x", Icon16x16));
+		Set("SCSEditor.PromoteToBlueprintIcon", new IMAGE_BRUSH_SVG("Starship/Common/blueprint", Icon16x16));
 
 		Set("SCSEditor.TopBar.Font", FTextBlockStyle(NormalText)
 			.SetFont(DEFAULT_FONT("Bold", 10))
@@ -4236,7 +4224,7 @@ void FStarshipEditorStyle::FStyle::SetupLevelEditorStyle()
 		Set( "EditorViewport.TranslateRotateMode", new IMAGE_BRUSH( "Icons/icon_translate_rotate_40x", Icon20x20 ) );
 		Set( "EditorViewport.TranslateRotate2DMode", new IMAGE_BRUSH("Icons/icon_translate_rotate_2d_40x", Icon20x20));
 
-		Set( "EditorViewport.ToggleRealTime", new IMAGE_BRUSH("Icons/icon_MatEd_Realtime_40x", Icon40x40));
+		Set( "EditorViewport.ToggleRealTime", new IMAGE_BRUSH_SVG("Starship/Common/Realtime", Icon16x16));
 
 		Set( "EditorViewport.LocationGridSnap", new IMAGE_BRUSH_SVG("Starship/EditorViewport/grid", Icon16x16));
 		Set( "EditorViewport.RotationGridSnap", new IMAGE_BRUSH_SVG("Starship/EditorViewport/angle", Icon16x16));
@@ -5100,54 +5088,36 @@ void FStarshipEditorStyle::FStyle::SetupPersonaStyle()
 
 	// Kismet 2
 	{
-		Set( "FullBlueprintEditor.SwitchToScriptingMode", new IMAGE_BRUSH( "Icons/icon_BlueprintEditor_EventGraph_40x", Icon40x40 ) );
-		Set( "FullBlueprintEditor.SwitchToScriptingMode.Small", new IMAGE_BRUSH( "Icons/icon_BlueprintEditor_EventGraph_40x", Icon20x20 ) );
-		Set( "FullBlueprintEditor.SwitchToBlueprintDefaultsMode", new IMAGE_BRUSH( "Icons/icon_BlueprintEditor_Defaults_40x", Icon40x40 ) );
-		Set( "FullBlueprintEditor.SwitchToBlueprintDefaultsMode.Small", new IMAGE_BRUSH( "Icons/icon_BlueprintEditor_Defaults_40x", Icon20x20 ) );
-		Set( "FullBlueprintEditor.SwitchToComponentsMode", new IMAGE_BRUSH( "Icons/icon_BlueprintEditor_Components_40x", Icon40x40 ) );
-		Set( "FullBlueprintEditor.SwitchToComponentsMode.Small", new IMAGE_BRUSH( "Icons/icon_BlueprintEditor_Components_40x", Icon20x20 ) );
+		Set( "FullBlueprintEditor.SwitchToScriptingMode", new IMAGE_BRUSH_SVG( "Starship/Blueprints/icon_BlueprintEditor_EventGraph", Icon20x20));
 
-		Set( "FullBlueprintEditor.EditGlobalOptions", new IMAGE_BRUSH( "Icons/icon_Blueprint_Options_40px", Icon40x40));
-		Set( "FullBlueprintEditor.EditGlobalOptions.Small", new IMAGE_BRUSH( "Icons/icon_Blueprint_Options_40px", Icon20x20));
+		// @todo - Icon Replacement - The UI commands using these icons are never visible in the editor
+		Set( "FullBlueprintEditor.SwitchToBlueprintDefaultsMode", new IMAGE_BRUSH_SVG("Starship/Common/Defaults", Icon20x20));
+		Set( "FullBlueprintEditor.SwitchToComponentsMode", new IMAGE_BRUSH_SVG("Starship/Common/Component", Icon20x20));
 
-		Set("FullBlueprintEditor.EditClassDefaults", new IMAGE_BRUSH("Icons/icon_BlueprintEditor_Defaults_40x", Icon40x40));
-		Set("FullBlueprintEditor.EditClassDefaults.Small", new IMAGE_BRUSH("Icons/icon_BlueprintEditor_Defaults_40x", Icon20x20));
-
-		Set( "BlueprintEditor.Details.DeleteButton", new IMAGE_BRUSH( "/Icons/GenericDelete_Black", Icon16x16 ) );
-
-		Set( "BlueprintEditor.Details.ArgUpButton", new IMAGE_BRUSH( "/Icons/icon_FunctionArgUp", Icon16x16, FLinearColor(0.0f, 0.0f, 0.0f)) );
-		Set( "BlueprintEditor.Details.ArgDownButton", new IMAGE_BRUSH( "/Icons/icon_FunctionArgDown", Icon16x16, FLinearColor(0.0f, 0.0f, 0.0f)) );
+		Set( "FullBlueprintEditor.EditGlobalOptions", new CORE_IMAGE_BRUSH_SVG( "Starship/Common/Settings", Icon20x20));
+		Set("FullBlueprintEditor.EditClassDefaults", new IMAGE_BRUSH_SVG("Starship/Common/Details", Icon20x20));
 
 		Set( "FullBlueprintEditor.Diff", new IMAGE_BRUSH( "Icons/BlueprintEditorDiff", Icon40x40 ) );
 		Set( "FullBlueprintEditor.Diff.Small", new IMAGE_BRUSH( "Icons/BlueprintEditorDiff", Icon20x20 ) );
 
+		Set( "BlueprintEditor.ActionMenu.ContextDescriptionFont",  DEFAULT_FONT("Regular", 12) );
 		Set( "BlueprintEditor.ActionMenu.ContextDescriptionFont",  DEFAULT_FONT("Regular", 12) );
 
-		Set( "BlueprintEditor.FindInBlueprints.MenuIcon", new IMAGE_BRUSH_SVG("Starship/Common/FindInBlueprints", Icon16x16) );
-
-		Set( "BlueprintEditor.FindInBlueprint", new IMAGE_BRUSH( "Icons/icon_Blueprint_Find_40px", Icon40x40 ) );
-		Set( "BlueprintEditor.FindInBlueprint.Small", new IMAGE_BRUSH( "Icons/icon_Blueprint_Find_40px", Icon20x20 ) );
-
-		Set( "BlueprintEditor.FindInBlueprints", new IMAGE_BRUSH( "Icons/icon_FindInAnyBlueprint_40px", Icon40x40 ) );
-		Set( "BlueprintEditor.FindInBlueprints.Small", new IMAGE_BRUSH( "Icons/icon_FindInAnyBlueprint_40px", Icon20x20 ) );
-
-		Set( "Kismet.CompileBlueprint", new IMAGE_BRUSH("/Icons/icon_kismet_compile_16px", Icon16x16) );
-		Set( "Kismet.DeleteUnusedVariables", new IMAGE_BRUSH("/Icons/icon_kismet_findunused_16px", Icon16x16) );
-
-		Set( "Kismet.Toolbar.SelectedDebugObject.Background", new IMAGE_BRUSH( "Old/Kismet2/DebugObject_Background", Icon40x40) );
+		Set("BlueprintEditor.FindInBlueprints.MenuIcon", new IMAGE_BRUSH_SVG("Starship/Common/FindInBlueprints", Icon16x16));
+		Set("BlueprintEditor.FindInBlueprint", new IMAGE_BRUSH_SVG("Starship/Common/FindInBlueprints", Icon20x20));
 
+		Set( "Kismet.DeleteUnusedVariables", new IMAGE_BRUSH_SVG("/Starship/Blueprints/icon_kismet_findunused", Icon16x16) );
 		{
-			Set( "Kismet.Tabs.Variables", new IMAGE_BRUSH( "/Icons/pill_16x", Icon16x16 ) );
-			Set( "Kismet.Tabs.Palette", new IMAGE_BRUSH( "/Icons/levels_16x", Icon16x16 ) );
-			Set( "Kismet.Tabs.CompilerResults", new IMAGE_BRUSH( "Icons/icon_tab_OutputLog_16x", Icon16x16 ) );
-			Set( "Kismet.Tabs.FindResults", new IMAGE_BRUSH( "/Icons/icon_Genericfinder_16x", Icon16x16 ) );
-			Set( "Kismet.Tabs.Bookmarks", new IMAGE_BRUSH( "/Icons/icon_BlueprintEditor_Bookmarks_16x", Icon16x16 ) );
-			Set( "Kismet.Tabs.Components", new IMAGE_BRUSH( "/Icons/icon_BlueprintEditor_Components_16x", Icon16x16 ) );
+			Set( "Kismet.Tabs.Variables", new IMAGE_BRUSH_SVG( "Starship/Blueprints/pill", Icon16x16 ) );
+			Set( "Kismet.Tabs.Palette", new IMAGE_BRUSH_SVG( "Starship/Blueprints/Palette", Icon16x16 ) );
+			Set( "Kismet.Tabs.CompilerResults", new IMAGE_BRUSH_SVG("Starship/Common/OutputLog", Icon16x16 ) );
+			Set( "Kismet.Tabs.FindResults", new CORE_IMAGE_BRUSH_SVG("Starship/Common/search", Icon16x16 ) );
+			Set( "Kismet.Tabs.Bookmarks", new IMAGE_BRUSH_SVG( "Starship/Common/Bookmarks", Icon16x16 ) );
+			Set( "Kismet.Tabs.Components", new IMAGE_BRUSH_SVG( "Starship/Common/Component", Icon16x16 ) );
 			Set( "Kismet.Tabs.BlueprintDefaults", new IMAGE_BRUSH( "Icons/icon_BlueprintEditor_Defaults_40x", Icon16x16 ) );
 		}
 
-		Set("Kismet.Palette.Favorites", new IMAGE_BRUSH("Icons/Star_16x", Icon16x16, FLinearColor(0.4f, 0.4, 0.4f, 1.f)));
-		Set("Kismet.Palette.Library",   new IMAGE_BRUSH("Icons/icon_MeshPaint_Find_16x", Icon16x16, FLinearColor(0.4f, 0.4, 0.4f, 1.f)));
+		Set("Kismet.Palette.Favorites", new IMAGE_BRUSH_SVG("Starship/Common/Favorite", Icon16x16));
 		const FCheckBoxStyle KismetFavoriteToggleStyle = FCheckBoxStyle()
 			.SetCheckBoxType(ESlateCheckBoxType::CheckBox)
 			.SetUncheckedImage( IMAGE_BRUSH("Icons/EmptyStar_16x", Icon10x10, FLinearColor(0.8f, 0.8f, 0.8f, 1.f)) )
@@ -5160,63 +5130,55 @@ void FStarshipEditorStyle::FStyle::SetupPersonaStyle()
 
 		Set( "Kismet.Tooltip.SubtextFont", DEFAULT_FONT("Regular", 8) );
 
+		Set("Blueprint.CompileStatus.Background", new IMAGE_BRUSH_SVG("Starship/Blueprints/CompileStatus_Background", Icon20x20));
+
+		// @todo - Icon Replacement - trying out tinting compile backgrounds
+		Set("Blueprint.CompileStatus.Background.Unknown", new IMAGE_BRUSH_SVG("Starship/Blueprints/CompileStatus_Background", Icon20x20, FStyleColors::AccentYellow));
+		Set("Blueprint.CompileStatus.Background.Warning", new IMAGE_BRUSH_SVG("Starship/Blueprints/CompileStatus_Background", Icon20x20, FStyleColors::Warning));
+		Set("Blueprint.CompileStatus.Background.Good", new IMAGE_BRUSH_SVG("Starship/Blueprints/CompileStatus_Background", Icon20x20, FStyleColors::AccentGreen));
+		Set("Blueprint.CompileStatus.Background.Error", new IMAGE_BRUSH_SVG("Starship/Blueprints/CompileStatus_Background", Icon20x20, FStyleColors::Error));
+
+		Set("Blueprint.CompileStatus.Overlay.Unknown", new IMAGE_BRUSH_SVG("Starship/Blueprints/CompileStatus_Unknown_Badge", Icon20x20, FStyleColors::AccentYellow));
+		Set("Blueprint.CompileStatus.Overlay.Warning", new IMAGE_BRUSH_SVG("Starship/Blueprints/CompileStatus_Warning_Badge", Icon20x20, FStyleColors::Warning));
+		Set("Blueprint.CompileStatus.Overlay.Good", new IMAGE_BRUSH_SVG("Starship/Blueprints/CompileStatus_Good_Badge", Icon20x20, FStyleColors::AccentGreen));
+		Set("Blueprint.CompileStatus.Overlay.Error", new IMAGE_BRUSH_SVG("Starship/Blueprints/CompileStatus_Fail_Badge", Icon20x20, FStyleColors::Error));
+
+		// @todo - Icon Replacement - these are hijacked by non-blueprint things. Keeping them around for now
 		Set( "Kismet.Status.Unknown", new IMAGE_BRUSH( "Old/Kismet2/CompileStatus_Working", Icon40x40 ) );
 		Set( "Kismet.Status.Error", new IMAGE_BRUSH( "Old/Kismet2/CompileStatus_Fail", Icon40x40 ) );
 		Set( "Kismet.Status.Good", new IMAGE_BRUSH( "Old/Kismet2/CompileStatus_Good", Icon40x40 ) );
-		Set( "Kismet.Status.Instrumented", new IMAGE_BRUSH( "Old/Kismet2/CompileStatus_Instrumented", Icon40x40 ) );
-		Set( "Kismet.Status.NotInstrumented", new IMAGE_BRUSH( "Old/Kismet2/CompileStatus_NotInstrumented", Icon40x40 ) );
 		Set( "Kismet.Status.Warning", new IMAGE_BRUSH( "Old/Kismet2/CompileStatus_Warning", Icon40x40 ) );
 
-		Set( "BlueprintEditor.AddNewVariable", new IMAGE_BRUSH( "Icons/icon_Blueprint_AddVariable_40px", Icon40x40) );
-		Set( "BlueprintEditor.AddNewVariable.Small", new IMAGE_BRUSH( "Icons/icon_Blueprint_AddVariable_40px", Icon20x20) );
-		Set( "BlueprintEditor.AddNewVariableButton", new IMAGE_BRUSH( "Icons/icon_Blueprint_AddVariable_40px", Icon12x12) );
-		Set( "BlueprintEditor.AddNewLocalVariable", new IMAGE_BRUSH( "Icons/icon_Blueprint_AddVariable_40px", Icon40x40) );
-		Set( "BlueprintEditor.AddNewLocalVariable.Small", new IMAGE_BRUSH( "Icons/icon_Blueprint_AddVariable_40px", Icon20x20) );
-		Set( "BlueprintEditor.AddNewFunction", new IMAGE_BRUSH( "Icons/icon_Blueprint_AddFunction_40px", Icon40x40) );
-		Set( "BlueprintEditor.AddNewFunction.Small", new IMAGE_BRUSH( "Icons/icon_Blueprint_AddFunction_40px", Icon20x20) );
-		Set( "BlueprintEditor.AddNewMacroDeclaration", new IMAGE_BRUSH( "Icons/icon_Blueprint_AddMacro_40px", Icon40x40) );
-		Set( "BlueprintEditor.AddNewMacroDeclaration.Small", new IMAGE_BRUSH( "Icons/icon_Blueprint_AddMacro_40px", Icon20x20) );
-		Set( "BlueprintEditor.AddNewAnimationGraph", new IMAGE_BRUSH( "Icons/icon_BlueprintEditor_AddDocument_40x", Icon40x40) );
-		Set( "BlueprintEditor.AddNewAnimationGraph.Small", new IMAGE_BRUSH( "Icons/icon_BlueprintEditor_AddDocument_40x", Icon20x20) );
-		Set( "BlueprintEditor.AddNewEventGraph", new IMAGE_BRUSH( "Icons/icon_Blueprint_AddGraph_40px", Icon40x40) );
-		Set( "BlueprintEditor.AddNewEventGraph.Small", new IMAGE_BRUSH( "Icons/icon_Blueprint_AddGraph_40px", Icon20x20) );
-		Set( "BlueprintEditor.ManageInterfaces", new IMAGE_BRUSH( "Icons/icon_BlueprintEditor_Interfaces_40x", Icon40x40) );	
-		Set( "BlueprintEditor.ManageInterfaces.Small", new IMAGE_BRUSH( "Icons/icon_BlueprintEditor_Interfaces_40x", Icon20x20) );	
-		Set( "BlueprintEditor.AddNewDelegate.Small", new IMAGE_BRUSH( "Icons/icon_Blueprint_AddDelegate_40x", Icon20x20) );
-		Set( "BlueprintEditor.AddNewDelegate", new IMAGE_BRUSH( "Icons/icon_Blueprint_AddDelegate_40x", Icon40x40) );
-		Set( "BlueprintEditor.AddNewAnimationLayer.Small", new IMAGE_BRUSH( "Icons/icon_Blueprint_AddFunction_40px", Icon20x20) );
-		Set( "BlueprintEditor.AddNewAnimationLayer", new IMAGE_BRUSH( "Icons/icon_Blueprint_AddFunction_40px", Icon40x40) );
-
-		Set( "Kismet.Status.Unknown.Small", new IMAGE_BRUSH( "Old/Kismet2/CompileStatus_Unknown_Small", Icon16x16 ) );
-		Set( "Kismet.Status.Error.Small", new IMAGE_BRUSH( "Old/Kismet2/CompileStatus_Broken_Small", Icon16x16 ) );
-		Set( "Kismet.Status.Good.Small", new IMAGE_BRUSH( "Old/Kismet2/CompileStatus_Good_Small", Icon16x16 ) );
-		Set( "Kismet.Status.Warning.Small", new IMAGE_BRUSH( "Old/Kismet2/CompileStatus_Warning_Small", Icon16x16 ) );
-
-		Set( "Kismet.TitleBarEditor.ArrowUp", new IMAGE_BRUSH( "Old/ArrowUp", Icon16x16 ) );
-		Set( "Kismet.TitleBarEditor.ArrowDown", new IMAGE_BRUSH( "Old/ArrowDown", Icon16x16 ) );
-
-		Set( "Kismet.VariableList.TypeIcon", new IMAGE_BRUSH( "/Icons/pill_16x", Icon16x16 ) );
-		Set( "Kismet.VariableList.ArrayTypeIcon", new IMAGE_BRUSH( "/Icons/pillarray_16x", Icon16x16 ) );
-		Set( "Kismet.VariableList.SetTypeIcon", new IMAGE_BRUSH( "/Icons/pillset_16x", Icon16x16 ) );
-		Set( "Kismet.VariableList.SetTypeIconLarge", new IMAGE_BRUSH( "/Icons/pillset_40x", Icon40x40 ) );
-		Set( "Kismet.VariableList.MapValueTypeIcon", new IMAGE_BRUSH( "/Icons/pillmapvalue_16x", Icon16x16 ) );
-		Set( "Kismet.VariableList.MapKeyTypeIcon", new IMAGE_BRUSH( "/Icons/pillmapkey_16x", Icon16x16 ) );
+		Set("BlueprintEditor.AddNewVariable", new IMAGE_BRUSH_SVG("Starship/Blueprints/icon_Blueprint_AddVariable", Icon20x20));
+		Set("BlueprintEditor.AddNewLocalVariable", new IMAGE_BRUSH_SVG("Starship/Blueprints/icon_Blueprint_AddVariable", Icon20x20));
+		Set("BlueprintEditor.AddNewFunction", new IMAGE_BRUSH_SVG("Starship/Blueprints/icon_Blueprint_AddFunction", Icon20x20));
+		Set("BlueprintEditor.AddNewMacroDeclaration", new IMAGE_BRUSH_SVG("Starship/Blueprints/icon_Blueprint_AddMacro", Icon20x20));
+
+
+		Set( "BlueprintEditor.AddNewAnimationGraph", new CORE_IMAGE_BRUSH_SVG("Starship/Common/folder-plus", Icon16x16));
+		Set( "BlueprintEditor.AddNewEventGraph", new IMAGE_BRUSH_SVG( "Starship/Blueprints/icon_Blueprint_AddGraph", Icon20x20) );
+
+		Set( "BlueprintEditor.AddNewDelegate", new IMAGE_BRUSH_SVG( "Starship/Blueprints/icon_Blueprint_AddDelegate", Icon20x20) );
+
+		Set( "BlueprintEditor.AddNewAnimationLayer", new IMAGE_BRUSH_SVG( "Starship/Blueprints/icon_Blueprint_AddFunction", Icon20x20) );
+
+
+		Set("Kismet.VariableList.TypeIcon", new IMAGE_BRUSH_SVG( "Starship/Blueprints/pill", Icon16x16));
+		Set("Kismet.VariableList.ArrayTypeIcon", new IMAGE_BRUSH_SVG( "Starship/Blueprints/pillarray", Icon16x16));
+		Set("Kismet.VariableList.SetTypeIcon", new IMAGE_BRUSH_SVG( "Starship/Blueprints/pillset", Icon16x16));
+		Set("Kismet.VariableList.SetTypeIconLarge", new IMAGE_BRUSH_SVG( "Starship/Blueprints/pillset", Icon40x40));
+		Set("Kismet.VariableList.MapValueTypeIcon", new IMAGE_BRUSH_SVG( "Starship/Blueprints/pillmapvalue", Icon16x16));
+		Set("Kismet.VariableList.MapKeyTypeIcon", new IMAGE_BRUSH_SVG( "Starship/Blueprints/pillmapkey", Icon16x16));
 
 		Set("Kismet.VariableList.ExposeForInstance", new CORE_IMAGE_BRUSH_SVG("Starship/Common/visible", Icon16x16));
 		Set("Kismet.VariableList.HideForInstance", new CORE_IMAGE_BRUSH_SVG("Starship/Common/hidden", Icon16x16));
-	
-		Set( "Kismet.VariableList.VariableIsUsed", new IMAGE_BRUSH( "/Icons/icon_variable_used_16x", Icon16x16 ) );
-		Set( "Kismet.VariableList.VariableNotUsed", new IMAGE_BRUSH( "/Icons/icon_variable_not_used_16x", Icon16x16 ) );
-
-		Set( "Kismet.VariableList.Replicated", new IMAGE_BRUSH( "/Icons/icon_replication_16px", Icon16x16, FLinearColor(0.0f, 0.0f, 0.0f, 1.0f) ) );
-		Set( "Kismet.VariableList.NotReplicated", new IMAGE_BRUSH( "/Icons/icon_replication_16px", Icon16x16, FLinearColor(0.0f, 0.0f, 0.0f, 0.5f) ) );
 
 		Set( "Kismet.Explorer.Title", FTextBlockStyle(NormalText) .SetFont(DEFAULT_FONT( "BoldCondensedItalic", 11)));
 		Set( "Kismet.Explorer.SearchDepthFont", DEFAULT_FONT( "Bold", 14) );
 
 		Set( "Kismet.Interfaces.Title",  FTextBlockStyle(NormalText) .SetFont( DEFAULT_FONT("Bold", 11 ) ) );
-		Set( "Kismet.Interfaces.Implement", new IMAGE_BRUSH( "Icons/assign_left_16x", Icon16x16) );
-		Set( "Kismet.Interfaces.Remove", new IMAGE_BRUSH( "Icons/assign_right_16x", Icon16x16) );
+		Set( "Kismet.Interfaces.Implement", new CORE_IMAGE_BRUSH_SVG( "Starship/Common/PreviousArrow", Icon16x16) );
+		Set( "Kismet.Interfaces.Remove", new CORE_IMAGE_BRUSH_SVG( "Icons/assign_right_16x", Icon16x16) );
 
 		Set( "Kismet.TypePicker.CategoryFont", DEFAULT_FONT( "BoldCondensedItalic", 11) );
 		Set( "Kismet.TypePicker.NormalFont", DEFAULT_FONT( "Regular", 11) );
@@ -5224,25 +5186,23 @@ void FStarshipEditorStyle::FStyle::SetupPersonaStyle()
 		Set( "Kismet.GraphPicker.Title", FTextBlockStyle(NormalText) .SetFont( DEFAULT_FONT("BoldCondensedItalic", 11) ) );
 
 		Set( "Kismet.CreateBlueprint", new IMAGE_BRUSH( "/Icons/CreateBlueprint", Icon16x16) );
-		Set( "Kismet.HarvestBlueprintFromActors", new IMAGE_BRUSH( "/Icons/HarvestBlueprintFromActors", Icon16x16) );
+		Set( "Kismet.HarvestBlueprintFromActors", new IMAGE_BRUSH_SVG( "Starship/Blueprints/HarvestBlueprintFromActors", Icon16x16) );
 
-		Set( "Kismet.Comment.Handle", new IMAGE_BRUSH( "Old/Kismet2/Comment_Handle", FVector2D(14.0f, 12.0f)) );
 		Set( "Kismet.Comment.Background", new IMAGE_BRUSH( "Old/Kismet2/Comment_Background", FVector2D(100.0f, 68.0f)) );
 
-		Set( "Kismet.AllClasses.VariableIcon", new IMAGE_BRUSH( "/Icons/pill_16x", Icon16x16 ) );
-		Set( "Kismet.AllClasses.ArrayVariableIcon", new IMAGE_BRUSH( "/Icons/pillarray_16x", Icon16x16 ) );
-		Set( "Kismet.AllClasses.SetVariableIcon", new IMAGE_BRUSH( "/Icons/pillset_16x", Icon16x16 ) );
-		Set( "Kismet.AllClasses.MapValueVariableIcon", new IMAGE_BRUSH( "/Icons/pillmapvalue_16x", Icon16x16 ) );
-		Set( "Kismet.AllClasses.MapKeyVariableIcon", new IMAGE_BRUSH( "/Icons/pillmapkey_16x", Icon16x16 ) );
-		Set( "Kismet.AllClasses.FunctionIcon", new IMAGE_BRUSH( "/Icons/icon_BluePrintEditor_Function_16px", Icon16x16 ) );
-
-		Set( "BlueprintEditor.ResetCamera", new IMAGE_BRUSH( "Icons/icon_Camera_Reset_40px", Icon16x16));
-		Set( "Kismet.SetRealtimePreview", new IMAGE_BRUSH( "Icons/icon_MatEd_Realtime_40x", Icon16x16));
-		Set( "BlueprintEditor.ShowFloor", new IMAGE_BRUSH( "Icons/icon_Show_Floor_40px", Icon16x16));
-		Set( "BlueprintEditor.ShowGrid", new IMAGE_BRUSH( "Icons/icon_ShowGrid_16x", Icon16x16 ) );
+		Set( "Kismet.AllClasses.VariableIcon", new IMAGE_BRUSH_SVG("Starship/Blueprints/pill", Icon16x16));
+		Set( "Kismet.AllClasses.ArrayVariableIcon", new IMAGE_BRUSH_SVG("Starship/Blueprints/pillarray", Icon16x16));
+		Set( "Kismet.AllClasses.SetVariableIcon", new IMAGE_BRUSH_SVG("Starship/Blueprints/pillset", Icon16x16));
+		Set( "Kismet.AllClasses.MapValueVariableIcon", new IMAGE_BRUSH_SVG("Starship/Blueprints/pillmapvalue", Icon16x16));
+		Set( "Kismet.AllClasses.MapKeyVariableIcon", new IMAGE_BRUSH_SVG("Starship/Blueprints/pillmapkey", Icon16x16));
+
+		Set( "Kismet.AllClasses.FunctionIcon", new IMAGE_BRUSH_SVG("Starship/Blueprints/icon_Blueprint_Function", Icon16x16));
+
+		Set( "BlueprintEditor.ResetCamera", new IMAGE_BRUSH_SVG("Starship/Common/ResetCamera", Icon16x16));
+		Set( "BlueprintEditor.ShowFloor", new IMAGE_BRUSH_SVG( "Starship/Common/Floor", Icon16x16));
+		Set( "BlueprintEditor.ShowGrid", new IMAGE_BRUSH_SVG( "Starship/Common/Grid", Icon16x16 ) );
 		Set( "BlueprintEditor.EnableSimulation", new IMAGE_BRUSH( "Icons/icon_Enable_Simulation_40px", Icon40x40));
-		Set( "BlueprintEditor.EnableProfiling", new IMAGE_BRUSH("/Icons/icon_Enable_Profiling_40x", Icon40x40) );
-		Set( "BlueprintEditor.EnableSimulation.Small", new IMAGE_BRUSH( "Icons/icon_Enable_Simulation_40px", Icon20x20));
+		Set( "BlueprintEditor.EnableProfiling", new IMAGE_BRUSH_SVG("Starship/Common/Statistics", Icon20x20) );
 		Set( "SCS.NativeComponent", new IMAGE_BRUSH( "Icons/NativeSCSComponent", Icon20x20 ));
 		Set( "SCS.Component", new IMAGE_BRUSH( "Icons/SCSComponent", Icon20x20 ));
 
@@ -5334,16 +5294,12 @@ void FStarshipEditorStyle::FStyle::SetupPersonaStyle()
 
 	//Blueprint Diff
 	{
-		Set("BlueprintDif.HasGraph", new IMAGE_BRUSH("/Icons/blueprint_Dif_has_graph_8x", Icon8x8));
-		Set("BlueprintDif.HasGraph.Small", new IMAGE_BRUSH("/Icons/blueprint_Dif_has_graph_8x", Icon8x8));
-		Set("BlueprintDif.MissingGraph", new IMAGE_BRUSH("/Icons/blueprint_Dif_missing_graph_8x", Icon8x8));
-		Set("BlueprintDif.MissingGraph.Small", new IMAGE_BRUSH("/Icons/blueprint_Dif_missing_graph_8x", Icon8x8));
-		Set("BlueprintDif.NextDiff", new IMAGE_BRUSH("/Icons/diff_next_40x", Icon16x16));
-		Set("BlueprintDif.NextDiff.Small", new IMAGE_BRUSH("/Icons/diff_next_40x", Icon16x16));
-		Set("BlueprintDif.PrevDiff", new IMAGE_BRUSH("/Icons/diff_prev_40x", Icon16x16));
-		Set("BlueprintDif.PrevDiff.Small", new IMAGE_BRUSH("/Icons/diff_prev_40x", Icon16x16));
-		Set("BlueprintDif.HorizontalDiff.Small", new IMAGE_BRUSH("/Icons/icon_horizontal_diff_view_16px", Icon16x16));
-		Set("BlueprintDif.VerticalDiff.Small", new IMAGE_BRUSH("/Icons/icon_vertical_diff_view_16px", Icon16x16));
+		Set("BlueprintDif.HasGraph", new IMAGE_BRUSH_SVG("Starship/Blueprints/blueprint_Dif_has_graph", Icon10x10));
+		Set("BlueprintDif.MissingGraph", new IMAGE_BRUSH_SVG("Starship/Blueprints/blueprint_Dif_missing_graph", Icon8x8));
+		Set("BlueprintDif.NextDiff", new IMAGE_BRUSH_SVG("/Starship/Blueprints/diff_next", Icon16x16));
+		Set("BlueprintDif.PrevDiff", new IMAGE_BRUSH_SVG("/Starship/Blueprints/diff_prev", Icon16x16));
+		Set("BlueprintDif.HorizontalDiff", new IMAGE_BRUSH_SVG("Starship/Blueprints/icon_horizontal_diff_view", Icon16x16));
+		Set("BlueprintDif.VerticalDiff", new IMAGE_BRUSH_SVG("Starship/Blueprints/icon_vertical_diff_view", Icon16x16));
 
 		Set("BlueprintDif.ItalicText", 
 			FTextBlockStyle(NormalText)
@@ -5354,20 +5310,19 @@ void FStarshipEditorStyle::FStyle::SetupPersonaStyle()
 
 	//Blueprint Merge
 	{
-		Set("BlueprintMerge.NextDiff", new IMAGE_BRUSH("/Icons/diff_next_40x", Icon16x16));
-		Set("BlueprintMerge.NextDiff.Small", new IMAGE_BRUSH("/Icons/diff_next_40x", Icon16x16));
-		Set("BlueprintMerge.PrevDiff", new IMAGE_BRUSH("/Icons/diff_prev_40x", Icon16x16));
-		Set("BlueprintMerge.PrevDiff.Small", new IMAGE_BRUSH("/Icons/diff_prev_40x", Icon16x16));
+		Set("BlueprintMerge.NextDiff", new IMAGE_BRUSH_SVG("Starship/Blueprints/diff_next", Icon16x16));
+		Set("BlueprintMerge.PrevDiff", new IMAGE_BRUSH_SVG("Starship/Blueprints/diff_prev", Icon16x16));
+
 		Set("BlueprintMerge.Finish", new IMAGE_BRUSH("/Icons/LV_Save", Icon16x16));
-		Set("BlueprintMerge.Finish.Small", new IMAGE_BRUSH("/Icons/LV_Save", Icon16x16));
+
 		Set("BlueprintMerge.Cancel", new IMAGE_BRUSH("/Icons/LV_Remove", Icon16x16));
-		Set("BlueprintMerge.Cancel.Small", new IMAGE_BRUSH("/Icons/LV_Remove", Icon16x16));
+
 		Set("BlueprintMerge.AcceptSource", new IMAGE_BRUSH("/Icons/AcceptMergeSource_40x", Icon16x16));
-		Set("BlueprintMerge.AcceptSource.Small", new IMAGE_BRUSH("/Icons/AcceptMergeSource_40x", Icon16x16));
+
 		Set("BlueprintMerge.AcceptTarget", new IMAGE_BRUSH("/Icons/AcceptMergeTarget_40x", Icon16x16));
-		Set("BlueprintMerge.AcceptTarget.Small", new IMAGE_BRUSH("/Icons/AcceptMergeTarget_40x", Icon16x16));
+
 		Set("BlueprintMerge.StartMerge", new IMAGE_BRUSH("/Icons/StartMerge_42x", Icon16x16));
-		Set("BlueprintMerge.StartMerge.Small", new IMAGE_BRUSH("/Icons/StartMerge_42x", Icon16x16));
+
 	}
 
 	// Play in editor / play in world
@@ -5451,28 +5406,27 @@ void FStarshipEditorStyle::FStyle::SetupPersonaStyle()
 
 	// Kismet 2 debugger
 	{
-		Set( "Kismet.Breakpoint.Disabled", new IMAGE_BRUSH( "Old/Kismet2/Breakpoint_Disabled_Small", Icon16x16 ) );
-		Set( "Kismet.Breakpoint.EnabledAndInvalid", new IMAGE_BRUSH( "Old/Kismet2/Breakpoint_Invalid_Small", Icon16x16 ) );
-		Set( "Kismet.Breakpoint.EnabledAndValid", new IMAGE_BRUSH( "Old/Kismet2/Breakpoint_Valid_Small", Icon16x16 ) );
-		Set( "Kismet.Breakpoint.NoneSpacer", new IMAGE_BRUSH( "Old/Kismet2/Breakpoint_None_Small", Icon16x16 ) );
-		Set( "Kismet.Breakpoint.MixedStatus", new IMAGE_BRUSH( "Old/Kismet2/Breakpoint_Mixed_Small", Icon16x16 ) );
+		Set( "Kismet.Breakpoint.Disabled", new IMAGE_BRUSH_SVG( "Starship/Blueprints/Breakpoint_Disabled", Icon16x16, FStyleColors::AccentRed ) );
+		Set( "Kismet.Breakpoint.EnabledAndInvalid", new IMAGE_BRUSH_SVG( "Starship/Blueprints/Breakpoint_Invalid", Icon16x16, FStyleColors::AccentRed) );
+		Set( "Kismet.Breakpoint.EnabledAndValid", new IMAGE_BRUSH_SVG( "Starship/Blueprints/Breakpoint_Valid", Icon16x16, FStyleColors::AccentRed) );
+		Set( "Kismet.Breakpoint.MixedStatus", new IMAGE_BRUSH_SVG( "Starship/Blueprints/Breakpoint_Mixed", Icon16x16, FStyleColors::AccentRed) );
 		
-		Set( "Kismet.WatchIcon", new IMAGE_BRUSH( "Old/Kismet2/WatchIcon", Icon16x16 ) );
-		Set( "Kismet.LatentActionIcon", new IMAGE_BRUSH( "Old/Kismet2/LatentActionIcon", Icon16x16 ) );
+		Set( "Kismet.WatchIcon", new IMAGE_BRUSH_SVG( "Starship/GraphEditors/WatchVariable", Icon16x16 ) );
+		Set( "Kismet.LatentActionIcon", new IMAGE_BRUSH_SVG( "Starship/Common/Timecode", Icon16x16 ) );
 
-		Set( "Kismet.Trace.CurrentIndex", new IMAGE_BRUSH( "Old/Kismet2/CurrentInstructionOverlay_Small", Icon16x16 ) );
-		Set( "Kismet.Trace.PreviousIndex", new IMAGE_BRUSH( "Old/Kismet2/FaintInstructionOverlay_Small", Icon16x16 ) );
+		Set( "Kismet.Trace.CurrentIndex", new IMAGE_BRUSH_SVG( "Starship/GraphEditors/NextArrow", Icon16x16 ) );
+		Set( "Kismet.Trace.PreviousIndex", new IMAGE_BRUSH_SVG( "Starship/GraphEditors/PreviousArrow", Icon16x16 ) );
 
-		Set( "Kismet.DebuggerOverlay.Breakpoint.Disabled", new IMAGE_BRUSH( "Old/Kismet2/Breakpoint_Disabled", Icon32x32 ) );
-		Set( "Kismet.DebuggerOverlay.Breakpoint.EnabledAndInvalid", new IMAGE_BRUSH( "Old/Kismet2/Breakpoint_Invalid", Icon32x32 ) );
-		Set( "Kismet.DebuggerOverlay.Breakpoint.EnabledAndValid", new IMAGE_BRUSH( "Old/Kismet2/Breakpoint_Valid", Icon32x32 ) );
+		Set( "Kismet.DebuggerOverlay.Breakpoint.Disabled", new IMAGE_BRUSH_SVG( "Starship/Blueprints/Breakpoint_Disabled", Icon32x32, FStyleColors::AccentRed) );
+		Set( "Kismet.DebuggerOverlay.Breakpoint.EnabledAndInvalid", new IMAGE_BRUSH_SVG( "Starship/Blueprints/Breakpoint_Invalid", Icon32x32, FStyleColors::AccentRed) );
+		Set( "Kismet.DebuggerOverlay.Breakpoint.EnabledAndValid", new IMAGE_BRUSH_SVG( "Starship/Blueprints/Breakpoint_Valid", Icon32x32, FStyleColors::AccentRed) );
 
-		Set( "Kismet.DebuggerOverlay.InstructionPointer", new IMAGE_BRUSH( "Old/Kismet2/IP_Normal", FVector2D(128,96)) );
-		Set( "Kismet.DebuggerOverlay.InstructionPointerBreakpoint", new IMAGE_BRUSH( "Old/Kismet2/IP_Breakpoint", FVector2D(128,96)) );
+		Set("Kismet.DebuggerOverlay.InstructionPointer", new IMAGE_BRUSH_SVG( "Starship/Blueprints/IP_Breakpoint", FVector2D(128,96)) );
+		Set("Kismet.DebuggerOverlay.InstructionPointerBreakpoint", new IMAGE_BRUSH_SVG( "Starship/Blueprints/IP_Breakpoint", FVector2D(128,96), FStyleColors::AccentRed) );
 
-		Set ("Kismet.CallStackViewer.CurrentStackFrame", new IMAGE_BRUSH( "Old/Kismet2/DebuggerArrow", Icon12x12 ));
-		Set( "Kismet.CallStackViewer.CurrentStackFrameColor", FLinearColor(0.728f, 0.364f, 0.003f) );
-		Set( "Kismet.CallStackViewer.LastStackFrameNavigatedToColor", FLinearColor( 0.4f, 0.5f, 0.7f, 1.0f ) );
+		Set("Kismet.CallStackViewer.CurrentStackFrame", new IMAGE_BRUSH_SVG( "Starship/Blueprints/DebuggerArrow", Icon12x12 ));
+		Set("Kismet.CallStackViewer.CurrentStackFrameColor", FLinearColor(0.728f, 0.364f, 0.003f) );
+		Set("Kismet.CallStackViewer.LastStackFrameNavigatedToColor", FLinearColor( 0.4f, 0.5f, 0.7f, 1.0f ) );
 	}
 
 	// Asset context menu
@@ -6718,8 +6672,8 @@ void FStarshipEditorStyle::FStyle::SetupToolkitStyles()
 		Set( "StaticMeshEditor.SetShowWireframe.Small", new IMAGE_BRUSH( "Icons/icon_StaticMeshEd_Wireframe_40x", Icon20x20 ) );
 		Set( "StaticMeshEditor.SetShowVertexColor", new IMAGE_BRUSH( "Icons/icon_StaticMeshEd_VertColor_40x", Icon40x40 ) );
 		Set( "StaticMeshEditor.SetShowVertexColor.Small", new IMAGE_BRUSH( "Icons/icon_StaticMeshEd_VertColor_40x", Icon20x20 ) );
-		Set( "StaticMeshEditor.SetRealtimePreview", new IMAGE_BRUSH( "Icons/icon_MatEd_Realtime_40x", Icon40x40 ) );
-		Set( "StaticMeshEditor.SetRealtimePreview.Small", new IMAGE_BRUSH( "Icons/icon_MatEd_Realtime_40x", Icon20x20 ) );
+		Set( "StaticMeshEditor.SetRealtimePreview", new IMAGE_BRUSH_SVG("Starship/Common/Realtime", Icon16x16));
+
 		Set( "StaticMeshEditor.ReimportMesh", new IMAGE_BRUSH(TEXT("Icons/icon_Persona_ReimportMesh_40x"), Icon40x40));
 		Set( "StaticMeshEditor.ReimportMesh.Small", new IMAGE_BRUSH(TEXT("Icons/icon_Persona_ReimportMesh_40x"), Icon20x20));
 		Set( "StaticMeshEditor.SetShowBounds", new IMAGE_BRUSH( "Icons/icon_StaticMeshEd_Bounds_40x", Icon40x40 ) );
@@ -6730,8 +6684,8 @@ void FStarshipEditorStyle::FStyle::SetupToolkitStyles()
 		Set( "StaticMeshEditor.SetShowGrid.Small", new IMAGE_BRUSH( "Icons/icon_MatEd_Grid_40x", Icon20x20 ) );
 		Set( "StaticMeshEditor.SetDrawUVs", new IMAGE_BRUSH( "Icons/icon_StaticMeshEd_UVOverlay_40x", Icon40x40 ) );
 		Set( "StaticMeshEditor.SetDrawUVs.Small", new IMAGE_BRUSH( "Icons/icon_StaticMeshEd_UVOverlay_40x", Icon20x20 ) );
-		Set( "StaticMeshEditor.ResetCamera", new IMAGE_BRUSH( "Icons/icon_Camera_Reset_40px", Icon40x40 ) );
-		Set( "StaticMeshEditor.ResetCamera.Small", new IMAGE_BRUSH( "Icons/icon_Camera_Reset_40px", Icon20x20 ) );
+		Set( "StaticMeshEditor.ResetCamera", new IMAGE_BRUSH_SVG("Starship/Common/ResetCamera", Icon16x16));
+
 		Set( "StaticMeshEditor.SetShowPivot", new IMAGE_BRUSH( "Icons/icon_StaticMeshEd_ShowPivot_40x", Icon40x40 ) );
 		Set( "StaticMeshEditor.SetShowPivot.Small", new IMAGE_BRUSH( "Icons/icon_StaticMeshEd_ShowPivot_40x", Icon20x20 ) );
 		Set( "StaticMeshEditor.SetShowSockets", new IMAGE_BRUSH( "Icons/icon_StaticMeshEd_ShowSockets_40x", Icon40x40 ) );
@@ -7547,7 +7501,7 @@ void FStarshipEditorStyle::FStyle::SetupAutomationStyles()
 	// Settings Editor
 	{
 		Set( "SettingsEditor.Collision_Engine", new IMAGE_BRUSH("Icons/icon_Cascade_RestartSim_40x", Icon16x16));
-		Set( "SettingsEditor.Collision_Game", new IMAGE_BRUSH("Icons/icon_MatEd_Realtime_40x", Icon16x16));
+		Set( "SettingsEditor.Collision_Game", new IMAGE_BRUSH_SVG("Starship/Common/Realtime", Icon16x16));
 
 		// Settings editor
 		Set("SettingsEditor.GoodIcon", new IMAGE_BRUSH("Settings/Settings_Good", Icon40x40));
diff --git a/Engine/Source/Editor/GraphEditor/Private/MaterialNodes/SGraphNodeMaterialBase.cpp b/Engine/Source/Editor/GraphEditor/Private/MaterialNodes/SGraphNodeMaterialBase.cpp
--- a/Engine/Source/Editor/GraphEditor/Private/MaterialNodes/SGraphNodeMaterialBase.cpp
+++ b/Engine/Source/Editor/GraphEditor/Private/MaterialNodes/SGraphNodeMaterialBase.cpp
@@ -528,7 +528,7 @@ ECheckBoxState SGraphNodeMaterialBase::IsExpressionPreviewChecked() const
 
 const FSlateBrush* SGraphNodeMaterialBase::GetExpressionPreviewArrow() const
 {
-	return FEditorStyle::GetBrush(MaterialNode->MaterialExpression->bCollapsed ? TEXT("Kismet.TitleBarEditor.ArrowDown") : TEXT("Kismet.TitleBarEditor.ArrowUp"));
+	return FEditorStyle::GetBrush(MaterialNode->MaterialExpression->bCollapsed ? TEXT("Icons.ChevronDown") : TEXT("Icons.ChevronUp"));
 }
 
 void SGraphNodeMaterialBase::PopulateMetaTag(FGraphNodeMetaData* TagMeta) const
diff --git a/Engine/Source/Editor/GraphEditor/Private/SGraphNode.cpp b/Engine/Source/Editor/GraphEditor/Private/SGraphNode.cpp
--- a/Engine/Source/Editor/GraphEditor/Private/SGraphNode.cpp
+++ b/Engine/Source/Editor/GraphEditor/Private/SGraphNode.cpp
@@ -1072,7 +1072,7 @@ ECheckBoxState SGraphNode::IsAdvancedViewChecked() const
 const FSlateBrush* SGraphNode::GetAdvancedViewArrow() const
 {
 	const bool bAdvancedPinsHidden = GraphNode && (ENodeAdvancedPins::Hidden == GraphNode->AdvancedPinDisplay);
-	return FEditorStyle::GetBrush(bAdvancedPinsHidden ? TEXT("Kismet.TitleBarEditor.ArrowDown") : TEXT("Kismet.TitleBarEditor.ArrowUp"));
+	return FEditorStyle::GetBrush(bAdvancedPinsHidden ? TEXT("Icons.ChevronDown") : TEXT("Icons.ChevronUp"));
 }
 
 /** Create widget to show/hide advanced pins */
diff --git a/Engine/Source/Editor/GraphEditor/Private/SGraphNodeComment.cpp b/Engine/Source/Editor/GraphEditor/Private/SGraphNodeComment.cpp
--- a/Engine/Source/Editor/GraphEditor/Private/SGraphNodeComment.cpp
+++ b/Engine/Source/Editor/GraphEditor/Private/SGraphNodeComment.cpp
@@ -348,7 +348,7 @@ void SGraphNodeComment::GetOverlayBrushes(bool bSelected, const FVector2D Widget
 
 	HandleSelection(bSelected);
 
-	FOverlayBrushInfo HandleBrush = FEditorStyle::GetBrush( TEXT("Kismet.Comment.Handle") );
+	FOverlayBrushInfo HandleBrush = FEditorStyle::GetBrush( TEXT("Graph.Node.Comment.Handle") );
 
 	HandleBrush.OverlayOffset.X = WidgetSize.X - HandleBrush.Brush->ImageSize.X - Fudge;
 	HandleBrush.OverlayOffset.Y = WidgetSize.Y - HandleBrush.Brush->ImageSize.Y - Fudge;
diff --git a/Engine/Source/Editor/GraphEditor/Private/SGraphPanel.cpp b/Engine/Source/Editor/GraphEditor/Private/SGraphPanel.cpp
--- a/Engine/Source/Editor/GraphEditor/Private/SGraphPanel.cpp
+++ b/Engine/Source/Editor/GraphEditor/Private/SGraphPanel.cpp
@@ -269,6 +269,9 @@ int32 SGraphPanel::OnPaint( const FPaintArgs& Args, const FGeometry& AllottedGeo
 								const float BounceValue = FMath::Sin(2.0f * PI * BounceCurve.GetLerp());
 								BouncedGeometry.DrawPosition += (OverlayInfo.AnimationEnvelope * BounceValue * ZoomFactor);
 
+								FLinearColor FinalColorAndOpacity(InWidgetStyle.GetColorAndOpacityTint()* OverlayBrush->GetTint(InWidgetStyle));
+								//FinalColorAndOpacity.A = Alpha;
+
 								CurWidgetsMaxLayerId++;
 								FSlateDrawElement::MakeBox(
 									OutDrawElements,
@@ -276,7 +279,7 @@ int32 SGraphPanel::OnPaint( const FPaintArgs& Args, const FGeometry& AllottedGeo
 									BouncedGeometry,
 									OverlayBrush,
 									ESlateDrawEffect::None,
-									FLinearColor(1.0f, 1.0f, 1.0f, Alpha)
+									FinalColorAndOpacity
 									);
 							}
 
diff --git a/Engine/Source/Editor/Kismet/Private/BlueprintDetailsCustomization.cpp b/Engine/Source/Editor/Kismet/Private/BlueprintDetailsCustomization.cpp
--- a/Engine/Source/Editor/Kismet/Private/BlueprintDetailsCustomization.cpp
+++ b/Engine/Source/Editor/Kismet/Private/BlueprintDetailsCustomization.cpp
@@ -2808,13 +2808,17 @@ void FBlueprintGraphArgumentLayout::GenerateHeaderRowContent( FDetailWidgetRow&
 		.FillWidth(1)
 		.VAlign(VAlign_Center)
 		[
-			SAssignNew(ArgumentNameWidget, SEditableTextBox)
+			SNew(SBox)
+			.MinDesiredWidth(125.f)
+			[
+				SAssignNew(ArgumentNameWidget, SEditableTextBox)
 				.Text( this, &FBlueprintGraphArgumentLayout::OnGetArgNameText )
 				.OnTextChanged(this, &FBlueprintGraphArgumentLayout::OnArgNameChange)
 				.OnTextCommitted(this, &FBlueprintGraphArgumentLayout::OnArgNameTextCommitted)
 				.ToolTipText(this, &FBlueprintGraphArgumentLayout::OnGetArgToolTipText)
 				.Font( IDetailLayoutBuilder::GetDetailFont() )
 				.IsEnabled(!ShouldPinBeReadOnly())
+			]
 		]
 	]
 	.ValueContent()
@@ -3421,6 +3425,7 @@ void FBlueprintGraphActionDetails::CustomizeDetails( IDetailLayoutBuilder& Detai
 				[
 					SAssignNew( ColorBlock, SColorBlock )
 						.Color( this, &FBlueprintGraphActionDetails::GetNodeTitleColor )
+						.CornerRadius(FVector4(4.0f, 4.0f, 4.0f, 4.0f))
 						.AlphaDisplayMode(EColorBlockAlphaDisplayMode::Ignore)
 						.OnMouseButtonDown( this, &FBlueprintGraphActionDetails::ColorBlock_OnMouseButtonDown )
 				];
diff --git a/Engine/Source/Editor/Kismet/Private/FormatTextDetails.cpp b/Engine/Source/Editor/Kismet/Private/FormatTextDetails.cpp
--- a/Engine/Source/Editor/Kismet/Private/FormatTextDetails.cpp
+++ b/Engine/Source/Editor/Kismet/Private/FormatTextDetails.cpp
@@ -142,7 +142,7 @@ void FFormatTextArgumentLayout::GenerateHeaderRowContent( FDetailWidgetRow& Node
 				.IsEnabled( bIsMoveUpEnabled )
 				[
 					SNew(SImage)
-					.Image(FEditorStyle::GetBrush("BlueprintEditor.Details.ArgUpButton"))
+					.Image(FEditorStyle::GetBrush("Icons.ChevronUp"))
 				]
 			]
 			+SHorizontalBox::Slot()
@@ -155,7 +155,7 @@ void FFormatTextArgumentLayout::GenerateHeaderRowContent( FDetailWidgetRow& Node
 				.IsEnabled( bIsMoveDownEnabled )
 				[
 					SNew(SImage)
-					.Image(FEditorStyle::GetBrush("BlueprintEditor.Details.ArgDownButton"))
+					.Image(FEditorStyle::GetBrush("Icons.ChevronDown"))
 				]
 			]
 
diff --git a/Engine/Source/Editor/Kismet/Private/SBlueprintEditorToolbar.cpp b/Engine/Source/Editor/Kismet/Private/SBlueprintEditorToolbar.cpp
--- a/Engine/Source/Editor/Kismet/Private/SBlueprintEditorToolbar.cpp
+++ b/Engine/Source/Editor/Kismet/Private/SBlueprintEditorToolbar.cpp
@@ -547,18 +547,25 @@ FSlateIcon FBlueprintEditorToolbar::GetStatusImage() const
 		Status = BS_UpToDate;
 	}
 
+	
+	static const FName CompileStatusBackground("Blueprint.CompileStatus.Background");
+	static const FName CompileStatusUnknown("Blueprint.CompileStatus.Overlay.Unknown");
+	static const FName CompileStatusError("Blueprint.CompileStatus.Overlay.Error");
+	static const FName CompileStatusGood("Blueprint.CompileStatus.Overlay.Good");
+	static const FName CompileStatusWarning("Blueprint.CompileStatus.Overlay.Warning");
+
 	switch (Status)
 	{
 	default:
 	case BS_Unknown:
 	case BS_Dirty:
-		return FSlateIcon(FEditorStyle::GetStyleSetName(), "Kismet.Status.Unknown");
+		return FSlateIcon(FEditorStyle::GetStyleSetName(), CompileStatusBackground, NAME_None, CompileStatusUnknown);
 	case BS_Error:
-		return FSlateIcon(FEditorStyle::GetStyleSetName(), "Kismet.Status.Error");
+		return FSlateIcon(FEditorStyle::GetStyleSetName(), CompileStatusBackground, NAME_None, CompileStatusError);
 	case BS_UpToDate:
-		return FSlateIcon(FEditorStyle::GetStyleSetName(), "Kismet.Status.Good");
+		return FSlateIcon(FEditorStyle::GetStyleSetName(), CompileStatusBackground, NAME_None, CompileStatusGood);
 	case BS_UpToDateWithWarnings:
-		return FSlateIcon(FEditorStyle::GetStyleSetName(), "Kismet.Status.Warning");
+		return FSlateIcon(FEditorStyle::GetStyleSetName(), CompileStatusBackground, NAME_None, CompileStatusWarning);
 	}
 }
 
diff --git a/Engine/Source/Editor/Kismet/Private/SBlueprintLibraryPalette.cpp b/Engine/Source/Editor/Kismet/Private/SBlueprintLibraryPalette.cpp
--- a/Engine/Source/Editor/Kismet/Private/SBlueprintLibraryPalette.cpp
+++ b/Engine/Source/Editor/Kismet/Private/SBlueprintLibraryPalette.cpp
@@ -196,7 +196,7 @@ void SBlueprintLibraryPalette::Construct(FArguments const& InArgs, TWeakPtr<FBlu
 {
 	SBlueprintSubPalette::FArguments SuperArgs;
 	SuperArgs._Title       = LOCTEXT("PaletteTitle", "Find a Node");
-	SuperArgs._Icon        = FEditorStyle::GetBrush("Kismet.Palette.Library");
+	SuperArgs._Icon        = FEditorStyle::GetBrush("Icons.Search");
 	SuperArgs._ToolTipText = LOCTEXT("PaletteToolTip", "An all encompassing list of every node that is available for this blueprint.");
 	SuperArgs._ShowFavoriteToggles = true;
 
diff --git a/Engine/Source/Editor/Kismet/Private/SGraphTitleBar.cpp b/Engine/Source/Editor/Kismet/Private/SGraphTitleBar.cpp
--- a/Engine/Source/Editor/Kismet/Private/SGraphTitleBar.cpp
+++ b/Engine/Source/Editor/Kismet/Private/SGraphTitleBar.cpp
@@ -139,6 +139,7 @@ void SGraphTitleBar::Construct( const FArguments& InArgs )
 					[
 						SNew(SImage)
 						.Image( this, &SGraphTitleBar::GetTypeGlyph )
+						.ColorAndOpacity(FSlateColor::UseForeground())
 					]
 					+SHorizontalBox::Slot()
 					.FillWidth(1.f)
diff --git a/Engine/Source/Editor/Kismet/Private/SGraphTitleBarAddNewBookmark.cpp b/Engine/Source/Editor/Kismet/Private/SGraphTitleBarAddNewBookmark.cpp
--- a/Engine/Source/Editor/Kismet/Private/SGraphTitleBarAddNewBookmark.cpp
+++ b/Engine/Source/Editor/Kismet/Private/SGraphTitleBarAddNewBookmark.cpp
@@ -11,6 +11,7 @@
 #include "Widgets/Input/SEditableTextBox.h"
 #include "Widgets/Layout/SSeparator.h"
 #include "Widgets/Layout/SUniformGridPanel.h"
+#include "Widgets/Images/SImage.h"
 
 #define LOCTEXT_NAMESPACE "SGraphTitleBarAddNewBookmark"
 
@@ -21,115 +22,103 @@ void SGraphTitleBarAddNewBookmark::Construct(const FArguments& InArgs)
 	SComboButton::FArguments Args;
 	Args.ButtonContent()
 	[
-		SNew(STextBlock)
-		.TextStyle(FEditorStyle::Get(), "GraphBookmarkButtonText")
-		.Text(this, &SGraphTitleBarAddNewBookmark::GetAddButtonGlyph)
+		SNew(SImage)
+		.Image(FAppStyle::Get().GetBrush("GraphEditor.Bookmark"))
+		.ColorAndOpacity(FSlateColor::UseForeground())
 	]
 	.MenuContent()
 	[
-		SNew(SBorder)
-		.BorderImage(FEditorStyle::GetBrush("Menu.Background"))
-		.Padding(4)
+		SNew(SBox)
+		.MinDesiredWidth(300)
 		[
-			SNew(SBox)
-			.MinDesiredWidth(300)
+			SNew(SVerticalBox)
+			+ SVerticalBox::Slot()
+			.Padding(8)
+			.FillHeight(1.f)
+			.HAlign(HAlign_Fill)
+			.VAlign(VAlign_Fill)
 			[
-				SNew(SVerticalBox)
-				+ SVerticalBox::Slot()
-				.Padding(8)
-				.FillHeight(1.f)
-				.HAlign(HAlign_Fill)
-				.VAlign(VAlign_Fill)
+				SNew(STextBlock)
+				.Text(this, &SGraphTitleBarAddNewBookmark::GetPopupTitleText)
+				.Font(FEditorStyle::GetFontStyle("StandardDialog.LargeFont"))
+			]
+			+ SVerticalBox::Slot()
+			.Padding(8, 4, 8, 8)
+			.AutoHeight()
+			[
+				SNew(SHorizontalBox)
+				+ SHorizontalBox::Slot()
+				.Padding(6)
+				.AutoWidth()
 				[
 					SNew(STextBlock)
-					.Text(this, &SGraphTitleBarAddNewBookmark::GetPopupTitleText)
-					.Font(FEditorStyle::GetFontStyle("StandardDialog.LargeFont"))
+					.Text(LOCTEXT("BookmarkNameFieldLabel", "Name:"))
 				]
-				+ SVerticalBox::Slot()
-				.Padding(8, 4, 8, 8)
-				.AutoHeight()
+				+ SHorizontalBox::Slot()
+				.Padding(1)
+				.FillWidth(1.f)
 				[
-					SNew(SHorizontalBox)
-					+ SHorizontalBox::Slot()
-					.Padding(6)
-					.AutoWidth()
-					[
-						SNew(STextBlock)
-						.Text(LOCTEXT("BookmarkNameFieldLabel", "Name:"))
-					]
-					+ SHorizontalBox::Slot()
-					.Padding(1)
-					.FillWidth(1.f)
-					[
-						SAssignNew(NameEntryWidget, SEditableTextBox)
-						.SelectAllTextWhenFocused(true)
-						.OnTextCommitted(this, &SGraphTitleBarAddNewBookmark::OnNameTextCommitted)
-						.OnTextChanged(this, &SGraphTitleBarAddNewBookmark::OnNameTextCommitted, ETextCommit::Default)
-						.Text_Lambda([this]() -> FText
-						{
-							return CurrentNameText;
-						})
-					]
+					SAssignNew(NameEntryWidget, SEditableTextBox)
+					.SelectAllTextWhenFocused(true)
+					.OnTextCommitted(this, &SGraphTitleBarAddNewBookmark::OnNameTextCommitted)
+					.OnTextChanged(this, &SGraphTitleBarAddNewBookmark::OnNameTextCommitted, ETextCommit::Default)
+					.Text_Lambda([this]() -> FText
+					{
+						return CurrentNameText;
+					})
 				]
-				+ SVerticalBox::Slot()
-				.Padding(8, 4, 4, 8)
-				.AutoHeight()
-				.VAlign(VAlign_Bottom)
+			]
+			+ SVerticalBox::Slot()
+			.Padding(8, 4, 4, 8)
+			.AutoHeight()
+			.VAlign(VAlign_Bottom)
+			[
+				SNew(SHorizontalBox)
+				+SHorizontalBox::Slot()
+				.AutoWidth()
 				[
-					SNew(SHorizontalBox)
-					+SHorizontalBox::Slot()
-					.AutoWidth()
+					SNew(SButton)
+					.Text(LOCTEXT("RemoveButtonLabel", "Remove"))
+					.ButtonStyle(FEditorStyle::Get(), "SecondaryButton")
+					.Visibility_Lambda([this]() -> EVisibility
+					{
+						return CurrentViewBookmarkId.IsValid() ? EVisibility::Visible : EVisibility::Collapsed;
+					})
+					.OnClicked(this, &SGraphTitleBarAddNewBookmark::OnRemoveButtonClicked)
+				]
+				+SHorizontalBox::Slot()
+				.FillWidth(1.f)
+				.HAlign(HAlign_Right)
+				[
+					SNew(SUniformGridPanel)
+					.SlotPadding(FEditorStyle::GetMargin("StandardDialog.SlotPadding"))
+					.MinDesiredSlotWidth(FEditorStyle::GetFloat("StandardDialog.MinDesiredSlotWidth"))
+					.MinDesiredSlotHeight(FEditorStyle::GetFloat("StandardDialog.MinDesiredSlotHeight"))
+					+ SUniformGridPanel::Slot(0, 0)
 					[
 						SNew(SButton)
-						.Text(LOCTEXT("RemoveButtonLabel", "Remove"))
-						.TextStyle(FEditorStyle::Get(), "FlatButton.DefaultTextStyle")
-						.ButtonStyle(FEditorStyle::Get(), "FlatButton.Default")
-						.Visibility_Lambda([this]() -> EVisibility
+						.HAlign(HAlign_Center)
+						.Text(LOCTEXT("CancelButtonLabel", "Cancel"))
+						.OnClicked_Lambda([this]() -> FReply
 						{
-							return CurrentViewBookmarkId.IsValid() ? EVisibility::Visible : EVisibility::Collapsed;
+							SetIsOpen(false);
+							return FReply::Handled();
 						})
-						.OnClicked(this, &SGraphTitleBarAddNewBookmark::OnRemoveButtonClicked)
 					]
-					+SHorizontalBox::Slot()
-					.FillWidth(1.f)
-					.HAlign(HAlign_Right)
+					+ SUniformGridPanel::Slot(1, 0)
 					[
-						SNew(SUniformGridPanel)
-						.SlotPadding(FEditorStyle::GetMargin("StandardDialog.SlotPadding"))
-						.MinDesiredSlotWidth(FEditorStyle::GetFloat("StandardDialog.MinDesiredSlotWidth"))
-						.MinDesiredSlotHeight(FEditorStyle::GetFloat("StandardDialog.MinDesiredSlotHeight"))
-						+ SUniformGridPanel::Slot(0, 0)
-						[
-							SNew(SButton)
-							.HAlign(HAlign_Center)
-							.Text(LOCTEXT("CancelButtonLabel", "Cancel"))
-							.TextStyle(FEditorStyle::Get(), "FlatButton.DefaultTextStyle")
-							.ButtonStyle(FEditorStyle::Get(), "FlatButton.Default")
-							.OnClicked_Lambda([this]() -> FReply
-							{
-								SetIsOpen(false);
-								return FReply::Handled();
-							})
-						]
-						+ SUniformGridPanel::Slot(1, 0)
-						[
-							SNew(SButton)
-							.HAlign(HAlign_Center)
-							.Text(this, &SGraphTitleBarAddNewBookmark::GetAddButtonLabel)
-							.TextStyle(FEditorStyle::Get(), "FlatButton.DefaultTextStyle")
-							.ButtonStyle(FEditorStyle::Get(), "FlatButton.Success")
-							.OnClicked(this, &SGraphTitleBarAddNewBookmark::OnAddButtonClicked)
-							.IsEnabled(this, &SGraphTitleBarAddNewBookmark::IsAddButtonEnabled)
-						]
+						SNew(SButton)
+						.HAlign(HAlign_Center)
+						.Text(this, &SGraphTitleBarAddNewBookmark::GetAddButtonLabel)
+						.ButtonStyle(FEditorStyle::Get(), "PrimaryButton")
+						.OnClicked(this, &SGraphTitleBarAddNewBookmark::OnAddButtonClicked)
+						.IsEnabled(this, &SGraphTitleBarAddNewBookmark::IsAddButtonEnabled)
 					]
 				]
 			]
 		]
 	]
-	.HasDownArrow(false)
-	.ComboButtonStyle(FEditorStyle::Get(), "ComboButton")
-	.ButtonStyle(FEditorStyle::Get(), "GraphBookmarkButton")
-	.ForegroundColor(FLinearColor(1, 1, 1, 0.5))
+	.ComboButtonStyle(FEditorStyle::Get(), "SimpleComboButton")
 	.ToolTipText(LOCTEXT("AddBookmarkButtonToolTip", "Bookmark Current Location"))
 	.OnComboBoxOpened(this, &SGraphTitleBarAddNewBookmark::OnComboBoxOpened);
 
diff --git a/Engine/Source/Editor/Kismet/Private/SSCSEditor.cpp b/Engine/Source/Editor/Kismet/Private/SSCSEditor.cpp
--- a/Engine/Source/Editor/Kismet/Private/SSCSEditor.cpp
+++ b/Engine/Source/Editor/Kismet/Private/SSCSEditor.cpp
@@ -3823,13 +3823,15 @@ FText SSCS_RowWidget_ActorRoot::GetActorMobilityText() const
 
 TSharedRef<SWidget> SSCS_RowWidget_Separator::GenerateWidgetForColumn(const FName& ColumnName)
 {
-	return SNew(SBox)
+	return SNullWidget::NullWidget;
+
+	/*return SNew(SBox)
 		.Padding(1.f)
 		[
 			SNew(SBorder)
 			.Padding(FEditorStyle::GetMargin(TEXT("Menu.Separator.Padding")))
 			.BorderImage(FEditorStyle::GetBrush(TEXT("Menu.Separator")))
-		];
+		];*/
 }
 
 //////////////////////////////////////////////////////////////////////////
@@ -4049,21 +4051,17 @@ void SSCSEditor::Construct( const FArguments& InArgs )
 		+ SVerticalBox::Slot()
 		.AutoHeight()
 		.VAlign(VAlign_Top)
-		.Padding(0)
+		.Padding(4.0)
 		[
-			SNew(SBorder)
-			.Padding(0)
-			.BorderImage(FEditorStyle::GetBrush("DetailsView.CategoryTop"))
-			.AddMetaData<FTagMetaData>(FTagMetaData(TEXT("ComponentsPanel")))
-			[
-				SAssignNew(HeaderBox, SVerticalBox)
-			]
+			SAssignNew(HeaderBox, SVerticalBox)
 		]
 
 		+ SVerticalBox::Slot()
-		.Padding(4.0f, 0.0f)
+		.Padding(4.0)
 		[
-			SNew(SBox)
+			SNew(SBorder)
+			.BorderImage(FAppStyle::Get().GetBrush("Brushes.Recessed"))
+			.Padding(0.0)
 			.AddMetaData<FTagMetaData>(FTagMetaData(TEXT("ComponentsPanel")))
 			.Visibility(this, &SSCSEditor::GetComponentsTreeVisibility)
 			[
diff --git a/Engine/Source/Editor/Kismet/Private/UserDefinedEnumEditor.cpp b/Engine/Source/Editor/Kismet/Private/UserDefinedEnumEditor.cpp
--- a/Engine/Source/Editor/Kismet/Private/UserDefinedEnumEditor.cpp
+++ b/Engine/Source/Editor/Kismet/Private/UserDefinedEnumEditor.cpp
@@ -552,7 +552,7 @@ void FUserDefinedEnumIndexLayout::GenerateHeaderRowContent( FDetailWidgetRow& No
 				.IsEnabled( bIsMoveUpEnabled )
 				[
 					SNew(SImage)
-					.Image(FEditorStyle::GetBrush("BlueprintEditor.Details.ArgUpButton"))
+					.Image(FEditorStyle::GetBrush("Icons.ChevronUp"))
 				]
 			]
 			+SHorizontalBox::Slot()
@@ -566,7 +566,7 @@ void FUserDefinedEnumIndexLayout::GenerateHeaderRowContent( FDetailWidgetRow& No
 				.IsEnabled( bIsMoveDownEnabled )
 				[
 					SNew(SImage)
-					.Image(FEditorStyle::GetBrush("BlueprintEditor.Details.ArgDownButton"))
+					.Image(FEditorStyle::GetBrush("Icons.ChevronDown"))
 				]
 			]
 
diff --git a/Engine/Source/Editor/Kismet/Private/UserDefinedStructureEditor.cpp b/Engine/Source/Editor/Kismet/Private/UserDefinedStructureEditor.cpp
--- a/Engine/Source/Editor/Kismet/Private/UserDefinedStructureEditor.cpp
+++ b/Engine/Source/Editor/Kismet/Private/UserDefinedStructureEditor.cpp
@@ -934,7 +934,7 @@ public:
 				.IsEnabled(!(EMemberFieldPosition::MFP_First & PositionFlags))
 				[
 					SNew(SImage)
-					.Image(FEditorStyle::GetBrush("BlueprintEditor.Details.ArgUpButton"))
+					.Image(FEditorStyle::GetBrush("Icons.ChevronUp"))
 				]
 			]
 			+ SHorizontalBox::Slot()
@@ -948,7 +948,7 @@ public:
 				.IsEnabled(!(EMemberFieldPosition::MFP_Last & PositionFlags))
 				[
 					SNew(SImage)
-					.Image(FEditorStyle::GetBrush("BlueprintEditor.Details.ArgDownButton"))
+					.Image(FEditorStyle::GetBrush("Icons.ChevronDown"))
 				]
 			]
 			+SHorizontalBox::Slot()
diff --git a/Engine/Source/Editor/KismetWidgets/Private/SPinTypeSelector.cpp b/Engine/Source/Editor/KismetWidgets/Private/SPinTypeSelector.cpp
--- a/Engine/Source/Editor/KismetWidgets/Private/SPinTypeSelector.cpp
+++ b/Engine/Source/Editor/KismetWidgets/Private/SPinTypeSelector.cpp
@@ -383,6 +383,7 @@ void SPinTypeSelector::Construct(const FArguments& InArgs, FGetPinTypeTree GetPi
 				.AutoWidth()
 				.VAlign(VAlign_Center)
 				.HAlign(HAlign_Center)
+				.Padding(2.0f)
 				[
 					ContainerControl.ToSharedRef()
 				];
diff --git a/Engine/Source/Runtime/Slate/Private/Framework/MultiBox/SToolBarButtonBlock.cpp b/Engine/Source/Runtime/Slate/Private/Framework/MultiBox/SToolBarButtonBlock.cpp
--- a/Engine/Source/Runtime/Slate/Private/Framework/MultiBox/SToolBarButtonBlock.cpp
+++ b/Engine/Source/Runtime/Slate/Private/Framework/MultiBox/SToolBarButtonBlock.cpp
@@ -3,7 +3,6 @@
 #include "Framework/MultiBox/SToolBarButtonBlock.h"
 #include "Widgets/SBoxPanel.h"
 #include "Framework/Application/SlateApplication.h"
-#include "Widgets/Images/SImage.h"
 #include "Framework/MultiBox/MultiBoxBuilder.h"
 #include "Widgets/Input/SButton.h"
 #include "Widgets/SToolTip.h"
@@ -11,6 +10,7 @@
 #include "Widgets/Layout/SBox.h"
 #include "Framework/MultiBox/SToolBarComboButtonBlock.h"
 #include "Styling/ToolBarStyle.h"
+#include "Widgets/Images/SLayeredImage.h"
 
 
 FToolBarButtonBlock::FToolBarButtonBlock( const TSharedPtr< const FUICommandInfo > InCommand, TSharedPtr< const FUICommandList > InCommandList, const TAttribute<FText>& InLabelOverride, const TAttribute<FText>& InToolTipOverride, const TAttribute<FSlateIcon>& InIconOverride )
@@ -135,9 +135,9 @@ void SToolBarButtonBlock::BuildMultiBlockWidget(const ISlateStyle* StyleSet, con
 
 	TSharedRef<const FMultiBox> MultiBox = OwnerMultiBoxWidget.Pin()->GetMultiBox();
 
-	TSharedRef< const FToolBarButtonBlock > ToolBarButtonBlock = StaticCastSharedRef< const FToolBarButtonBlock >(MultiBlock.ToSharedRef());
+	TSharedRef<const FToolBarButtonBlock > ToolBarButtonBlock = StaticCastSharedRef< const FToolBarButtonBlock >(MultiBlock.ToSharedRef());
 
-	TSharedPtr< const FUICommandInfo > UICommand = ToolBarButtonBlock->GetAction();
+	TSharedPtr<const FUICommandInfo> UICommand = ToolBarButtonBlock->GetAction();
 
 	// Allow the block to override the action's label and tool tip string, if desired
 	TAttribute<FText> ActualLabel;
@@ -165,14 +165,15 @@ void SToolBarButtonBlock::BuildMultiBlockWidget(const ISlateStyle* StyleSet, con
 
 	// If a key is bound to the command, append it to the tooltip text.
 	TWeakPtr<const FUICommandInfo> Action = ToolBarButtonBlock->GetAction();
-	ActualToolTip = TAttribute< FText >::Create( TAttribute< FText >::FGetter::CreateStatic( &Local::AppendKeyBindingToToolTip, ActualToolTip, Action ) );
+	ActualToolTip = TAttribute< FText >::Create(TAttribute<FText>::FGetter::CreateStatic(&Local::AppendKeyBindingToToolTip, ActualToolTip, Action ) );
 	
 	// If we were supplied an image than go ahead and use that, otherwise we use a null widget
-	TSharedRef<SImage> IconWidget =
-		SNew(SImage)
-		.ColorAndOpacity(FSlateColor::UseForeground())
+	TSharedRef<SLayeredImage> IconWidget =
+		SNew(SLayeredImage)
+		.ColorAndOpacity(this, &SToolBarButtonBlock::GetIconForegroundColor)
 		.Image(this, &SToolBarButtonBlock::GetIconBrush);
 
+	IconWidget->AddLayer(TAttribute<const FSlateBrush*>(this, &SToolBarButtonBlock::GetOverlayIconBrush));
 
 	// Create the content for our button
 	TSharedRef<SWidget> ButtonContent = SNullWidget::NullWidget;
@@ -180,6 +181,7 @@ void SToolBarButtonBlock::BuildMultiBlockWidget(const ISlateStyle* StyleSet, con
 	{
 		const FVector2D IconSize = ToolBarStyle.IconSize;
 
+		IconWidget->SetDesiredSizeOverride(IconSize);
 		ButtonContent =
 			SNew(SHorizontalBox)
 			.AddMetaData<FTagMetaData>(FTagMetaData(TutorialHighlightName))
@@ -188,12 +190,7 @@ void SToolBarButtonBlock::BuildMultiBlockWidget(const ISlateStyle* StyleSet, con
 			.VAlign(VAlign_Center)
 			.HAlign(HAlign_Center)
 			[
-				SNew(SBox)
-				.WidthOverride(IconSize.X)
-				.HeightOverride(IconSize.Y)
-				[
-					IconWidget
-				]	
+				IconWidget
 			]
 			// Label text
 			+ SHorizontalBox::Slot()
@@ -276,7 +273,7 @@ void SToolBarButtonBlock::BuildMultiBlockWidget(const ISlateStyle* StyleSet, con
 		ChildSlot
 		[
 			// Create a check box
-			SNew( SCheckBox )
+			SNew(SCheckBox)
 			// Use the tool bar style for this check box
 			.Style(&CheckStyle)
 			.IsFocusable(bIsFocusable)
@@ -317,7 +314,7 @@ void SToolBarButtonBlock::BuildMultiBlockWidget(const ISlateStyle* StyleSet, con
 	SetEnabled( TAttribute< bool >( this, &SToolBarButtonBlock::IsEnabled ) );
 
 	// Bind our widget's visible state to whether or not the button should be visible
-	SetVisibility( TAttribute<EVisibility>(this, &SToolBarButtonBlock::GetBlockVisibility) );
+	SetVisibility(TAttribute<EVisibility>(this, &SToolBarButtonBlock::GetBlockVisibility));
 }
 
 
@@ -446,9 +443,24 @@ const FSlateBrush* SToolBarButtonBlock::GetIconBrush() const
 	return bForceSmallIcons || FMultiBoxSettings::UseSmallToolBarIcons.Get() ? GetSmallIconBrush() : GetNormalIconBrush();
 }
 
+const FSlateBrush* SToolBarButtonBlock::GetOverlayIconBrush() const
+{
+	TSharedRef<const FToolBarButtonBlock> ToolBarButtonBlock = StaticCastSharedRef<const FToolBarButtonBlock >(MultiBlock.ToSharedRef());
+
+	const FSlateIcon ActionIcon = ToolBarButtonBlock->GetAction().IsValid() ? ToolBarButtonBlock->GetAction()->GetIcon() : FSlateIcon();
+	const FSlateIcon& ActualIcon = ToolBarButtonBlock->IconOverride.IsSet() ? ToolBarButtonBlock->IconOverride.Get() : ActionIcon;
+
+	if (ActualIcon.IsSet())
+	{
+		return ActualIcon.GetOverlayIcon();
+	}
+
+	return nullptr;
+}
+
 const FSlateBrush* SToolBarButtonBlock::GetNormalIconBrush() const
 {
-	TSharedRef< const FToolBarButtonBlock > ToolBarButtonBlock = StaticCastSharedRef< const FToolBarButtonBlock >(MultiBlock.ToSharedRef());
+	TSharedRef<const FToolBarButtonBlock> ToolBarButtonBlock = StaticCastSharedRef<const FToolBarButtonBlock >(MultiBlock.ToSharedRef());
 
 	const FSlateIcon ActionIcon = ToolBarButtonBlock->GetAction().IsValid() ? ToolBarButtonBlock->GetAction()->GetIcon() : FSlateIcon();
 	const FSlateIcon& ActualIcon = ToolBarButtonBlock->IconOverride.IsSet() ? ToolBarButtonBlock->IconOverride.Get() : ActionIcon;
@@ -491,3 +503,15 @@ const FSlateBrush* SToolBarButtonBlock::GetSmallIconBrush() const
 		return StyleSet->GetBrush(IconName);
 	}
 }
+
+FSlateColor SToolBarButtonBlock::GetIconForegroundColor() const
+{
+	// If any brush has a tint, don't assume it should be subdued
+	const FSlateBrush* Brush = GetIconBrush();
+	if (Brush && Brush->TintColor != FLinearColor::White)
+	{
+		return FLinearColor::White;
+	}
+
+	return FSlateColor::UseForeground();
+}
diff --git a/Engine/Source/Runtime/Slate/Public/Framework/MultiBox/SToolBarButtonBlock.h b/Engine/Source/Runtime/Slate/Public/Framework/MultiBox/SToolBarButtonBlock.h
--- a/Engine/Source/Runtime/Slate/Public/Framework/MultiBox/SToolBarButtonBlock.h
+++ b/Engine/Source/Runtime/Slate/Public/Framework/MultiBox/SToolBarButtonBlock.h
@@ -158,14 +158,19 @@ protected:
 	EVisibility GetBlockVisibility() const;
 
 private:
-
 	/** Called by Slate to determine whether icons/labels are visible */
 	EVisibility GetIconVisibility(bool bIsASmallIcon) const;
-	
+
 	/** Gets the icon brush for the toolbar block widget */
 	const FSlateBrush* GetIconBrush() const;
+
+	const FSlateBrush* GetOverlayIconBrush() const;
+
+	/** Gets the icon brush for the toolbar block widget */
+	const FSlateBrush* GetNormalIconBrush() const;
 	const FSlateBrush* GetSmallIconBrush() const;
 
+	FSlateColor GetIconForegroundColor() const;
 private:
 
 	TAttribute< EVisibility > LabelVisibility;
diff --git a/Engine/Source/Runtime/Slate/Public/Widgets/Colors/SColorBlock.h b/Engine/Source/Runtime/Slate/Public/Widgets/Colors/SColorBlock.h
--- a/Engine/Source/Runtime/Slate/Public/Widgets/Colors/SColorBlock.h
+++ b/Engine/Source/Runtime/Slate/Public/Widgets/Colors/SColorBlock.h
@@ -98,8 +98,6 @@ private:
 
 	TAttribute<const FSlateBrush*> AlphaBackgroundBrush;
 
-	TAttribute<const FSlateBrush*> SolidBackgroundBrush;
-
 	TAttribute<FVector4> GradientCornerRadius;
 
 	/** Whether the color displayed is HSV or not */
diff --git a/Engine/Source/Runtime/SlateCore/Private/Styling/StarshipCoreStyle.cpp b/Engine/Source/Runtime/SlateCore/Private/Styling/StarshipCoreStyle.cpp
--- a/Engine/Source/Runtime/SlateCore/Private/Styling/StarshipCoreStyle.cpp
+++ b/Engine/Source/Runtime/SlateCore/Private/Styling/StarshipCoreStyle.cpp
@@ -1805,7 +1805,7 @@ void FStarshipCoreStyle::SetupMultiboxStyles(TSharedRef<FStyle>& Style)
 			.SetSubMenuIndicator(IMAGE_BRUSH("Common/SubmenuArrow", Icon8x8))
 			.SetComboButtonPadding(FMargin(6.0f, 0.0f))
 			.SetButtonPadding(FMargin(4.0f, 0.0f))
-			.SetCheckBoxPadding(FMargin(10.0f, 0.0f))
+			.SetCheckBoxPadding(FMargin(4.0f, 0.0f))
 			.SetSeparatorBrush(FSlateColorBrush(FStyleColors::Recessed))
 			.SetSeparatorPadding(FMargin(8.f, 0))
 			.SetLabelStyle(FTextBlockStyle(NormalText))
@@ -1831,7 +1831,12 @@ void FStarshipCoreStyle::SetupMultiboxStyles(TSharedRef<FStyle>& Style)
 		SlimToolbarStyle.SetToggleButtonStyle(SlimToolBarToggleButtonCheckBoxStyle);
 		SlimToolbarStyle.SetButtonStyle(SlimToolBarButton);
 
-		SlimToolbarStyle.SetSettingsComboButtonStyle(SlimToolBarComboButton);
+		FComboButtonStyle SlimToolBarSettingsComboButton = FComboButtonStyle(Style->GetWidgetStyle<FComboButtonStyle>("ComboButton"))
+			.SetContentPadding(0)
+			.SetButtonStyle(SlimToolBarButton)
+			.SetDownArrowImage(IMAGE_BRUSH_SVG("Starship/Common/ellipsis-vertical-narrow", FVector2D(6, 24)));
+
+		SlimToolbarStyle.SetSettingsComboButtonStyle(SlimToolBarSettingsComboButton);
 		SlimToolbarStyle.SetIconSize(Icon20x20);
 
 		Style->Set("SlimToolBar", SlimToolbarStyle);
diff --git a/Engine/Source/Runtime/SlateCore/Private/Styling/StyleColors.cpp b/Engine/Source/Runtime/SlateCore/Private/Styling/StyleColors.cpp
--- a/Engine/Source/Runtime/SlateCore/Private/Styling/StyleColors.cpp
+++ b/Engine/Source/Runtime/SlateCore/Private/Styling/StyleColors.cpp
@@ -57,6 +57,8 @@ const FSlateColor FStyleColors::AccentGray = EStyleColor::AccentGray;
 const FSlateColor FStyleColors::AccentWhite = EStyleColor::AccentWhite;
 const FSlateColor FStyleColors::AccentFolder = EStyleColor::AccentFolder;
 
+const FSlateColor FStyleColors::Warning= EStyleColor::Warning;
+const FSlateColor FStyleColors::Error = EStyleColor::Error;
 
 USlateThemeManager::USlateThemeManager()
 {
@@ -112,6 +114,9 @@ void USlateThemeManager::InitalizeDefaults()
 	SetDefaultColor(EStyleColor::AccentGray, COLOR("#808080FF"));
 	SetDefaultColor(EStyleColor::AccentWhite, COLOR("#FFFFFFFF"));
 	SetDefaultColor(EStyleColor::AccentFolder, COLOR("#B68F55FF"));
+
+	SetDefaultColor(EStyleColor::Warning, COLOR("#D78102FF"));
+	SetDefaultColor(EStyleColor::Error, COLOR("#CF1212FF"));
 }
 
 #if ALLOW_THEMES
diff --git a/Engine/Source/Runtime/SlateCore/Private/Textures/SlateIcon.cpp b/Engine/Source/Runtime/SlateCore/Private/Textures/SlateIcon.cpp
--- a/Engine/Source/Runtime/SlateCore/Private/Textures/SlateIcon.cpp
+++ b/Engine/Source/Runtime/SlateCore/Private/Textures/SlateIcon.cpp
@@ -12,23 +12,16 @@ FSlateIcon::FSlateIcon( )
 	: StyleSetName(NAME_None)
 	, StyleName(NAME_None)
 	, SmallStyleName(NAME_None)
-	, bIsSet(false)
+	, StatusOverlayStyleName(NAME_None)
 { }
 
 
-FSlateIcon::FSlateIcon( const FName& InStyleSetName, const FName& InStyleName )
-	: StyleSetName(InStyleSetName)
-	, StyleName(InStyleName)
-	, SmallStyleName(ISlateStyle::Join(InStyleName, ".Small"))
-	, bIsSet(true)
-{ }
-
 
-FSlateIcon::FSlateIcon( const FName& InStyleSetName, const FName& InStyleName, const FName& InSmallStyleName )
+FSlateIcon::FSlateIcon(const FName InStyleSetName, const FName InStyleName, const FName InSmallStyleName, const FName InStatusOverlayStyleName)
 	: StyleSetName(InStyleSetName)
 	, StyleName(InStyleName)
-	, SmallStyleName(InSmallStyleName)
-	, bIsSet(true)
+	, SmallStyleName(InSmallStyleName == NAME_None ? ISlateStyle::Join(InStyleName, ".Small") :  InSmallStyleName)
+	, StatusOverlayStyleName(InStatusOverlayStyleName)
 { }
 
 
@@ -41,10 +34,14 @@ const FSlateBrush* FSlateIcon::GetSmallIcon( ) const
 	
 	if (StyleSet)
 	{
-		return StyleSet->GetOptionalBrush(SmallStyleName);
+		const FSlateBrush* FoundBrush = StyleSet->GetOptionalBrush(SmallStyleName);
+		if (FoundBrush != FStyleDefaults::GetNoBrush())
+		{
+			return FoundBrush;
+		}
 	}
 
-	return FStyleDefaults::GetNoBrush();
+	return FSlateIcon::GetIcon();
 }
 
 
@@ -69,13 +66,24 @@ const FSlateBrush* FSlateIcon::GetIcon( ) const
 
 const FSlateBrush* FSlateIcon::GetOptionalIcon( ) const
 {
-	const ISlateStyle* StyleSet = GetStyleSet();
-	return StyleSet ? StyleSet->GetOptionalBrush(StyleName) : nullptr;
+	const FSlateBrush* Icon = GetIcon();
+	return Icon == FStyleDefaults::GetNoBrush() ? nullptr : Icon;
 }
 
-
 const FSlateBrush* FSlateIcon::GetOptionalSmallIcon( ) const
+{
+	const FSlateBrush* Icon = GetSmallIcon();
+	return Icon == FStyleDefaults::GetNoBrush() ? nullptr : Icon;
+}
+
+const FSlateBrush* FSlateIcon::GetOverlayIcon() const
 {
 	const ISlateStyle* StyleSet = GetStyleSet();
-	return StyleSet ? StyleSet->GetOptionalBrush(SmallStyleName) : nullptr;
+
+	if (StyleSet)
+	{
+		return StyleSet->GetOptionalBrush(StatusOverlayStyleName);
+	}
+
+	return nullptr;
 }
diff --git a/Engine/Source/Runtime/SlateCore/Public/Styling/StyleColors.h b/Engine/Source/Runtime/SlateCore/Public/Styling/StyleColors.h
--- a/Engine/Source/Runtime/SlateCore/Public/Styling/StyleColors.h
+++ b/Engine/Source/Runtime/SlateCore/Public/Styling/StyleColors.h
@@ -57,6 +57,8 @@ enum class EStyleColor : uint8
 	AccentGray,
 	AccentWhite,
 	AccentFolder,
+	Warning,
+	Error,
 
 	/** Only user colors should be below this line
 	 * To use user colors:
@@ -321,4 +323,7 @@ struct SLATECORE_API FStyleColors
 	static const FSlateColor AccentGray;
 	static const FSlateColor AccentWhite;
 	static const FSlateColor AccentFolder;
+
+	static const FSlateColor Warning;
+	static const FSlateColor Error;
 };
diff --git a/Engine/Source/Runtime/SlateCore/Public/Textures/SlateIcon.h b/Engine/Source/Runtime/SlateCore/Public/Textures/SlateIcon.h
--- a/Engine/Source/Runtime/SlateCore/Public/Textures/SlateIcon.h
+++ b/Engine/Source/Runtime/SlateCore/Public/Textures/SlateIcon.h
@@ -14,15 +14,7 @@ struct SLATECORE_API FSlateIcon
 	/**
 	 * Default constructor (empty icon).
 	 */
-	FSlateIcon( );
-
-	/**
-	 * Creates and initializes a new icon from a style set and style name
-	 *
-	 * @param InStyleSetName The name of the style set the icon can be found in.
-	 * @param StyleName The name of the style for the icon (assumes there may also be a ".Small" variant)
-	 */
-	FSlateIcon( const FName& InStyleSetName, const FName& InStyleName );
+	FSlateIcon();
 
 	/**
 	 * Creates and initializes a new icon from a style set and style name
@@ -30,8 +22,9 @@ struct SLATECORE_API FSlateIcon
 	 * @param InStyleSetName The name of the style set the icon can be found in.
 	 * @param StyleName The name of the style for the icon
 	 * @param InSmallStyleName The name of the style for the small icon
+	 * @param InStatusOverlayStyleName The name of the style for a status icon to overlay on the base image
 	 */
-	FSlateIcon( const FName& InStyleSetName, const FName& InStyleName, const FName& InSmallStyleName );
+	FSlateIcon(const FName InStyleSetName, const FName InStyleName, const FName InSmallStyleName = NAME_None, const FName InStatusOverlayStyleName = NAME_None);
 
 public:
 	
@@ -40,7 +33,11 @@ public:
 	 */
 	friend bool operator==(const FSlateIcon& A, const FSlateIcon& B)
 	{
-		return A.bIsSet == B.bIsSet && A.StyleSetName == B.StyleSetName && A.StyleName == B.StyleName && A.SmallStyleName == B.SmallStyleName;
+		return A.IsSet() == B.IsSet()
+			&& A.StyleSetName == B.StyleSetName
+			&& A.StyleName == B.StyleName
+			&& A.SmallStyleName == B.SmallStyleName
+			&& A.StatusOverlayStyleName == B.StatusOverlayStyleName;
 	}
 
 	/**
@@ -76,7 +73,6 @@ public:
 	 */
 	const FSlateBrush* GetSmallIcon( ) const;
 
-
 	/**
 	 * Optionally gets the resolved small icon, returning nullptr if it's not defined
 	 *
@@ -84,13 +80,15 @@ public:
 	 */
 	const FSlateBrush* GetOptionalSmallIcon( ) const;
 
+	const FSlateBrush* GetOverlayIcon() const;
+
 	/**
 	 * Gets the name of the style for the icon.
 	 *
 	 * @return Style name.
 	 * @see GetStyleName, GetStyleSet, GetStyleSetName
 	 */
-	const FName& GetSmallStyleName( ) const
+	const FName& GetSmallStyleName() const
 	{
 		return SmallStyleName;
 	}
@@ -101,7 +99,7 @@ public:
 	 * @return Style name.
 	 * @see GetSmallStyleName, GetStyleSet, GetStyleSetName
 	 */
-	const FName& GetStyleName( ) const
+	const FName& GetStyleName() const
 	{
 		return StyleName;
 	}
@@ -112,7 +110,7 @@ public:
 	 * @return Style set, or nullptr if the style set wasn't found.
 	 * @see GetSmallStyleName, GetStyleName, GetStyleSetName
 	 */
-	const class ISlateStyle* GetStyleSet( ) const;
+	const class ISlateStyle* GetStyleSet() const;
 
 	/**
 	 * Gets the name of the style set the icon can be found in.
@@ -120,7 +118,7 @@ public:
 	 * @return Style name.
 	 * @see GetSmallStyleName, GetStyleName, GetStyleSet
 	 */
-	const FName& GetStyleSetName( ) const
+	const FName& GetStyleSetName() const
 	{
 		return StyleSetName;
 	}
@@ -130,9 +128,9 @@ public:
 	 *
 	 * @return true if the icon is set, false otherwise.
 	 */
-	const bool IsSet( ) const
+	const bool IsSet() const
 	{
-		return bIsSet;
+		return StyleSetName != NAME_None && StyleName != NAME_None;
 	}
 
 private:
@@ -146,6 +144,9 @@ private:
 	// The name of the style for the small icon.
 	FName SmallStyleName;
 
+	// Name of the style for the status overlay icon
+	FName StatusOverlayStyleName;
+
 	// Whether this icon has been explicitly set-up.
 	bool bIsSet;
 };
