Viewport Toolbar Updates

diff --git a/Engine/Source/Editor/EditorStyle/Private/StarshipStyle.cpp b/Engine/Source/Editor/EditorStyle/Private/StarshipStyle.cpp
--- a/Engine/Source/Editor/EditorStyle/Private/StarshipStyle.cpp
+++ b/Engine/Source/Editor/EditorStyle/Private/StarshipStyle.cpp
@@ -2352,59 +2352,117 @@ void FStarshipEditorStyle::FStyle::SetupViewportStyles()
 	{
 		FToolBarStyle ViewportToolbarStyle = FStarshipCoreStyle::GetCoreStyle().GetWidgetStyle<FToolBarStyle>("SlimToolBar");
 
-		FButtonStyle ViewportMenuButton = FButtonStyle(FStarshipCoreStyle::GetCoreStyle().GetWidgetStyle<FButtonStyle>("NoBorder"))
+		FLinearColor TransInput = FStyleColors::Input.GetSpecifiedColor();
+		TransInput.A = .71;
+
+		FMargin ViewportMargin(4.f, 4.f);
+
+		FSlateRoundedBoxBrush* ViewportGroupBrush = new FSlateRoundedBoxBrush(TransInput, 12.f, FStyleColors::Dropdown, 1.0);
+		Set("EditorViewportToolBar.Group", ViewportGroupBrush);
+
+		FButtonStyle ViewportMenuButton = FButtonStyle()
+			.SetNormal(*ViewportGroupBrush)
+			.SetHovered(*ViewportGroupBrush)
+			.SetPressed(*ViewportGroupBrush)
 			.SetNormalForeground(FStyleColors::Foreground)
 			.SetHoveredForeground(FStyleColors::ForegroundHover)
 			.SetPressedForeground(FStyleColors::ForegroundHover)
 			.SetDisabledForeground(FStyleColors::Foreground)
-			.SetNormalPadding(FMargin(0, 2, 0, 2))
-			.SetPressedPadding(FMargin(0, 3, 0, 1));
-
-
+			.SetNormalPadding(ViewportMargin)
+			.SetPressedPadding(ViewportMargin);
 		Set("EditorViewportToolBar.Button", ViewportMenuButton);
 
+		const FCheckBoxStyle ViewportMenuToggleLeftButtonStyle = FCheckBoxStyle(ViewportToolbarStyle.ToggleButton)
+			.SetCheckBoxType(ESlateCheckBoxType::ToggleButton)
+			.SetUncheckedImage(		  BOX_BRUSH("Starship/EditorViewport/ToolBarLeftGroup", 12.f/25.f))
+			.SetUncheckedPressedImage(BOX_BRUSH("Starship/EditorViewport/ToolBarLeftGroup", 12.f/25.f))
+			.SetUncheckedHoveredImage(BOX_BRUSH("Starship/EditorViewport/ToolBarLeftGroup", 12.f/25.f))
+			.SetCheckedHoveredImage(  BOX_BRUSH("Starship/EditorViewport/ToolBarLeftGroup", 12.f/25.f))
+			.SetCheckedPressedImage(  BOX_BRUSH("Starship/EditorViewport/ToolBarLeftGroup", 12.f/25.f))
+			.SetCheckedImage(         BOX_BRUSH("Starship/EditorViewport/ToolBarLeftGroup", 12.f/25.f))
+			.SetPadding(ViewportMargin);
+		Set("EditorViewportToolBar.ToggleButton.Start", ViewportMenuToggleLeftButtonStyle);
+
+		const FCheckBoxStyle ViewportMenuToggleMiddleButtonStyle = FCheckBoxStyle(ViewportToolbarStyle.ToggleButton)
+			.SetCheckBoxType(ESlateCheckBoxType::ToggleButton)
+			.SetUncheckedImage(		  BOX_BRUSH("Starship/EditorViewport/ToolBarMiddleGroup", 12.f/25.f))
+			.SetUncheckedPressedImage(BOX_BRUSH("Starship/EditorViewport/ToolBarMiddleGroup", 12.f/25.f))
+			.SetUncheckedHoveredImage(BOX_BRUSH("Starship/EditorViewport/ToolBarMiddleGroup", 12.f/25.f))
+			.SetCheckedHoveredImage(  BOX_BRUSH("Starship/EditorViewport/ToolBarMiddleGroup", 12.f/25.f))
+			.SetCheckedPressedImage(  BOX_BRUSH("Starship/EditorViewport/ToolBarMiddleGroup", 12.f/25.f))
+			.SetCheckedImage(         BOX_BRUSH("Starship/EditorViewport/ToolBarMiddleGroup", 12.f/25.f))
+			.SetPadding(ViewportMargin);
+		Set("EditorViewportToolBar.ToggleButton.Middle", ViewportMenuToggleMiddleButtonStyle);
+
+		const FCheckBoxStyle ViewportMenuToggleRightButtonStyle = FCheckBoxStyle(ViewportToolbarStyle.ToggleButton)
+			.SetCheckBoxType(ESlateCheckBoxType::ToggleButton)
+			.SetUncheckedImage(		  BOX_BRUSH("Starship/EditorViewport/ToolBarRightGroup", 12.f/25.f))
+			.SetUncheckedPressedImage(BOX_BRUSH("Starship/EditorViewport/ToolBarRightGroup", 12.f/25.f))
+			.SetUncheckedHoveredImage(BOX_BRUSH("Starship/EditorViewport/ToolBarRightGroup", 12.f/25.f))
+			.SetCheckedHoveredImage(  BOX_BRUSH("Starship/EditorViewport/ToolBarRightGroup", 12.f/25.f))
+			.SetCheckedPressedImage(  BOX_BRUSH("Starship/EditorViewport/ToolBarRightGroup", 12.f/25.f))
+			.SetCheckedImage(         BOX_BRUSH("Starship/EditorViewport/ToolBarRightGroup", 12.f/25.f))
+			.SetPadding(ViewportMargin);
+		Set("EditorViewportToolBar.ToggleButton.End", ViewportMenuToggleRightButtonStyle);
+
+		// We want a background-less version as the ComboMenu has its own unified background
+		const FToolBarStyle& SlimCoreToolBarStyle = FStarshipCoreStyle::GetCoreStyle().GetWidgetStyle<FToolBarStyle>("SlimToolBar");
+		Set("EditorViewportToolBar.ComboMenu.ButtonStyle", FButtonStyle(SlimCoreToolBarStyle.ButtonStyle).SetNormalPadding(0.0).SetPressedPadding(0.0));
+		Set("EditorViewportToolBar.ComboMenu.ToggleButton", FCheckBoxStyle(SlimCoreToolBarStyle.ToggleButton).SetPadding(0.0));
+		Set("EditorViewportToolBar.ComboMenu.LabelStyle", SlimCoreToolBarStyle.LabelStyle);
 
 		FCheckBoxStyle MaximizeRestoreButton = FCheckBoxStyle(ViewportToolbarStyle.ToggleButton)
+			.SetUncheckedImage(*ViewportGroupBrush)
+			.SetUncheckedPressedImage(*ViewportGroupBrush)
+			.SetUncheckedHoveredImage(*ViewportGroupBrush)
+			.SetCheckedImage(*ViewportGroupBrush)
+			.SetCheckedHoveredImage(*ViewportGroupBrush)
+			.SetCheckedPressedImage(*ViewportGroupBrush)
 			.SetForegroundColor(FStyleColors::Foreground)
 			.SetPressedForegroundColor(FStyleColors::ForegroundHover)
 			.SetHoveredForegroundColor(FStyleColors::ForegroundHover)
 			.SetCheckedForegroundColor(FStyleColors::Foreground)
 			.SetCheckedPressedForegroundColor(FStyleColors::ForegroundHover)
-			.SetCheckedHoveredForegroundColor(FStyleColors::ForegroundHover);
-
+			.SetCheckedHoveredForegroundColor(FStyleColors::ForegroundHover)
+			.SetPadding(ViewportMargin);
 		Set("EditorViewportToolBar.MaximizeRestoreButton", MaximizeRestoreButton);
 
+		Set("EditorViewportToolBar.Heading.Padding", FMargin(4.f));
+
+
+		// SComboBox 
+		FComboButtonStyle ViewportComboButton = FComboButtonStyle()
+			.SetButtonStyle(ViewportMenuButton)
+			.SetContentPadding(ViewportMargin);
+
 		ViewportToolbarStyle
 			.SetBackground(FSlateNoResource())
 			.SetIconSize(Icon16x16)
 			.SetBackgroundPadding(FMargin(0))
 			.SetLabelPadding(FMargin(0))
-			.SetComboButtonPadding(FMargin(0))
-			.SetBlockPadding(FMargin(2.0f,0.0f))
+			.SetComboButtonPadding(FMargin(4.f, 0.0f))
+			.SetBlockPadding(FMargin(0.0f,0.0f))
 			.SetIndentedBlockPadding(FMargin(0))
 			.SetButtonPadding(FMargin(0))
-			.SetCheckBoxPadding(FMargin(3.0f, 0.0f))
+			.SetCheckBoxPadding(FMargin(4.0f, 0.0f))
+			.SetComboButtonStyle(ViewportComboButton)
 			.SetButtonStyle(ViewportMenuButton)
 			.SetSeparatorBrush(FSlateNoResource())
-			.SetSeparatorPadding(FMargin(4.0f, 0.0f))
+			.SetSeparatorPadding(FMargin(2.0f, 0.0f))
 			.SetExpandBrush(IMAGE_BRUSH("Icons/toolbar_expand_16x", Icon8x8));
-
 		Set("EditorViewportToolBar", ViewportToolbarStyle);
 
-		FButtonStyle ViewportMenuWarningButton = FButtonStyle(FStarshipCoreStyle::GetCoreStyle().GetWidgetStyle<FButtonStyle>("NoBorder"))
+		FButtonStyle ViewportMenuWarningButton = FButtonStyle(ViewportMenuButton)
 			.SetNormalForeground(FStyleColors::AccentYellow)
 			.SetHoveredForeground(FStyleColors::ForegroundHover)
 			.SetPressedForeground(FStyleColors::ForegroundHover)
-			.SetDisabledForeground(FStyleColors::AccentYellow)
-			.SetNormalPadding(FMargin(0, 2, 0, 2))
-			.SetPressedPadding(FMargin(0, 3, 0, 1));
-
+			.SetDisabledForeground(FStyleColors::AccentYellow);
 		Set("EditorViewportToolBar.WarningButton", ViewportMenuWarningButton);
 
 		FLinearColor ToolbarBackgroundColor = FStyleColors::Foldout.GetSpecifiedColor();
 		ToolbarBackgroundColor.A = .75f;
 
-		Set("EditorViewportToolBar.Background", new FSlateColorBrush(ToolbarBackgroundColor));
+		Set("EditorViewportToolBar.Background", new FSlateNoResource());
 		Set("EditorViewportToolBar.OptionsDropdown", new IMAGE_BRUSH_SVG("Starship/EditorViewport/menu", Icon16x16));
 
 		Set("EditorViewportToolBar.Font", FStyleFonts::Get().Normal);
@@ -2420,6 +2478,7 @@ void FStarshipEditorStyle::FStyle::SetupViewportStyles()
 		Set("EditorViewportToolBar.Maximize.Normal", new IMAGE_BRUSH_SVG("Starship/EditorViewport/square", Icon16x16));
 		Set("EditorViewportToolBar.Maximize.Checked", new IMAGE_BRUSH_SVG("Starship/EditorViewport/quad", Icon16x16));
 		Set("EditorViewportToolBar.RestoreFromImmersive.Normal", new IMAGE_BRUSH("Icons/icon_RestoreFromImmersive_16px", Icon16x16));
+
 	}
 
 	// Legacy Viewport ToolbarBar
diff --git a/Engine/Source/Editor/LevelEditor/Private/SLevelViewportToolBar.cpp b/Engine/Source/Editor/LevelEditor/Private/SLevelViewportToolBar.cpp
--- a/Engine/Source/Editor/LevelEditor/Private/SLevelViewportToolBar.cpp
+++ b/Engine/Source/Editor/LevelEditor/Private/SLevelViewportToolBar.cpp
@@ -169,7 +169,6 @@ void SLevelViewportToolBar::Construct( const FArguments& InArgs )
 	FLevelEditorModule& LevelEditorModule = FModuleManager::GetModuleChecked<FLevelEditorModule>(TEXT("LevelEditor"));
 
 	const FMargin ToolbarSlotPadding(4.0f, 1.0f);
-	const FMargin ToolbarButtonPadding(4.0f, 0.0f);
 
 	ChildSlot
 	[
@@ -200,6 +199,7 @@ void SLevelViewportToolBar::Construct( const FArguments& InArgs )
 					SNew( SEditorViewportToolbarMenu )
 					.ParentToolBar( SharedThis( this ) )
 					.Label( this, &SLevelViewportToolBar::GetCameraMenuLabel )
+					.LabelIcon( this, &SLevelViewportToolBar::GetCameraMenuLabelIcon )
 					.AddMetaData<FTagMetaData>(FTagMetaData(TEXT("EditorViewportToolBar.CameraMenu")))
 					.OnGetMenuContent( this, &SLevelViewportToolBar::GenerateCameraMenu ) 
 				]
@@ -288,7 +288,7 @@ void SLevelViewportToolBar::Construct( const FArguments& InArgs )
 				+ SHorizontalBox::Slot()
 				.HAlign(HAlign_Right)
 				.AutoWidth()
-				.Padding(ToolbarButtonPadding)
+				.Padding(ToolbarSlotPadding)
 				[
 					//The Maximize/Minimize button is only displayed when not in Immersive mode.
 					SNew(SEditorViewportToolBarButton)
@@ -303,7 +303,7 @@ void SLevelViewportToolBar::Construct( const FArguments& InArgs )
 				+ SHorizontalBox::Slot()
 				.HAlign(HAlign_Right)
 				.AutoWidth()
-				.Padding(ToolbarButtonPadding)
+				.Padding(ToolbarSlotPadding)
 				[
 					//The Restore from Immersive' button is only displayed when the editor is in Immersive mode.
 					SNew(SEditorViewportToolBarButton)
@@ -346,6 +346,17 @@ FText SLevelViewportToolBar::GetCameraMenuLabel() const
 	return LOCTEXT("CameraMenuTitle_Default", "Camera");
 }
 
+const FSlateBrush* SLevelViewportToolBar::GetCameraMenuLabelIcon() const
+{
+	TSharedPtr< SLevelViewport > PinnedViewport( Viewport.Pin() );
+	if( PinnedViewport.IsValid() )
+	{
+		return GetCameraMenuLabelIconFromViewportType(PinnedViewport->GetLevelViewportClient().ViewportType );
+	}
+
+	return FStyleDefaults::GetNoBrush();
+}
+
 FText SLevelViewportToolBar::GetDevicePreviewMenuLabel() const
 {
 	FText Label = LOCTEXT("DevicePreviewMenuTitle_Default", "Preview");
diff --git a/Engine/Source/Editor/LevelEditor/Private/SLevelViewportToolBar.h b/Engine/Source/Editor/LevelEditor/Private/SLevelViewportToolBar.h
--- a/Engine/Source/Editor/LevelEditor/Private/SLevelViewportToolBar.h
+++ b/Engine/Source/Editor/LevelEditor/Private/SLevelViewportToolBar.h
@@ -44,6 +44,12 @@ private:
 	 */
 	FText GetCameraMenuLabel() const;
 
+	/* Returns the label icon for the Camera tool bar menu, which changes depending on viewport type 
+	 *
+	 * @return	Label icon to use for this menu label
+	 *
+	 */
+	const FSlateBrush* GetCameraMenuLabelIcon() const;
 
 	/**
 	 * Returns the label for the "View" tool bar menu, which changes depending on viewport show flags
diff --git a/Engine/Source/Editor/Persona/Private/SAnimViewportToolBar.cpp b/Engine/Source/Editor/Persona/Private/SAnimViewportToolBar.cpp
--- a/Engine/Source/Editor/Persona/Private/SAnimViewportToolBar.cpp
+++ b/Engine/Source/Editor/Persona/Private/SAnimViewportToolBar.cpp
@@ -180,7 +180,7 @@ void SAnimViewportToolBar::Construct(const FArguments& InArgs, TSharedPtr<class
 	}
 
 	const FMargin ToolbarSlotPadding(4.0f, 1.0f);
-	const FMargin ToolbarButtonPadding(4.0f, 0.0f);
+	const FMargin ToolbarButtonPadding(4.0f, 1.0f);
 
 	TSharedRef<SHorizontalBox> LeftToolbar = SNew(SHorizontalBox)
 		+SHorizontalBox::Slot()
@@ -204,6 +204,7 @@ void SAnimViewportToolBar::Construct(const FArguments& InArgs, TSharedPtr<class
 			.ParentToolBar(SharedThis(this))
 			.Cursor(EMouseCursor::Default)
 			.Label(this, &SAnimViewportToolBar::GetCameraMenuLabel)
+			.LabelIcon( this, &SAnimViewportToolBar::GetCameraMenuLabelIcon )
 			.AddMetaData<FTagMetaData>(FTagMetaData(TEXT("EditorViewportToolBar.CameraMenu")))
 			.OnGetMenuContent(this, &SAnimViewportToolBar::GenerateViewportTypeMenu)
 		]
@@ -1066,6 +1067,17 @@ FText SAnimViewportToolBar::GetCameraMenuLabel() const
 	return LOCTEXT("Viewport_Default", "Camera");
 }
 
+const FSlateBrush* SAnimViewportToolBar::GetCameraMenuLabelIcon() const
+{
+	TSharedPtr< SAnimationEditorViewportTabBody > PinnedViewport( Viewport.Pin() );
+	if( PinnedViewport.IsValid() )
+	{
+		return GetCameraMenuLabelIconFromViewportType(PinnedViewport->GetLevelViewportClient().ViewportType );
+	}
+
+	return FAppStyle::Get().GetBrush("NoBrush");
+}
+
 TOptional<float> SAnimViewportToolBar::OnGetFOVValue() const
 {
 	if(Viewport.IsValid())
diff --git a/Engine/Source/Editor/Persona/Private/SAnimViewportToolBar.h b/Engine/Source/Editor/Persona/Private/SAnimViewportToolBar.h
--- a/Engine/Source/Editor/Persona/Private/SAnimViewportToolBar.h
+++ b/Engine/Source/Editor/Persona/Private/SAnimViewportToolBar.h
@@ -133,6 +133,13 @@ private:
 	 */
 	FText GetCameraMenuLabel() const;
 
+	/* Returns the label icon for the Camera tool bar menu, which changes depending on viewport type 
+	 *
+	 * @return	Label icon to use for this menu label
+	 *
+	 */
+	const FSlateBrush* GetCameraMenuLabelIcon() const;
+
 	/** Called by the FOV slider in the perspective viewport to get the FOV value */
 	TOptional<float> OnGetFOVValue() const;
 	/** Called when the FOV slider is adjusted in the perspective viewport */
diff --git a/Engine/Source/Editor/UnrealEd/Private/SEditorViewportToolBarButton.cpp b/Engine/Source/Editor/UnrealEd/Private/SEditorViewportToolBarButton.cpp
--- a/Engine/Source/Editor/UnrealEd/Private/SEditorViewportToolBarButton.cpp
+++ b/Engine/Source/Editor/UnrealEd/Private/SEditorViewportToolBarButton.cpp
@@ -35,6 +35,7 @@ void SEditorViewportToolBarButton::Construct( const FArguments& Declaration)
 			.OnClicked( OnClickedDelegate )
 			.HAlign( HAlign_Center )
 			.VAlign( VAlign_Center )
+			.ContentPadding(FMargin(4.f, 0.f))
 			[
 				// If we have a content override use it instead of the default image
 				bContentOverride
@@ -63,7 +64,7 @@ void SEditorViewportToolBarButton::Construct( const FArguments& Declaration)
 				bContentOverride ? ContentSlotWidget :
 				TSharedRef<SWidget>(
 					SNew( SBox )
-					.Padding(1.0f)
+					.Padding(0.0f)
 					.VAlign( VAlign_Center )
 					.HAlign( HAlign_Center )
 					[
diff --git a/Engine/Source/Editor/UnrealEd/Private/SEditorViewportToolBarMenu.cpp b/Engine/Source/Editor/UnrealEd/Private/SEditorViewportToolBarMenu.cpp
--- a/Engine/Source/Editor/UnrealEd/Private/SEditorViewportToolBarMenu.cpp
+++ b/Engine/Source/Editor/UnrealEd/Private/SEditorViewportToolBarMenu.cpp
@@ -48,7 +48,7 @@ void SEditorViewportToolbarMenu::Construct( const FArguments& Declaration )
 				.VAlign(VAlign_Center)
 				.HAlign(HAlign_Center)
 				.AutoWidth()
-				.Padding(5.0f, 0.0f)
+				.Padding(4.0f, 0.0f, 4.f, 0.0f)
 				[
 					SNew( SImage )
 					.Visibility(this, &SEditorViewportToolbarMenu::GetLabelIconVisibility)
@@ -56,7 +56,7 @@ void SEditorViewportToolbarMenu::Construct( const FArguments& Declaration )
 					.ColorAndOpacity(FSlateColor::UseForeground())
 				]
 				+SHorizontalBox::Slot()
-				.Padding(2.0f,1.0f, 0.0f, 0.0f)
+				.Padding(0.0f, 0.0f, 4.f, 0.0f)
 				.AutoWidth()
 				.VAlign(VAlign_Center)
 				[
@@ -71,6 +71,7 @@ void SEditorViewportToolbarMenu::Construct( const FArguments& Declaration )
 			ButtonContent = 
 				SNew(SBox)
 				.VAlign(VAlign_Center)
+				.Padding(FMargin(4.f, 0.f))
 				[
 					SNew( STextBlock )
 					.TextStyle(&FAppStyle::Get().GetWidgetStyle<FTextBlockStyle>("SmallText"))
@@ -87,7 +88,7 @@ void SEditorViewportToolbarMenu::Construct( const FArguments& Declaration )
 			SNew( SButton )
 			// Allows users to drag with the mouse to select options after opening the menu */
 			.ClickMethod( EButtonClickMethod::MouseDown )
-			.ContentPadding( FMargin( 2.0f,0.0f ) )
+			.ContentPadding( FMargin(0.0f, 0.0f) )
 			.VAlign(VAlign_Center)
 			.ButtonStyle(Declaration._MenuStyle)
 			.OnClicked( this, &SEditorViewportToolbarMenu::OnMenuClicked )
diff --git a/Engine/Source/Editor/UnrealEd/Private/SEditorViewportViewMenu.cpp b/Engine/Source/Editor/UnrealEd/Private/SEditorViewportViewMenu.cpp
--- a/Engine/Source/Editor/UnrealEd/Private/SEditorViewportViewMenu.cpp
+++ b/Engine/Source/Editor/UnrealEd/Private/SEditorViewportViewMenu.cpp
@@ -26,6 +26,7 @@ void SEditorViewportViewMenu::Construct( const FArguments& InArgs, TSharedRef<SE
 			.ParentToolBar( InParentToolBar)
 			.Cursor( EMouseCursor::Default )
 			.Label(this, &SEditorViewportViewMenu::GetViewMenuLabel)
+			.LabelIcon(this, &SEditorViewportViewMenu::GetViewMenuLabelIcon)
 			.OnGetMenuContent( this, &SEditorViewportViewMenu::GenerateViewMenuContent )
 	);
 }
@@ -54,6 +55,22 @@ FText SEditorViewportViewMenu::GetViewMenuLabel() const
 	return Label;
 }
 
+const FSlateBrush* SEditorViewportViewMenu::GetViewMenuLabelIcon() const
+{
+
+	TSharedPtr< SEditorViewport > PinnedViewport = Viewport.Pin();
+	if( PinnedViewport.IsValid() )
+	{
+		const TSharedPtr<FEditorViewportClient> ViewportClient = PinnedViewport->GetViewportClient();
+		check(ViewportClient.IsValid());
+		const EViewModeIndex ViewMode = ViewportClient->GetViewMode();
+
+		return UViewModeUtils::GetViewModeDisplayIcon(ViewMode);
+	}
+
+	return FStyleDefaults::GetNoBrush();
+}
+
 void SEditorViewportViewMenu::RegisterMenus() const
 {
 	if (!UToolMenus::Get()->IsMenuRegistered(BaseMenuName))
diff --git a/Engine/Source/Editor/UnrealEd/Private/STransformViewportToolbar.cpp b/Engine/Source/Editor/UnrealEd/Private/STransformViewportToolbar.cpp
--- a/Engine/Source/Editor/UnrealEd/Private/STransformViewportToolbar.cpp
+++ b/Engine/Source/Editor/UnrealEd/Private/STransformViewportToolbar.cpp
@@ -178,6 +178,7 @@ TSharedRef< SWidget > STransformViewportToolBar::MakeTransformToolBar( const TSh
 
 	ToolbarBuilder.BeginSection("Transform");
 	{
+		ToolbarBuilder.BeginBlockGroup();
 		// Move Mode
 		static FName TranslateModeName = FName(TEXT("TranslateMode"));
 		ToolbarBuilder.AddToolBarButton( FEditorViewportCommands::Get().TranslateMode, NAME_None, TAttribute<FText>(), TAttribute<FText>(), TAttribute<FSlateIcon>(), TranslateModeName );
@@ -199,6 +200,7 @@ TSharedRef< SWidget > STransformViewportToolBar::MakeTransformToolBar( const TSh
 		ToolbarBuilder.AddToolBarButton( FEditorViewportCommands::Get().ScaleMode, NAME_None, TAttribute<FText>(), TAttribute<FText>(), TAttribute<FSlateIcon>(), ScaleModeName );
 
 
+		ToolbarBuilder.EndBlockGroup();
 		ToolbarBuilder.AddSeparator();
 
 		ToolbarBuilder.SetIsFocusable( true );
@@ -212,9 +214,6 @@ TSharedRef< SWidget > STransformViewportToolBar::MakeTransformToolBar( const TSh
 										);
 	}
 
-	// Add a bit of spacing before the snapping section
-	ToolbarBuilder.AddWidget(SNew(SSpacer).Size(FVector2D(15.0f, 1.0f)));
-
 	ToolbarBuilder.EndSection();
 
 	ToolbarBuilder.BeginSection("LocationGridSnap");
diff --git a/Engine/Source/Editor/UnrealEd/Private/SViewportToolBarComboMenu.cpp b/Engine/Source/Editor/UnrealEd/Private/SViewportToolBarComboMenu.cpp
--- a/Engine/Source/Editor/UnrealEd/Private/SViewportToolBarComboMenu.cpp
+++ b/Engine/Source/Editor/UnrealEd/Private/SViewportToolBarComboMenu.cpp
@@ -17,9 +17,9 @@ void SViewportToolBarComboMenu::Construct( const FArguments& InArgs )
 
 	EMultiBlockLocation::Type BlockLocation = InArgs._BlockLocation;
 
-	const FButtonStyle& ButtonStyle = ToolBarStyle.ButtonStyle;
-	const FCheckBoxStyle& CheckBoxStyle = ToolBarStyle.ToggleButton;
-	const FTextBlockStyle& LabelStyle = ToolBarStyle.LabelStyle;
+	const FButtonStyle& ButtonStyle = FAppStyle::Get().GetWidgetStyle<FButtonStyle>("EditorViewportToolBar.ComboMenu.ButtonStyle");
+	const FCheckBoxStyle& CheckBoxStyle = FAppStyle::Get().GetWidgetStyle<FCheckBoxStyle>("EditorViewportToolBar.ComboMenu.ToggleButton");
+	const FTextBlockStyle& LabelStyle = FAppStyle::Get().GetWidgetStyle<FTextBlockStyle>("EditorViewportToolBar.ComboMenu.LabelStyle");
 
 	const FSlateIcon& Icon = InArgs._Icon.Get();
 
@@ -28,7 +28,7 @@ void SViewportToolBarComboMenu::Construct( const FArguments& InArgs )
 	TSharedPtr<SCheckBox> ToggleControl;
 	{
 		ToggleControl = SNew(SCheckBox)
-		.Padding(ToolBarStyle.CheckBoxPadding)
+		.Padding(0.f)
 		.Style(&CheckBoxStyle)
 		.OnCheckStateChanged(InArgs._OnCheckStateChanged)
 		.ToolTipText(InArgs._ToggleButtonToolTip)
@@ -49,7 +49,7 @@ void SViewportToolBarComboMenu::Construct( const FArguments& InArgs )
 		TSharedRef<SWidget> ButtonContents =
 			SNew(SButton)
 			.ButtonStyle(&ButtonStyle)
-			.ContentPadding( FMargin( 2.0f, 0.0f ) )
+			.ContentPadding(0.f)
 			.ToolTipText(InArgs._MenuButtonToolTip)
 			.OnClicked(this, &SViewportToolBarComboMenu::OnMenuClicked)
 			[
@@ -79,18 +79,26 @@ void SViewportToolBarComboMenu::Construct( const FArguments& InArgs )
 
 	ChildSlot
 	[
-		SNew(SHorizontalBox)
-		+SHorizontalBox::Slot()
-		.VAlign(VAlign_Center)
-		.AutoWidth()
-		[
-			ToggleControl.ToSharedRef()
-		]
-		+SHorizontalBox::Slot()
-		.VAlign(VAlign_Center)
-		.AutoWidth()
+
+		SNew(SBorder)
+		.Padding(FMargin(6.f, 0.f, 6.f, 0.f))
+		.BorderImage(FAppStyle::Get().GetBrush("EditorViewportToolBar.Group"))
 		[
-			MenuAnchor.ToSharedRef()
+			SNew(SHorizontalBox)
+
+			+SHorizontalBox::Slot()
+			.VAlign(VAlign_Center)
+			.AutoWidth()
+			[
+				ToggleControl.ToSharedRef()
+			]
+
+			+SHorizontalBox::Slot()
+			.VAlign(VAlign_Center)
+			.AutoWidth()
+			[
+				MenuAnchor.ToSharedRef()
+			]
 		]
 	];
 }
diff --git a/Engine/Source/Editor/UnrealEd/Public/SEditorViewportViewMenu.h b/Engine/Source/Editor/UnrealEd/Public/SEditorViewportViewMenu.h
--- a/Engine/Source/Editor/UnrealEd/Public/SEditorViewportViewMenu.h
+++ b/Engine/Source/Editor/UnrealEd/Public/SEditorViewportViewMenu.h
@@ -25,6 +25,7 @@ public:
 
 private:
 	FText GetViewMenuLabel() const;
+	const FSlateBrush* GetViewMenuLabelIcon() const;
 	void FillViewMenu(UToolMenu* Menu) const;
 
 protected:
diff --git a/Engine/Source/Runtime/Engine/Classes/Engine/EngineBaseTypes.h b/Engine/Source/Runtime/Engine/Classes/Engine/EngineBaseTypes.h
--- a/Engine/Source/Runtime/Engine/Classes/Engine/EngineBaseTypes.h
+++ b/Engine/Source/Runtime/Engine/Classes/Engine/EngineBaseTypes.h
@@ -14,9 +14,11 @@
 #include "UObject/WeakObjectPtr.h"
 #include "Misc/CoreMisc.h"
 #include "Async/TaskGraphInterfaces.h"
+
 #include "EngineBaseTypes.generated.h"
 
 class UActorComponent;
+struct FSlateBrush;
 struct FTickContext;
 
 //
@@ -990,6 +992,7 @@ public:
 	 * Get the display name associated with a particular EViewModeIndex
 	 */
 	static FText GetViewModeDisplayName(const EViewModeIndex ViewModeIndex);
+	static const FSlateBrush* GetViewModeDisplayIcon(const EViewModeIndex ViewModeIndex);
 };
 
 
diff --git a/Engine/Source/Runtime/Engine/Private/ViewModeNames.cpp b/Engine/Source/Runtime/Engine/Private/ViewModeNames.cpp
--- a/Engine/Source/Runtime/Engine/Private/ViewModeNames.cpp
+++ b/Engine/Source/Runtime/Engine/Private/ViewModeNames.cpp
@@ -176,4 +176,178 @@ FText UViewModeUtils::GetViewModeDisplayName(const EViewModeIndex ViewModeIndex)
 	return ViewModeName;
 }
 
+// Icons 
+TArray<const FSlateBrush*> FillViewModeDisplayIcons()
+{
+	const FSlateBrush* IconBrush = nullptr;
+
+	// Allocate size
+	TArray<const FSlateBrush*> ViewModeDisplayIcons;
+	ViewModeDisplayIcons.Reserve(VMI_Unknown);
+
+	// Fill ViewModeIndex, VMI_Unknown+1 to include VMI_Unknown too
+	for (int32 Index = 0; Index < VMI_Unknown + 1; ++Index)
+	{
+		const EViewModeIndex ViewModeIndex = (EViewModeIndex)Index;
+
+		if (ViewModeIndex == VMI_BrushWireframe)
+		{
+			ViewModeDisplayIcons.Emplace(FAppStyle::Get().GetBrush("EditorViewport.WireframeMode"));
+		}
+		// Wireframe w/ BSP
+		else if (ViewModeIndex == VMI_Wireframe)
+		{
+			ViewModeDisplayIcons.Emplace(FAppStyle::Get().GetBrush("EditorViewport.WireframeMode"));
+		}
+		// Unlit
+		else if (ViewModeIndex == VMI_Unlit)
+		{
+			ViewModeDisplayIcons.Emplace(FAppStyle::Get().GetBrush("EditorViewport.UnlitMode"));
+		}
+		// Lit
+		else if (ViewModeIndex == VMI_Lit)
+		{
+			ViewModeDisplayIcons.Emplace(FAppStyle::Get().GetBrush("EditorViewport.LitMode"));
+		}
+		else if (ViewModeIndex == VMI_Lit_DetailLighting)
+		{
+			ViewModeDisplayIcons.Emplace(FAppStyle::Get().GetBrush("EditorViewport.DetailLightingMode"));
+		}
+		// Lit wo/ materials
+		else if (ViewModeIndex == VMI_LightingOnly)
+		{
+			ViewModeDisplayIcons.Emplace(FAppStyle::Get().GetBrush("EditorViewport.LightingOnlyMode"));
+		}
+		// Colored according to light count
+		else if (ViewModeIndex == VMI_LightComplexity)
+		{
+			ViewModeDisplayIcons.Emplace(FAppStyle::Get().GetBrush("EditorViewport.LightComplexityMode"));
+		}
+		// Colored according to shader complexity
+		else if (ViewModeIndex == VMI_ShaderComplexity)
+		{
+			ViewModeDisplayIcons.Emplace(FAppStyle::Get().GetBrush("EditorViewport.ShaderComplexityMode"));
+		}
+		// Colored according to world-space LightMap texture density
+		else if (ViewModeIndex == VMI_LightmapDensity)
+		{
+			ViewModeDisplayIcons.Emplace(FAppStyle::Get().GetBrush("EditorViewport.LightmapDensityMode"));
+		}
+		// Colored according to light count - showing lightmap texel density on texture mapped objects
+		else if (ViewModeIndex == VMI_LitLightmapDensity)
+		{
+			ViewModeDisplayIcons.Emplace(FAppStyle::Get().GetBrush("EditorViewport.LitLightmapDensityMode"));
+		}
+		else if (ViewModeIndex == VMI_ReflectionOverride)
+		{
+			ViewModeDisplayIcons.Emplace(FAppStyle::Get().GetBrush("EditorViewport.ReflectionOverrideMode"));
+		}
+		else if (ViewModeIndex == VMI_VisualizeBuffer)
+		{
+			ViewModeDisplayIcons.Emplace(FAppStyle::Get().GetBrush("EditorViewport.VisualizeBufferMode"));
+		}
+		//	VMI_VoxelLighting = 13,
+	
+		// Colored according to stationary light overlap
+		else if (ViewModeIndex == VMI_StationaryLightOverlap)
+		{
+			ViewModeDisplayIcons.Emplace(FAppStyle::Get().GetBrush("EditorViewport.StationaryLightOverlapMode"));
+		}
+	
+		else if (ViewModeIndex == VMI_CollisionPawn)
+		{
+			ViewModeDisplayIcons.Emplace(FAppStyle::Get().GetBrush("EditorViewport.CollisionPawn"));
+		}
+		else if (ViewModeIndex == VMI_CollisionVisibility)
+		{
+			ViewModeDisplayIcons.Emplace(FAppStyle::Get().GetBrush("EditorViewport.CollisionVisibility"));
+		}
+		//VMI_UNUSED = 17,
+		// Colored according to the current LOD index
+		else if (ViewModeIndex == VMI_LODColoration)
+		{
+			ViewModeDisplayIcons.Emplace(FAppStyle::Get().GetBrush("EditorViewport.LOD"));
+		}
+		// Colored according to the quad coverage
+		else if (ViewModeIndex == VMI_QuadOverdraw)
+		{
+			ViewModeDisplayIcons.Emplace(FAppStyle::Get().GetBrush("EditorViewport.QuadOverdrawMode"));
+		}
+		// Visualize the accuracy of the primitive distance computed for texture streaming
+		else if (ViewModeIndex == VMI_PrimitiveDistanceAccuracy)
+		{
+			ViewModeDisplayIcons.Emplace(FAppStyle::Get().GetBrush("EditorViewport.TexStreamAccPrimitiveDistanceMode"));
+		}
+		// Visualize the accuracy of the mesh UV densities computed for texture streaming
+		else if (ViewModeIndex == VMI_MeshUVDensityAccuracy)
+		{
+			ViewModeDisplayIcons.Emplace(FAppStyle::Get().GetBrush("EditorViewport.TexStreamAccMeshUVDensityMode"));
+		}
+		// Colored according to shader complexity, including quad overdraw
+		else if (ViewModeIndex == VMI_ShaderComplexityWithQuadOverdraw)
+		{
+			ViewModeDisplayIcons.Emplace(FAppStyle::Get().GetBrush("EditorViewport.ShaderComplexityWithQuadOverdrawMode"));
+		}
+		// Colored according to the current HLOD index
+		else if (ViewModeIndex == VMI_HLODColoration)
+		{
+			ViewModeDisplayIcons.Emplace(FAppStyle::Get().GetBrush("EditorViewport.LOD"));
+		}
+		// Group item for LOD and HLOD coloration
+		else if (ViewModeIndex == VMI_GroupLODColoration)
+		{
+			ViewModeDisplayIcons.Emplace(FAppStyle::Get().GetBrush("EditorViewport.LOD"));
+		}
+		// Visualize the accuracy of the material texture scales used for texture streaming
+		else if (ViewModeIndex == VMI_MaterialTextureScaleAccuracy)
+		{
+			ViewModeDisplayIcons.Emplace(FAppStyle::Get().GetBrush("EditorViewport.TexStreamAccMaterialTextureScaleMode"));
+		}
+		// Compare the required texture resolution to the actual resolution
+		else if (ViewModeIndex == VMI_RequiredTextureResolution)
+		{
+			ViewModeDisplayIcons.Emplace(FAppStyle::Get().GetBrush("EditorViewport.RequiredTextureResolutionMode"));
+		}
+	
+		// Ray tracing modes
+		// Run path tracing pipeline
+		else if (ViewModeIndex == VMI_PathTracing)
+		{
+			ViewModeDisplayIcons.Emplace(FAppStyle::Get().GetBrush("EditorViewport.PathTracingMode"));
+		}
+		// Run ray tracing debug pipeline 
+		else if (ViewModeIndex == VMI_RayTracingDebug)
+		{
+			ViewModeDisplayIcons.Emplace(FAppStyle::Get().GetBrush("EditorViewport.RayTracingDebugMode"));
+		}
+
+		// VMI_Max
+		else if (ViewModeIndex == VMI_Max)
+		{
+			ViewModeDisplayIcons.Emplace(FAppStyle::Get().GetBrush("EditorViewport.MaxMode"));
+		}
+		// VMI_Unknown
+		else 
+		{
+			ViewModeDisplayIcons.Emplace(FStyleDefaults::GetNoBrush());
+		}
+	}
+
+	return ViewModeDisplayIcons;
+}
+
+static TArray<const FSlateBrush*> GViewModeDisplayIcons;
+
+const FSlateBrush* UViewModeUtils::GetViewModeDisplayIcon(const EViewModeIndex ViewModeIndex) 
+{
+	static bool IconsInitialized = false;
+	if (!IconsInitialized)
+	{
+		GViewModeDisplayIcons = FillViewModeDisplayIcons();
+		IconsInitialized = true;
+	}
+
+	return GViewModeDisplayIcons[ViewModeIndex];
+}
+
 #undef LOCTEXT_NAMESPACE
