Advanced dropdown border color now the same as the background color.

diff --git a/Engine/Source/Editor/EditorStyle/Private/StarshipStyle.cpp b/Engine/Source/Editor/EditorStyle/Private/StarshipStyle.cpp
--- a/Engine/Source/Editor/EditorStyle/Private/StarshipStyle.cpp
+++ b/Engine/Source/Editor/EditorStyle/Private/StarshipStyle.cpp
@@ -3186,9 +3186,8 @@ void FStarshipEditorStyle::FStyle::SetupPropertyEditorStyles()
 			.SetTextStyle(BPWarningMessageHyperlinkTextStyle)
 			.SetPadding(FMargin(0.0f));
 
-		Set("DetailsView.BPMessageHyperlinkStyle", BPWarningMessageHyperlinkStyle);
-		Set("DetailsView.BPMessageTextStyle", BPWarningMessageTextStyle);
-
+		Set( "DetailsView.BPMessageHyperlinkStyle", BPWarningMessageHyperlinkStyle );
+		Set( "DetailsView.BPMessageTextStyle", BPWarningMessageTextStyle );
 
 		Set( "DetailsView.GroupSection",              new FSlateNoResource());
 
@@ -3206,6 +3205,8 @@ void FStarshipEditorStyle::FStyle::SetupPropertyEditorStyles()
 		Set( "DetailsView.CategoryTop",               new FSlateColorBrush(FStyleColors::Header));
 		Set( "DetailsView.CategoryTop_Hovered",       new FSlateColorBrush(FStyleColors::Header));
 		Set( "DetailsView.CategoryBottom",            new FSlateColorBrush(FStyleColors::Recessed));
+		
+		// these are not actually displayed as white, see PropertyEditorConstants::GetRowBackgroundColor
 		Set( "DetailsView.CategoryMiddle",            new FSlateColorBrush(FStyleColors::White));
 
 		Set( "DetailsView.PropertyIsFavorite", new IMAGE_BRUSH("PropertyView/Favorites_Enabled", Icon12x12));
@@ -3222,10 +3223,10 @@ void FStarshipEditorStyle::FStyle::SetupPropertyEditorStyles()
 			.SetHandleHighlightBrush(FSlateColorBrush(FStyleColors::Recessed))
 		);
 
-		Set("DetailsView.GridLine", new FSlateColorBrush(FStyleColors::Recessed));
+		Set( "DetailsView.GridLine", new FSlateColorBrush(FStyleColors::Recessed) );
 
-		Set( "DetailsView.AdvancedDropdownBorder",      new FSlateColorBrush(FStyleColors::Recessed));
-		Set( "DetailsView.AdvancedDropdownBorder.Open", new FSlateColorBrush(FStyleColors::Recessed));
+		Set( "DetailsView.AdvancedDropdownBorder",      new FSlateColorBrush(FStyleColors::Panel));
+		Set( "DetailsView.AdvancedDropdownBorder.Open", new FSlateColorBrush(FStyleColors::Panel));
 
 		Set( "DetailsView.CategoryFontStyle", FStyleFonts::Get().NormalBold);
 		Set( "DetailsView.CategoryTextStyle", FTextBlockStyle(NormalText)
