Updated editor icons

diff --git a/Engine/Source/Editor/EditorStyle/Private/StarshipStyle.cpp b/Engine/Source/Editor/EditorStyle/Private/StarshipStyle.cpp
--- a/Engine/Source/Editor/EditorStyle/Private/StarshipStyle.cpp
+++ b/Engine/Source/Editor/EditorStyle/Private/StarshipStyle.cpp
@@ -3996,14 +3996,23 @@ void FStarshipEditorStyle::FStyle::SetupGraphEditorStyles()
 		Set("ClassIcon.EdGraphNode_Comment",	new IMAGE_BRUSH("Icons/icon_Blueprint_Comment_16x", Icon16x16));
 
 
+		Set("GraphEditor.OverrideFunction_16x", new IMAGE_BRUSH_SVG("Starship/Blueprints/icon_Blueprint_OverrideFunction", Icon16x16));
+		Set("GraphEditor.OverrideFunction_24x", new IMAGE_BRUSH_SVG("Starship/Blueprints/icon_Blueprint_OverrideFunction", Icon24x24));
+		Set("GraphEditor.EventGraph_16x", new IMAGE_BRUSH_SVG("Starship/Blueprints/icon_BlueprintEditor_EventGraph", Icon16x16));
+		Set("GraphEditor.EventGraph_24x", new IMAGE_BRUSH_SVG("Starship/Blueprints/icon_BlueprintEditor_EventGraph", Icon24x24));
+		Set("GraphEditor.Macro_16x", new IMAGE_BRUSH_SVG("Starship/Blueprints/icon_Blueprint_Macro", Icon16x16));
+		Set("GraphEditor.Macro_24x", new IMAGE_BRUSH_SVG("Starship/Blueprints/icon_Blueprint_Macro", Icon24x24));
+		Set("GraphEditor.Function_16x", new IMAGE_BRUSH_SVG("Starship/Blueprints/icon_Blueprint_Function", Icon16x16));
+		Set("GraphEditor.Function_24x", new IMAGE_BRUSH_SVG("Starship/Blueprints/icon_Blueprint_Function", Icon24x24));
+		Set("GraphEditor.Delegate_16x", new IMAGE_BRUSH_SVG("Starship/Blueprints/icon_Blueprint_Delegate", Icon16x16));
+		Set("GraphEditor.Delegate_24x", new IMAGE_BRUSH_SVG("Starship/Blueprints/icon_Blueprint_Delegate", Icon24x24));
+
+
+
 		Set( "GraphEditor.Default_16x", new IMAGE_BRUSH("Icons/icon_Blueprint_Node_16x", Icon16x16));
-		Set( "GraphEditor.EventGraph_16x", new IMAGE_BRUSH_SVG( "Starship/Blueprints/icon_BlueprintEditor_EventGraph", Icon16x16 ) );
 		Set( "GraphEditor.InterfaceFunction_16x", new IMAGE_BRUSH( "Icons/icon_Blueprint_Interfacefunction_16x", Icon16x16));
-		Set( "GraphEditor.Macro_16x", new IMAGE_BRUSH_SVG( "Starship/Blueprints/icon_Blueprint_Macro", Icon16x16));
-		Set( "GraphEditor.Function_16x", new IMAGE_BRUSH_SVG( "Starship/Blueprints/icon_Blueprint_Function", Icon16x16));
 		Set( "GraphEditor.PureFunction_16x", new IMAGE_BRUSH( "Icons/icon_Blueprint_NewPureFunction_16x", Icon16x16 ) );
 		Set( "GraphEditor.PotentialOverrideFunction_16x", new IMAGE_BRUSH( "Icons/icon_Blueprint_OverrideableFunction_16x", Icon16x16 ) );
-		Set( "GraphEditor.OverrideFunction_16x", new IMAGE_BRUSH( "Icons/icon_Blueprint_OverrideFunction_16x", Icon16x16 ) );
 		Set( "GraphEditor.OverridePureFunction_16x", new IMAGE_BRUSH( "Icons/icon_Blueprint_OverridePureFunction_16x", Icon16x16 ) );
 		Set( "GraphEditor.SubGraph_16x", new IMAGE_BRUSH( "Icons/icon_Blueprint_SubgraphComposite_16x", Icon16x16 ) );
 		Set( "GraphEditor.Animation_16x", new IMAGE_BRUSH( "Icons/icon_Blueprint_Anim_16x", Icon16x16 ) );
@@ -4042,15 +4051,12 @@ void FStarshipEditorStyle::FStyle::SetupGraphEditorStyles()
 		Set( "GraphEditor.Macro.FlipFlop_16x", new IMAGE_BRUSH("Icons/icon_Blueprint_FlipFlop_16x", Icon16x16));
 		Set( "GraphEditor.Macro.ForEach_16x", new IMAGE_BRUSH("Icons/icon_Blueprint_ForEach_16x", Icon16x16));
 
-		Set( "GraphEditor.Delegate_16x", new IMAGE_BRUSH( "Icons/icon_Blueprint_Delegate_16x", Icon16x16 ) );
-		Set( "GraphEditor.Delegate_24x", new IMAGE_BRUSH( "Icons/icon_Blueprint_Delegate_24x", Icon24x24 ) );
 
-		Set( "GraphEditor.EventGraph_24x", new IMAGE_BRUSH_SVG( "Starship/Blueprints/icon_BlueprintEditor_EventGraph", Icon24x24 ) );
+
 		Set( "GraphEditor.InterfaceFunction_24x", new IMAGE_BRUSH( "Icons/icon_Blueprint_InterfaceFunction_24x", Icon24x24 ) );
-		Set( "GraphEditor.Macro_24x", new IMAGE_BRUSH_SVG( "Starship/Blueprints/icon_Blueprint_Macro", Icon24x24 ) );
-		Set( "GraphEditor.Function_24x", new IMAGE_BRUSH_SVG( "Starship/Blueprints/icon_Blueprint_Function", Icon24x24 ) );
+
 		Set( "GraphEditor.PotentialOverrideFunction_24x", new IMAGE_BRUSH( "Icons/icon_Blueprint_OverrideableFunction_24x", Icon24x24 ) );
-		Set( "GraphEditor.OverrideFunction_24x", new IMAGE_BRUSH( "Icons/icon_Blueprint_OverrideFunction_24x", Icon24x24 ) );
+
 		Set( "GraphEditor.SubGraph_24x", new IMAGE_BRUSH( "Icons/icon_Blueprint_SubgraphComposite_24x", Icon24x24 ) );
 		Set( "GraphEditor.Animation_24x", new IMAGE_BRUSH( "Icons/icon_Blueprint_Anim_24x", Icon24x24 ) );
 		Set( "GraphEditor.Conduit_24x", new IMAGE_BRUSH( "Icons/icon_Blueprint_Conduit_24x", Icon24x24 ) );
diff --git a/Engine/Source/Runtime/SlateCore/Private/Styling/StarshipCoreStyle.cpp b/Engine/Source/Runtime/SlateCore/Private/Styling/StarshipCoreStyle.cpp
--- a/Engine/Source/Runtime/SlateCore/Private/Styling/StarshipCoreStyle.cpp
+++ b/Engine/Source/Runtime/SlateCore/Private/Styling/StarshipCoreStyle.cpp
@@ -278,7 +278,7 @@ TSharedRef<ISlateStyle> FStarshipCoreStyle::Create()
 		Style->Set("AppIcon", new IMAGE_BRUSH_SVG("Starship/Common/unreal", FVector2D(36, 36), FStyleColors::Foreground));
 		Style->Set("AppIconPadding", FMargin(11, 11, 3, 5));
 #else
-		Style->Set("AppIcon", new IMAGE_BRUSH("Starship/Common/uelogo_chrome_1", FVector2D(53, 53), FStyleColors::White));
+		Style->Set("AppIcon", new IMAGE_BRUSH("Starship/Common/UELogo", FVector2D(45, 45), FStyleColors::White));
 		Style->Set("AppIconPadding", FMargin(5, 5, 5, 5));
 #endif
 
