Added missing slate color for anim graph poses

diff --git a/Engine/Source/Editor/EditorStyle/Private/StarshipStyle.cpp b/Engine/Source/Editor/EditorStyle/Private/StarshipStyle.cpp
--- a/Engine/Source/Editor/EditorStyle/Private/StarshipStyle.cpp
+++ b/Engine/Source/Editor/EditorStyle/Private/StarshipStyle.cpp
@@ -5759,6 +5759,7 @@ void FStarshipEditorStyle::FStyle::SetupPersonaStyle()
 			.SetFont(DEFAULT_FONT("Regular", FStarshipCoreStyle::SmallTextSize))
 			.SetColorAndOpacity(AttributeTextColor));
 
+		Set("AnimGraph.Attribute.Pose.Color", FStyleColors::White);
 		Set("AnimGraph.Attribute.Curves.Icon", new IMAGE_BRUSH("GenericCurveEditor/Icons/GenericCurveEditor_48x", Icon10x10));
 		Set("AnimGraph.Attribute.Curves.Color", FStyleColors::AccentPurple);
 		Set("AnimGraph.Attribute.Attributes.Icon", new IMAGE_BRUSH("Icons/pillarray_16x", Icon10x10));
