Added icons for Path Tracing and Ray Tracing Debug view modes

diff --git a/Engine/Source/Editor/EditorStyle/Private/StarshipStyle.cpp b/Engine/Source/Editor/EditorStyle/Private/StarshipStyle.cpp
--- a/Engine/Source/Editor/EditorStyle/Private/StarshipStyle.cpp
+++ b/Engine/Source/Editor/EditorStyle/Private/StarshipStyle.cpp
@@ -4346,6 +4346,9 @@ void FStarshipEditorStyle::FStyle::SetupLevelEditorStyle()
 		Set( "EditorViewport.WireframeMode",      	  new IMAGE_BRUSH_SVG("Starship/Common/BrushWireframe", Icon16x16 ) );
 		Set( "EditorViewport.DetailLightingMode", 	  new IMAGE_BRUSH_SVG("Starship/Common/DetailLighting", Icon16x16 ) );
 		Set( "EditorViewport.LightingOnlyMode",   	  new IMAGE_BRUSH_SVG("Starship/Common/LightBulb", Icon16x16 ) );
+		
+		Set( "EditorViewport.PathTracingMode",   	  new IMAGE_BRUSH_SVG("Starship/Common/PathTracing", Icon16x16 ) );
+		Set( "EditorViewport.RayTracingDebugMode",    new IMAGE_BRUSH_SVG("Starship/Common/RayTracingDebug", Icon16x16 ) );
 
 		Set( "EditorViewport.LightComplexityMode", new IMAGE_BRUSH( "Icons/icon_ViewMode_LightComplexity_16px", Icon16x16 ) );
 		Set( "EditorViewport.ShaderComplexityMode", new IMAGE_BRUSH( "Icons/icon_ViewMode_Shadercomplexity_16px", Icon16x16 ) );
diff --git a/Engine/Source/Editor/UnrealEd/Private/SEditorViewportViewMenu.cpp b/Engine/Source/Editor/UnrealEd/Private/SEditorViewportViewMenu.cpp
--- a/Engine/Source/Editor/UnrealEd/Private/SEditorViewportViewMenu.cpp
+++ b/Engine/Source/Editor/UnrealEd/Private/SEditorViewportViewMenu.cpp
@@ -216,7 +216,7 @@ void SEditorViewportViewMenu::FillViewMenu(UToolMenu* Menu) const
 					}
 				};
 
-				Section.AddSubMenu("RayTracingDebugSubMenu", LOCTEXT("RayTracingDebugSubMenu", "Ray Tracing Debug"), LOCTEXT("RayTracing_ToolTip", "Select ray tracing buffer visualization view modes"), FNewMenuDelegate::CreateStatic(&Local::BuildRayTracingDebugMenu, ParentToolBar));
+				Section.AddSubMenu("RayTracingDebugSubMenu", LOCTEXT("RayTracingDebugSubMenu", "Ray Tracing Debug"), LOCTEXT("RayTracing_ToolTip", "Select ray tracing buffer visualization view modes"), FNewMenuDelegate::CreateStatic(&Local::BuildRayTracingDebugMenu, ParentToolBar), false, FSlateIcon(FAppStyle::GetAppStyleSetName(), "EditorViewport.RayTracingDebugMode"));
 			}
 #endif
 
