Static mesh editor icons and a few editors that shared the same icons

diff --git a/Engine/Plugins/Animation/LiveLink/Source/LiveLinkEditor/Private/LiveLinkClientPanelToolbar.cpp b/Engine/Plugins/Animation/LiveLink/Source/LiveLinkEditor/Private/LiveLinkClientPanelToolbar.cpp
--- a/Engine/Plugins/Animation/LiveLink/Source/LiveLinkEditor/Private/LiveLinkClientPanelToolbar.cpp
+++ b/Engine/Plugins/Animation/LiveLink/Source/LiveLinkEditor/Private/LiveLinkClientPanelToolbar.cpp
@@ -776,7 +776,7 @@ TSharedRef<SWidget> SLiveLinkClientPanelToolbar::OnPresetGeneratePresetsMenu()
 	MenuBuilder.AddMenuEntry(
 		LOCTEXT("SaveAsPreset_Text", "Save As Preset"),
 		LOCTEXT("SaveAsPreset_Tip", "Save the current setup as a new preset that can be imported at a later date"),
-		FSlateIcon(FEditorStyle::Get().GetStyleSetName(), "AssetEditor.SaveAsset.Greyscale"),
+		FSlateIcon(FEditorStyle::Get().GetStyleSetName(), "AssetEditor.SaveAsset"),
 		FUIAction(
 			FExecuteAction::CreateSP(this, &SLiveLinkClientPanelToolbar::OnSaveAsPreset)
 		)
diff --git a/Engine/Plugins/Developer/TraceSourceFiltering/Source/SourceFilteringEditor/Private/SourceFilterStyle.cpp b/Engine/Plugins/Developer/TraceSourceFiltering/Source/SourceFilteringEditor/Private/SourceFilterStyle.cpp
--- a/Engine/Plugins/Developer/TraceSourceFiltering/Source/SourceFilteringEditor/Private/SourceFilterStyle.cpp
+++ b/Engine/Plugins/Developer/TraceSourceFiltering/Source/SourceFilteringEditor/Private/SourceFilterStyle.cpp
@@ -4,6 +4,12 @@
 #include "Styling/SlateStyleRegistry.h"
 #include "Styling/SlateTypes.h"
 #include "Styling/CoreStyle.h"
+#include "Styling/StarshipCoreStyle.h"
+#include "SlateOptMacros.h"
+#include "Styling/SlateStyleMacros.h"
+
+// This is to fix the issue that SlateStyleMacros like IMAGE_BRUSH look for RootToContentDir but StyleSet->RootToContentDir is how this style is set up
+#define RootToContentDir StyleSet->RootToContentDir
 
 FSlateBrush* FSourceFilterStyle::FilterBrush = nullptr;
 FSlateBrush* FSourceFilterStyle::FilterSetBrush = nullptr;
@@ -12,11 +18,8 @@ TSharedPtr< FSlateStyleSet > FSourceFilterStyle::StyleSet = nullptr;
 
 FTextBlockStyle FSourceFilterStyle::NormalText;
 
-#define IMAGE_BRUSH( RelativePath, ... ) FSlateImageBrush( StyleSet->RootToContentDir( RelativePath, TEXT(".png") ), __VA_ARGS__ )
-#define BOX_BRUSH( RelativePath, ... ) FSlateBoxBrush( StyleSet->RootToContentDir( RelativePath, TEXT(".png") ), __VA_ARGS__ )
-#define BORDER_BRUSH( RelativePath, ... ) FSlateBorderBrush( StyleSet->RootToContentDir( RelativePath, TEXT(".png") ), __VA_ARGS__ )
-#define DEFAULT_FONT(...) FCoreStyle::GetDefaultFontStyle(__VA_ARGS__)
-#define ICON_FONT(...) FSlateFontInfo(StyleSet->RootToContentDir("Fonts/FontAwesome", TEXT(".ttf")), __VA_ARGS__)
+
+#define ICON_FONT(...) FSlateFontInfo(RootToContentDir("Fonts/FontAwesome", TEXT(".ttf")), __VA_ARGS__)
 
 // Const icon sizes
 static const FVector2D Icon8x8(8.0f, 8.0f);
@@ -150,7 +153,7 @@ void FSourceFilterStyle::Initialize()
 
 	StyleSet->Set("SourceFilter.NewPreset", new IMAGE_BRUSH("Icons/icon_file_saveas_16px", Icon16x16));
 	StyleSet->Set("SourceFilter.LoadPreset", new IMAGE_BRUSH("Icons/icon_file_open_16px", Icon16x16));
-	StyleSet->Set("SourceFilter.SavePreset", new IMAGE_BRUSH("Icons/icon_file_save_16px", Icon16x16));
+	StyleSet->Set("SourceFilter.SavePreset", new IMAGE_BRUSH_SVG("Starship/Common/SaveCurrent", Icon16x16));
 
 
 	StyleSet->Set("SourceFilter.TextStyle", FTextBlockStyle(NormalText)
diff --git a/Engine/Plugins/MovieScene/MovieRenderPipeline/Source/MovieRenderPipelineEditor/Private/Widgets/SMoviePipelineConfigPanel.cpp b/Engine/Plugins/MovieScene/MovieRenderPipeline/Source/MovieRenderPipelineEditor/Private/Widgets/SMoviePipelineConfigPanel.cpp
--- a/Engine/Plugins/MovieScene/MovieRenderPipeline/Source/MovieRenderPipelineEditor/Private/Widgets/SMoviePipelineConfigPanel.cpp
+++ b/Engine/Plugins/MovieScene/MovieRenderPipeline/Source/MovieRenderPipelineEditor/Private/Widgets/SMoviePipelineConfigPanel.cpp
@@ -343,7 +343,7 @@ TSharedRef<SWidget> SMoviePipelineConfigPanel::OnGeneratePresetsMenu()
 	MenuBuilder.AddMenuEntry(
 		LOCTEXT("SaveAsPreset_Text", "Save As Preset"),
 		LOCTEXT("SaveAsPreset_Tip", "Save the current configuration as a new preset that can be shared between multiple jobs, or imported later as the base of a new configuration."),
-		FSlateIcon(FEditorStyle::Get().GetStyleSetName(), "AssetEditor.SaveAsset.Greyscale"),
+		FSlateIcon(FEditorStyle::Get().GetStyleSetName(), "AssetEditor.SaveAsset"),
 		FUIAction(FExecuteAction::CreateSP(this, &SMoviePipelineConfigPanel::OnSaveAsPreset))
 	);
 
diff --git a/Engine/Plugins/MovieScene/MovieRenderPipeline/Source/MovieRenderPipelineEditor/Private/Widgets/SMoviePipelineQueuePanel.cpp b/Engine/Plugins/MovieScene/MovieRenderPipeline/Source/MovieRenderPipelineEditor/Private/Widgets/SMoviePipelineQueuePanel.cpp
--- a/Engine/Plugins/MovieScene/MovieRenderPipeline/Source/MovieRenderPipelineEditor/Private/Widgets/SMoviePipelineQueuePanel.cpp
+++ b/Engine/Plugins/MovieScene/MovieRenderPipeline/Source/MovieRenderPipelineEditor/Private/Widgets/SMoviePipelineQueuePanel.cpp
@@ -493,7 +493,7 @@ TSharedRef<SWidget> SMoviePipelineQueuePanel::OnGenerateSavedQueuesMenu()
 	MenuBuilder.AddMenuEntry(
 		LOCTEXT("SaveAsQueue_Text", "Save As Asset"),
 		LOCTEXT("SaveAsQueue_Tip", "Save the current configuration as a new preset that can be shared between multiple jobs, or imported later as the base of a new configuration."),
-		FSlateIcon(FEditorStyle::Get().GetStyleSetName(), "AssetEditor.SaveAsset.Greyscale"),
+		FSlateIcon(FEditorStyle::Get().GetStyleSetName(), "AssetEditor.SaveAsset"),
 		FUIAction(FExecuteAction::CreateSP(this, &SMoviePipelineQueuePanel::OnSaveAsAsset))
 	);
 
diff --git a/Engine/Plugins/Runtime/HairStrands/Source/HairStrandsEditor/Private/GroomAssetDetails.cpp b/Engine/Plugins/Runtime/HairStrands/Source/HairStrandsEditor/Private/GroomAssetDetails.cpp
--- a/Engine/Plugins/Runtime/HairStrands/Source/HairStrandsEditor/Private/GroomAssetDetails.cpp
+++ b/Engine/Plugins/Runtime/HairStrands/Source/HairStrandsEditor/Private/GroomAssetDetails.cpp
@@ -1004,7 +1004,7 @@ void FGroomRenderingDetails::OnGenerateElementForHairGroup(TSharedRef<IPropertyH
 							.OnClicked(this, &FGroomRenderingDetails::OnSaveCards, GroupIndex, Property)
 							[
 								SNew(SImage)
-								.Image(FEditorStyle::GetBrush("AssetEditor.SaveAsset.Greyscale"))
+								.Image(FEditorStyle::GetBrush("AssetEditor.SaveAsset"))
 							]
 						]
 						+ SHorizontalBox::Slot()
diff --git a/Engine/Plugins/VirtualProduction/Takes/Source/TakeRecorder/Private/Widgets/STakeRecorderPanel.cpp b/Engine/Plugins/VirtualProduction/Takes/Source/TakeRecorder/Private/Widgets/STakeRecorderPanel.cpp
--- a/Engine/Plugins/VirtualProduction/Takes/Source/TakeRecorder/Private/Widgets/STakeRecorderPanel.cpp
+++ b/Engine/Plugins/VirtualProduction/Takes/Source/TakeRecorder/Private/Widgets/STakeRecorderPanel.cpp
@@ -556,7 +556,7 @@ TSharedRef<SWidget> STakeRecorderPanel::OnGeneratePresetsMenu()
 	MenuBuilder.AddMenuEntry(
 		LOCTEXT("SaveAsPreset_Text", "Save As Preset"),
 		LOCTEXT("SaveAsPreset_Tip", "Save the current setup as a new preset that can be imported at a later date"),
-		FSlateIcon(FEditorStyle::Get().GetStyleSetName(), "AssetEditor.SaveAsset.Greyscale"),
+		FSlateIcon(FEditorStyle::Get().GetStyleSetName(), "AssetEditor.SaveAsset"),
 		FUIAction(
 			FExecuteAction::CreateSP(this, &STakeRecorderPanel::OnSaveAsPreset)
 		)
diff --git a/Engine/Source/Editor/DetailCustomizations/Private/MobilityCustomization.cpp b/Engine/Source/Editor/DetailCustomizations/Private/MobilityCustomization.cpp
--- a/Engine/Source/Editor/DetailCustomizations/Private/MobilityCustomization.cpp
+++ b/Engine/Source/Editor/DetailCustomizations/Private/MobilityCustomization.cpp
@@ -53,9 +53,24 @@ void FMobilityCustomization::CreateMobilityCustomization(IDetailCategoryBuilder&
 		// Static Mobility
 		ButtonOptionsPanel->AddSlot(EComponentMobility::Static)
 		[
-			SNew(STextBlock)
-			.Font(FAppStyle::Get().GetFontStyle("PropertyWindow.MobilityFont"))
-			.Text(LOCTEXT("Static", "Static"))
+			SNew(SHorizontalBox)
+			+ SHorizontalBox::Slot()
+			.AutoWidth()
+			.VAlign(VAlign_Center)
+			.Padding(0.0f, 0.0f, 3.0f, 0.0f)
+			[
+				SNew(SImage)
+				.Image(FAppStyle::Get().GetBrush("PropertyWindow.MobilityStatic"))
+				.ColorAndOpacity(FSlateColor::UseForeground())
+			]
+			+ SHorizontalBox::Slot()
+			.AutoWidth()
+			.VAlign(VAlign_Center)
+			[
+				SNew(STextBlock)
+				.Font(FAppStyle::Get().GetFontStyle("PropertyWindow.MobilityFont"))
+				.Text(LOCTEXT("Static", "Static"))
+			]
 		]
 		.ToolTip(StaticTooltip);
 	}
@@ -69,9 +84,24 @@ void FMobilityCustomization::CreateMobilityCustomization(IDetailCategoryBuilder&
 
 		ButtonOptionsPanel->AddSlot(EComponentMobility::Stationary)
 		[
-			SNew(STextBlock)
-			.Font(FAppStyle::Get().GetFontStyle("PropertyWindow.MobilityFont"))
-			.Text(LOCTEXT("Stationary", "Stationary"))
+			SNew(SHorizontalBox)
+			+ SHorizontalBox::Slot()
+			.AutoWidth()
+			.VAlign(VAlign_Center)
+			.Padding(0.0f, 0.0f, 3.0f, 0.0f)
+			[
+				SNew(SImage)
+				.Image(FAppStyle::Get().GetBrush("PropertyWindow.MobilityStationary"))
+				.ColorAndOpacity(FSlateColor::UseForeground())
+			]
+			+ SHorizontalBox::Slot()
+			.AutoWidth()
+			.VAlign(VAlign_Center)
+			[
+				SNew(STextBlock)
+				.Font(FAppStyle::Get().GetFontStyle("PropertyWindow.MobilityFont"))
+				.Text(LOCTEXT("Stationary", "Stationary"))
+			]
 		]
 		.ToolTip(StationaryTooltip);
 	}
@@ -83,9 +113,24 @@ void FMobilityCustomization::CreateMobilityCustomization(IDetailCategoryBuilder&
 	// Movable Mobility
 	ButtonOptionsPanel->AddSlot(EComponentMobility::Movable)
 	[
-		SNew(STextBlock)
-		.Font(FAppStyle::Get().GetFontStyle("PropertyWindow.MobilityFont"))
-		.Text(LOCTEXT("Movable", "Movable"))
+		SNew(SHorizontalBox)
+		+ SHorizontalBox::Slot()
+		.AutoWidth()
+		.VAlign(VAlign_Center)
+		.Padding(0.0f, 0.0f, 3.0f, 0.0f)
+		[
+			SNew(SImage)
+			.Image(FAppStyle::Get().GetBrush("PropertyWindow.MobilityMoveable"))
+			.ColorAndOpacity(FSlateColor::UseForeground())
+		]
+		+ SHorizontalBox::Slot()
+		.AutoWidth()
+		.VAlign(VAlign_Center)
+		[
+			SNew(STextBlock)
+			.Font(FAppStyle::Get().GetFontStyle("PropertyWindow.MobilityFont"))
+			.Text(LOCTEXT("Movable", "Movable"))
+		]
 	]
 	.ToolTip(MovableTooltip);
 
diff --git a/Engine/Source/Editor/EditorStyle/Private/StarshipStyle.cpp b/Engine/Source/Editor/EditorStyle/Private/StarshipStyle.cpp
--- a/Engine/Source/Editor/EditorStyle/Private/StarshipStyle.cpp
+++ b/Engine/Source/Editor/EditorStyle/Private/StarshipStyle.cpp
@@ -1260,10 +1260,9 @@ void FStarshipEditorStyle::FStyle::SetupGeneralStyles()
 #if WITH_EDITOR || (IS_PROGRAM && WITH_UNREAL_DEVELOPER_TOOLS)
 	// Asset editors (common)
 	{
-		Set( "AssetEditor.SaveAsset.Greyscale", new IMAGE_BRUSH( "Icons/icon_file_save_16px", Icon16x16 ) );
-		Set( "AssetEditor.SaveAsset", new IMAGE_BRUSH_SVG( "Starship/MainToolbar/save", Icon40x40 ) );
-		Set( "AssetEditor.SaveAssetAs", new IMAGE_BRUSH( "Icons/icon_file_saveas_40x", Icon40x40 ) );
-		Set( "AssetEditor.SaveAssetAs.Small", new IMAGE_BRUSH( "Icons/icon_file_saveas_40x", Icon20x20 ) );
+		Set( "AssetEditor.SaveAsset", new IMAGE_BRUSH_SVG("Starship/Common/SaveCurrent", Icon16x16));
+		Set( "AssetEditor.SaveAssetAs", new IMAGE_BRUSH_SVG("Starship/Common/SaveCurrentAs", Icon16x16));
+
 		Set( "AssetEditor.ReimportAsset", new IMAGE_BRUSH( "Icons/icon_TextureEd_Reimport_40x", Icon40x40 ) );
 		Set( "AssetEditor.ReimportAsset.Small", new IMAGE_BRUSH( "Icons/icon_TextureEd_Reimport_40x", Icon20x20 ) );
 	}
@@ -3165,7 +3164,11 @@ void FStarshipEditorStyle::FStyle::SetupPropertyEditorStyles()
 		FSlateFontInfo MobilityFont = FStyleFonts::Get().Small;
 		MobilityFont.LetterSpacing = 100;
 
-		Set( "PropertyWindow.MobilityFont", MobilityFont );
+		Set("PropertyWindow.MobilityFont", MobilityFont );
+		Set("PropertyWindow.MobilityStatic", new IMAGE_BRUSH_SVG("Starship/Common/MobilityStatic", Icon16x16));
+		Set("PropertyWindow.MobilityStationary", new IMAGE_BRUSH_SVG("Starship/Common/MobilityStationary", Icon16x16));
+		Set("PropertyWindow.MobilityMoveable", new IMAGE_BRUSH_SVG("Starship/Common/MobilityMoveable", Icon16x16));
+
 		Set( "PropertyWindow.NoOverlayColor", new FSlateNoResource() );
 		Set( "PropertyWindow.EditConstColor", new FSlateColorBrush( FColor( 152, 152, 152, 80 ) ) );
 		Set( "PropertyWindow.FilteredColor", new FSlateColorBrush( FColor( 0, 255, 0, 80 ) ) );
@@ -6775,66 +6778,37 @@ void FStarshipEditorStyle::FStyle::SetupToolkitStyles()
 
 	// Static Mesh Editor
 	{
-		Set("StaticMeshEditor.Tabs.Properties", new IMAGE_BRUSH("/Icons/icon_tab_SelectionDetails_16x", Icon16x16));
-		Set("StaticMeshEditor.Tabs.SocketManager", new IMAGE_BRUSH("/Icons/icon_Static_Mesh_SocketManager_16x", Icon16x16));
-		Set("StaticMeshEditor.Tabs.ConvexDecomposition", new IMAGE_BRUSH("/Icons/icon_Static_Mesh_Convex_Decomposition_16x", Icon16x16));
-
-		Set( "StaticMeshEditor.NormalFont", DEFAULT_FONT( "Regular", 8 ) );
-		Set( "StaticMeshEditor.SetShowWireframe", new IMAGE_BRUSH( "Icons/icon_StaticMeshEd_Wireframe_40x", Icon40x40 ) );
-		Set( "StaticMeshEditor.SetShowWireframe.Small", new IMAGE_BRUSH( "Icons/icon_StaticMeshEd_Wireframe_40x", Icon20x20 ) );
-		Set( "StaticMeshEditor.SetShowVertexColor", new IMAGE_BRUSH( "Icons/icon_StaticMeshEd_VertColor_40x", Icon40x40 ) );
-		Set( "StaticMeshEditor.SetShowVertexColor.Small", new IMAGE_BRUSH( "Icons/icon_StaticMeshEd_VertColor_40x", Icon20x20 ) );
-		Set( "StaticMeshEditor.SetRealtimePreview", new IMAGE_BRUSH_SVG("Starship/Common/Realtime", Icon16x16));
-
-		Set( "StaticMeshEditor.ReimportMesh", new IMAGE_BRUSH(TEXT("Icons/icon_Persona_ReimportMesh_40x"), Icon40x40));
-		Set( "StaticMeshEditor.ReimportMesh.Small", new IMAGE_BRUSH(TEXT("Icons/icon_Persona_ReimportMesh_40x"), Icon20x20));
-		Set( "StaticMeshEditor.SetShowBounds", new IMAGE_BRUSH( "Icons/icon_StaticMeshEd_Bounds_40x", Icon40x40 ) );
-		Set( "StaticMeshEditor.SetShowBounds.Small", new IMAGE_BRUSH( "Icons/icon_StaticMeshEd_Bounds_40x", Icon20x20 ) );
-		Set( "StaticMeshEditor.SetShowCollision", new IMAGE_BRUSH( "Icons/icon_StaticMeshEd_Collision_40x", Icon40x40 ) );
-		Set( "StaticMeshEditor.SetShowCollision.Small", new IMAGE_BRUSH( "Icons/icon_StaticMeshEd_Collision_40x", Icon20x20 ) );
-		Set( "StaticMeshEditor.SetShowGrid", new IMAGE_BRUSH( "Icons/icon_MatEd_Grid_40x", Icon40x40 ) );
-		Set( "StaticMeshEditor.SetShowGrid.Small", new IMAGE_BRUSH( "Icons/icon_MatEd_Grid_40x", Icon20x20 ) );
-		Set( "StaticMeshEditor.SetDrawUVs", new IMAGE_BRUSH( "Icons/icon_StaticMeshEd_UVOverlay_40x", Icon40x40 ) );
-		Set( "StaticMeshEditor.SetDrawUVs.Small", new IMAGE_BRUSH( "Icons/icon_StaticMeshEd_UVOverlay_40x", Icon20x20 ) );
-		Set( "StaticMeshEditor.ResetCamera", new IMAGE_BRUSH_SVG("Starship/Common/ResetCamera", Icon16x16));
-
-		Set( "StaticMeshEditor.SetShowPivot", new IMAGE_BRUSH( "Icons/icon_StaticMeshEd_ShowPivot_40x", Icon40x40 ) );
-		Set( "StaticMeshEditor.SetShowPivot.Small", new IMAGE_BRUSH( "Icons/icon_StaticMeshEd_ShowPivot_40x", Icon20x20 ) );
-		Set( "StaticMeshEditor.SetShowSockets", new IMAGE_BRUSH( "Icons/icon_StaticMeshEd_ShowSockets_40x", Icon40x40 ) );
-		Set( "StaticMeshEditor.SetShowSockets.Small", new IMAGE_BRUSH( "Icons/icon_StaticMeshEd_ShowSockets_40x", Icon20x20 ) );
-		Set( "StaticMeshEditor.SaveThumbnail", new IMAGE_BRUSH( "Icons/icon_StaticMeshEd_SaveThumbnail_40x", Icon40x40 ) );
-		Set( "StaticMeshEditor.SaveThumbnail.Small", new IMAGE_BRUSH( "Icons/icon_StaticMeshEd_SaveThumbnail_40x", Icon20x20 ) );
-		Set( "StaticMeshEditor.SetShowNormals", new IMAGE_BRUSH( "Icons/icon_StaticMeshEd_Normals_40x", Icon40x40 ) );
-		Set( "StaticMeshEditor.SetShowNormals.Small", new IMAGE_BRUSH( "Icons/icon_StaticMeshEd_Normals_40x", Icon20x20 ) );
-		Set( "StaticMeshEditor.SetShowTangents", new IMAGE_BRUSH( "Icons/icon_StaticMeshEd_Tangents_40x", Icon40x40 ) );
-		Set( "StaticMeshEditor.SetShowTangents.Small", new IMAGE_BRUSH( "Icons/icon_StaticMeshEd_Tangents_40x", Icon20x20 ) );
-		Set( "StaticMeshEditor.SetShowBinormals", new IMAGE_BRUSH( "Icons/icon_StaticMeshEd_Binormals_40x", Icon40x40 ) );
-		Set( "StaticMeshEditor.SetShowBinormals.Small", new IMAGE_BRUSH( "Icons/icon_StaticMeshEd_Binormals_40x", Icon20x20 ) );
-		Set( "StaticMeshEditor.SetDrawAdditionalData", new IMAGE_BRUSH( "Icons/icon_StaticMeshEd_AdditionalData_40x", Icon40x40 ) );
-		Set( "StaticMeshEditor.SetDrawAdditionalData.Small", new IMAGE_BRUSH( "Icons/icon_StaticMeshEd_AdditionalData_40x", Icon20x20 ) );
-		Set( "StaticMeshEditor.GroupSection", new BOX_BRUSH( "Common/RoundedSelection_16x", FMargin( 4.0f / 16.0f ) ) );
-		Set( "StaticMeshEditor.SetShowVertices", new IMAGE_BRUSH("Icons/icon_StaticMeshEd_Vertices_40x", Icon40x40 ) );
-		Set( "StaticMeshEditor.SetShowVertices.Small", new IMAGE_BRUSH("Icons/icon_StaticMeshEd_Vertices_40x", Icon20x20 ) );
-		Set( "StaticMeshEditor.ToggleShowPivots", new IMAGE_BRUSH("Icons/icon_StaticMeshEd_ShowPivot_40x", Icon40x40));
-		Set( "StaticMeshEditor.ToggleShowPivot.Small", new IMAGE_BRUSH("Icons/icon_StaticMeshEd_ShowPivot_40x", Icon20x20));
-		Set( "StaticMeshEditor.ToggleShowSockets", new IMAGE_BRUSH("Icons/icon_StaticMeshEd_ShowSockets_40x", Icon40x40));
-		Set( "StaticMeshEditor.ToggleShowSockets.Small", new IMAGE_BRUSH("Icons/icon_StaticMeshEd_ShowSockets_40x", Icon20x20));
-		Set( "StaticMeshEditor.ToggleShowNormals", new IMAGE_BRUSH("Icons/icon_StaticMeshEd_Normals_40x", Icon40x40));
-		Set( "StaticMeshEditor.ToggleShowNormals.Small", new IMAGE_BRUSH("Icons/icon_StaticMeshEd_Normals_40x", Icon20x20));
-		Set( "StaticMeshEditor.ToggleShowTangents", new IMAGE_BRUSH("Icons/icon_StaticMeshEd_Tangents_40x", Icon40x40));
-		Set( "StaticMeshEditor.TogglShowTangents.Small", new IMAGE_BRUSH("Icons/icon_StaticMeshEd_Tangents_40x", Icon20x20));
-		Set( "StaticMeshEditor.ToggleShowBinormals", new IMAGE_BRUSH("Icons/icon_StaticMeshEd_Binormals_40x", Icon40x40));
-		Set( "StaticMeshEditor.ToggleShowBinormals.Small", new IMAGE_BRUSH("Icons/icon_StaticMeshEd_Binormals_40x", Icon20x20));
-		Set( "StaticMeshEditor.ToggleShowBounds", new IMAGE_BRUSH("Icons/icon_StaticMeshEd_Bounds_40x", Icon40x40));
-		Set( "StaticMeshEditor.ToggleShowBounds.Small", new IMAGE_BRUSH("Icons/icon_StaticMeshEd_Bounds_40x", Icon20x20));
-		Set( "StaticMeshEditor.ToggleShowGrids", new IMAGE_BRUSH("Icons/icon_MatEd_Grid_40x", Icon40x40));
-		Set( "StaticMeshEditor.ToggleShowGrids.Small", new IMAGE_BRUSH("Icons/icon_MatEd_Grid_40x", Icon20x20));
-		Set( "StaticMeshEditor.ToggleShowVertices", new IMAGE_BRUSH("Icons/icon_StaticMeshEd_Vertices_40x", Icon40x40));
-		Set( "StaticMeshEditor.ToggleShowVertices.Small", new IMAGE_BRUSH("Icons/icon_StaticMeshEd_Vertices_40x", Icon20x20));
-		Set( "StaticMeshEditor.ToggleShowWireframes", new IMAGE_BRUSH("Icons/icon_StaticMeshEd_Wireframe_40x", Icon40x40));
-		Set( "StaticMeshEditor.ToggleShowWireframes.Small", new IMAGE_BRUSH("Icons/icon_StaticMeshEd_Wireframe_40x", Icon20x20));
-		Set( "StaticMeshEditor.ToggleShowVertexColors", new IMAGE_BRUSH("Icons/icon_StaticMeshEd_VertColor_40x", Icon40x40));
-		Set( "StaticMeshEditor.ToggleShowVertexColors.Small", new IMAGE_BRUSH("Icons/icon_StaticMeshEd_VertColor_40x", Icon20x20));
+		Set("StaticMeshEditor.Tabs.Properties", new IMAGE_BRUSH_SVG("Starship/Common/Details", Icon16x16));
+		Set("StaticMeshEditor.Tabs.SocketManager", new IMAGE_BRUSH_SVG("Starship/StaticMeshEditor/SocketManager", Icon16x16));
+		Set("StaticMeshEditor.Tabs.ConvexDecomposition", new IMAGE_BRUSH_SVG("Starship/StaticMeshEditor/ConvexDecomposition", Icon16x16));
+		Set("StaticMeshEditor.SetShowWireframe", new IMAGE_BRUSH_SVG("Starship/Common/BrushWireframe", Icon16x16));
+		Set("StaticMeshEditor.SetShowVertexColor", new IMAGE_BRUSH_SVG( "Starship/Common/SetShowVertexColors", Icon16x16));
+		Set("StaticMeshEditor.SetRealtimePreview", new IMAGE_BRUSH_SVG("Starship/Common/Realtime", Icon16x16));
+		Set("StaticMeshEditor.ReimportMesh", new IMAGE_BRUSH_SVG("Starship/StaticMeshEditor/ReimportMesh", Icon20x20));
+		Set("StaticMeshEditor.SetShowBounds", new IMAGE_BRUSH_SVG("Starship/Common/SetShowBounds", Icon16x16));
+		Set("StaticMeshEditor.SetDrawUVs", new IMAGE_BRUSH_SVG("Starship/Common/SetDrawUVs", Icon20x20));
+		Set("StaticMeshEditor.SetShowCollision", new IMAGE_BRUSH_SVG("Starship/Common/Collision", Icon20x20));
+		Set("StaticMeshEditor.SetShowGrid", new IMAGE_BRUSH_SVG("Starship/Common/Grid", Icon16x16));
+		Set("StaticMeshEditor.ResetCamera", new IMAGE_BRUSH_SVG("Starship/Common/ResetCamera", Icon16x16));
+		Set("StaticMeshEditor.SetShowPivot", new IMAGE_BRUSH_SVG("Starship/Common/SetShowPivot", Icon16x16));
+		Set("StaticMeshEditor.SetShowSockets", new IMAGE_BRUSH_SVG( "Starship/Common/SetShowSockets", Icon16x16));
+		Set("StaticMeshEditor.SaveThumbnail", new IMAGE_BRUSH_SVG( "Starship/Common/SaveThumbnail", Icon16x16));
+		Set("StaticMeshEditor.SetShowNormals", new IMAGE_BRUSH_SVG( "Starship/Common/SetShowNormals", Icon16x16));
+		Set("StaticMeshEditor.SetShowTangents", new IMAGE_BRUSH_SVG("Starship/Common/SetShowTangents", Icon16x16));
+		Set("StaticMeshEditor.SetShowBinormals", new IMAGE_BRUSH_SVG("Starship/Common/SetShowBinormals", Icon16x16));
+		Set("StaticMeshEditor.SetDrawAdditionalData", new IMAGE_BRUSH_SVG( "Starship/StaticMeshEditor/AdditionalData", Icon16x16));
+		Set("StaticMeshEditor.SetShowVertices", new IMAGE_BRUSH_SVG("Starship/Common/SetShowVertices", Icon16x16));
+		Set("StaticMeshEditor.ToggleShowPivots", new IMAGE_BRUSH_SVG("Starship/Common/SetShowPivot", Icon16x16));
+		Set("StaticMeshEditor.ToggleShowSockets", new IMAGE_BRUSH_SVG("Starship/Common/SetShowSockets", Icon16x16));
+		Set("StaticMeshEditor.ToggleShowNormals", new IMAGE_BRUSH_SVG("Starship/Common/SetShowNormals", Icon16x16));
+		Set("StaticMeshEditor.ToggleShowTangents", new IMAGE_BRUSH_SVG("Starship/Common/SetShowTangents", Icon16x16));
+		Set("StaticMeshEditor.ToggleShowBinormals", new IMAGE_BRUSH_SVG("Starship/Common/SetShowBinormals", Icon16x16));
+		Set("StaticMeshEditor.ToggleShowBounds", new IMAGE_BRUSH_SVG("Starship/Common/SetShowBounds", Icon16x16));
+		Set("StaticMeshEditor.ToggleShowGrids", new IMAGE_BRUSH_SVG("Starship/Common/Grid", Icon16x16));
+		Set("StaticMeshEditor.ToggleShowVertices", new IMAGE_BRUSH_SVG("Starship/Common/SetShowVertices", Icon16x16));
+		Set("StaticMeshEditor.ToggleShowWireframes", new IMAGE_BRUSH_SVG("Starship/Common/BrushWireframe", Icon16x16));
+		Set("StaticMeshEditor.ToggleShowVertexColors", new IMAGE_BRUSH_SVG("Starship/Common/SetShowVertexColors", Icon16x16));
+
 	}
 
 	// Skeletal Mesh Editor
@@ -7497,7 +7471,7 @@ void FStarshipEditorStyle::FStyle::SetupAutomationStyles()
 
 		//test preset icons
 		Set( "AutomationWindow.PresetNew", new IMAGE_BRUSH( "Icons/icon_add_40x", Icon16x16 ) );
-		Set( "AutomationWindow.PresetSave", new IMAGE_BRUSH( "Icons/icon_file_save_16px", Icon16x16 ) );
+		Set( "AutomationWindow.PresetSave", new IMAGE_BRUSH_SVG("Starship/Common/SaveCurrent", Icon16x16));
 		Set( "AutomationWindow.PresetRemove", new IMAGE_BRUSH( "Icons/icon_Cascade_DeleteLOD_40x", Icon16x16 ) );
 
 		//test backgrounds
diff --git a/Engine/Source/Editor/StaticMeshEditor/Private/StaticMeshEditorTools.cpp b/Engine/Source/Editor/StaticMeshEditor/Private/StaticMeshEditorTools.cpp
--- a/Engine/Source/Editor/StaticMeshEditor/Private/StaticMeshEditorTools.cpp
+++ b/Engine/Source/Editor/StaticMeshEditor/Private/StaticMeshEditorTools.cpp
@@ -2096,7 +2096,7 @@ TSharedRef<SWidget> FMeshSectionSettingsLayout::OnGenerateCustomSectionWidgetsFo
 				.OnCheckStateChanged(this, &FMeshSectionSettingsLayout::OnSectionCastShadowChanged, SectionIndex)
 			[
 				SNew(STextBlock)
-					.Font(FEditorStyle::GetFontStyle("StaticMeshEditor.NormalFont"))
+					.Font(IDetailLayoutBuilder::GetDetailFont())
 					.Text(LOCTEXT("CastShadow", "Cast Shadow"))
 			]
 		]
@@ -2111,7 +2111,7 @@ TSharedRef<SWidget> FMeshSectionSettingsLayout::OnGenerateCustomSectionWidgetsFo
 				.OnCheckStateChanged(this, &FMeshSectionSettingsLayout::OnSectionCollisionChanged, SectionIndex)
 			[
 				SNew(STextBlock)
-					.Font(FEditorStyle::GetFontStyle("StaticMeshEditor.NormalFont"))
+					.Font(IDetailLayoutBuilder::GetDetailFont())
 					.Text(LOCTEXT("EnableCollision", "Enable Collision"))
 			]
 		]
@@ -2124,7 +2124,7 @@ TSharedRef<SWidget> FMeshSectionSettingsLayout::OnGenerateCustomSectionWidgetsFo
 			.OnCheckStateChanged(this, &FMeshSectionSettingsLayout::OnSectionVisibleInRayTracingChanged, SectionIndex)
 			[
 				SNew(STextBlock)
-				.Font(FEditorStyle::GetFontStyle("StaticMeshEditor.NormalFont"))
+				.Font(IDetailLayoutBuilder::GetDetailFont())
 			.Text(LOCTEXT("VisibleInRayTracing", "Visible In Ray Tracing"))
 			]
 		]
@@ -2137,7 +2137,7 @@ TSharedRef<SWidget> FMeshSectionSettingsLayout::OnGenerateCustomSectionWidgetsFo
 			.OnCheckStateChanged(this, &FMeshSectionSettingsLayout::OnSectionForceOpaqueFlagChanged, SectionIndex)
 			[
 				SNew(STextBlock)
-					.Font(FEditorStyle::GetFontStyle("StaticMeshEditor.NormalFont"))
+					.Font(IDetailLayoutBuilder::GetDetailFont())
 					.Text(LOCTEXT("ForceOpaque", "Force Opaque"))
 			]
 		];
@@ -2700,7 +2700,7 @@ TSharedRef<SWidget> FMeshMaterialsLayout::OnGenerateWidgetsForMaterial(UMaterial
 				.OnCheckStateChanged(this, &FMeshMaterialsLayout::OnOverrideUVDensityChanged, SlotIndex)
 			[
 				SNew(STextBlock)
-					.Font(FEditorStyle::GetFontStyle("StaticMeshEditor.NormalFont"))
+					.Font(IDetailLayoutBuilder::GetDetailFont())
 					.Text(LOCTEXT("OverrideUVDensity", "Override UV Density"))
 			]
 		]
@@ -3538,7 +3538,7 @@ void FLevelOfDetailSettingsLayout::AddLODLevelCategories( IDetailLayoutBuilder&
 						.AutoWidth()
 						[
 							SNew(STextBlock)
-							.Font(FEditorStyle::GetFontStyle("StaticMeshEditor.NormalFont"))
+							.Font(IDetailLayoutBuilder::GetDetailFont())
 							.Text(this, &FLevelOfDetailSettingsLayout::GetLODScreenSizeTitle, LODIndex)
 							.Visibility( LODIndex > 0 ? EVisibility::Visible : EVisibility::Collapsed )
 						]
@@ -3547,7 +3547,7 @@ void FLevelOfDetailSettingsLayout::AddLODLevelCategories( IDetailLayoutBuilder&
 						.AutoWidth()
 						[
 							SNew(STextBlock)
-							.Font(FEditorStyle::GetFontStyle("StaticMeshEditor.NormalFont"))
+							.Font(IDetailLayoutBuilder::GetDetailFont())
 							.Text( FText::Format( LOCTEXT("Triangles_MeshSimplification", "Triangles: {0}"), FText::AsNumber( StaticMeshEditor.GetNumTriangles(LODIndex) ) ) )
 						]
 						+ SHorizontalBox::Slot()
@@ -3555,7 +3555,7 @@ void FLevelOfDetailSettingsLayout::AddLODLevelCategories( IDetailLayoutBuilder&
 						.AutoWidth()
 						[
 							SNew(STextBlock)
-							.Font(FEditorStyle::GetFontStyle("StaticMeshEditor.NormalFont"))
+							.Font(IDetailLayoutBuilder::GetDetailFont())
 							.Text( FText::Format( LOCTEXT("Vertices_MeshSimplification", "Vertices: {0}"), FText::AsNumber( StaticMeshEditor.GetNumVertices(LODIndex) ) ) )
 						]
 					]
diff --git a/Engine/Source/Editor/TranslationEditor/Private/TranslationEditorMenu.cpp b/Engine/Source/Editor/TranslationEditor/Private/TranslationEditorMenu.cpp
--- a/Engine/Source/Editor/TranslationEditor/Private/TranslationEditorMenu.cpp
+++ b/Engine/Source/Editor/TranslationEditor/Private/TranslationEditorMenu.cpp
@@ -37,7 +37,7 @@ void FTranslationEditorMenu::SetupTranslationEditorMenu( TSharedPtr< FExtender >
 		{
 			static void AddSaveMenuOption( FMenuBuilder& MenuBuilder )
 			{
-				MenuBuilder.AddMenuEntry( FTranslationEditorCommands::Get().SaveTranslations, "SaveTranslations", TAttribute<FText>(), TAttribute<FText>(), FSlateIcon(FEditorStyle::GetStyleSetName(), "AssetEditor.SaveAsset.Greyscale") );
+				MenuBuilder.AddMenuEntry( FTranslationEditorCommands::Get().SaveTranslations, "SaveTranslations", TAttribute<FText>(), TAttribute<FText>(), FSlateIcon(FEditorStyle::GetStyleSetName(), "AssetEditor.SaveAsset") );
 			}
 
 			static void AddTranslationEditorMenu( FMenuBarBuilder& MenuBarBuilder )
diff --git a/Engine/Source/Editor/UnrealEd/Private/Toolkits/AssetEditorToolkit.cpp b/Engine/Source/Editor/UnrealEd/Private/Toolkits/AssetEditorToolkit.cpp
--- a/Engine/Source/Editor/UnrealEd/Private/Toolkits/AssetEditorToolkit.cpp
+++ b/Engine/Source/Editor/UnrealEd/Private/Toolkits/AssetEditorToolkit.cpp
@@ -941,10 +941,10 @@ void FAssetEditorToolkit::FillDefaultFileMenuCommands(FToolMenuSection& InSectio
 {
 	const FToolMenuInsert InsertPosition(NAME_None, EToolMenuInsertType::First);
 
-	InSection.AddMenuEntry(FAssetEditorCommonCommands::Get().SaveAsset, TAttribute<FText>(), TAttribute<FText>(), FSlateIcon(FEditorStyle::GetStyleSetName(), "AssetEditor.SaveAsset.Greyscale")).InsertPosition = InsertPosition;
+	InSection.AddMenuEntry(FAssetEditorCommonCommands::Get().SaveAsset, TAttribute<FText>(), TAttribute<FText>(), FSlateIcon(FEditorStyle::GetStyleSetName(), "AssetEditor.SaveAsset")).InsertPosition = InsertPosition;
 	if( IsActuallyAnAsset() )
 	{
-		InSection.AddMenuEntry(FAssetEditorCommonCommands::Get().SaveAssetAs, TAttribute<FText>(), TAttribute<FText>(), FSlateIcon(FEditorStyle::GetStyleSetName(), "AssetEditor.SaveAssetAs.Small")).InsertPosition = InsertPosition;
+		InSection.AddMenuEntry(FAssetEditorCommonCommands::Get().SaveAssetAs, TAttribute<FText>(), TAttribute<FText>(), FSlateIcon(FEditorStyle::GetStyleSetName(), "AssetEditor.SaveAssetAs")).InsertPosition = InsertPosition;
 	}
 	InSection.AddSeparator("DefaultFileMenuCommandsSeparator").InsertPosition = InsertPosition;;
 
