Fixes PlaceActors Icon in the PlaceActors Menu

diff --git a/Engine/Source/Editor/EditorStyle/Private/StarshipStyle.cpp b/Engine/Source/Editor/EditorStyle/Private/StarshipStyle.cpp
--- a/Engine/Source/Editor/EditorStyle/Private/StarshipStyle.cpp
+++ b/Engine/Source/Editor/EditorStyle/Private/StarshipStyle.cpp
@@ -4566,7 +4566,7 @@ void FStarshipEditorStyle::FStyle::SetupLevelEditorStyle()
 
 		Set( "LevelEditor.WorldProperties.Tab", new IMAGE_BRUSH_SVG( "Starship/Common/WorldSettings", Icon16x16 ) );
 
-		Set( "LevelEditor.OpenCreateMenu", new IMAGE_BRUSH_SVG( "Starship/Common/PlaceActors", Icon20x20 ) );
+		Set( "LevelEditor.OpenPlaceActors", new IMAGE_BRUSH_SVG( "Starship/Common/PlaceActors", Icon20x20 ) );
 		Set( "LevelEditor.OpenContentBrowser", new IMAGE_BRUSH_SVG( "Starship/MainToolbar/content", Icon20x20 ) );
 		Set( "LevelEditor.OpenMarketplace", new IMAGE_BRUSH_SVG( "Starship/MainToolbar/marketplace", Icon20x20) );
 		Set( "LevelEditor.OpenLevelBlueprint", new IMAGE_BRUSH_SVG( "Starship/MainToolbar/blueprints", Icon20x20) );
diff --git a/Engine/Source/Editor/LevelEditor/Private/LevelEditorToolBar.cpp b/Engine/Source/Editor/LevelEditor/Private/LevelEditorToolBar.cpp
--- a/Engine/Source/Editor/LevelEditor/Private/LevelEditorToolBar.cpp
+++ b/Engine/Source/Editor/LevelEditor/Private/LevelEditorToolBar.cpp
@@ -1208,7 +1208,7 @@ void FLevelEditorToolBar::RegisterLevelEditorToolBar( const TSharedRef<FUIComman
 				FOnGetContent::CreateStatic(&FLevelEditorToolBar::GenerateCreateMenuWidget, InCommandList, TWeakPtr<SLevelEditor>(InLevelEditor)),
 				LOCTEXT("OpenCreate_Label", "Content"),
 				LOCTEXT("OpenCreate_Tooltip", "Quickly add to the project."),
-				FSlateIcon(FEditorStyle::GetStyleSetName(), "LevelEditor.OpenCreateMenu")
+				FSlateIcon(FEditorStyle::GetStyleSetName(), "LevelEditor.OpenPlaceActors")
 			));
 
 			Section.AddEntry(FToolMenuEntry::InitComboButton(
