Added 4 pixels of padding to the top of the name area in the actor details panel.

Removed 2 pixels of padding from and the down arrow from SEditorHeaderButtons.

Added 2 pixels of padding to the right of Blueprint add parameter button.

diff --git a/Engine/Source/Editor/EditorStyle/Private/StarshipStyle.cpp b/Engine/Source/Editor/EditorStyle/Private/StarshipStyle.cpp
--- a/Engine/Source/Editor/EditorStyle/Private/StarshipStyle.cpp
+++ b/Engine/Source/Editor/EditorStyle/Private/StarshipStyle.cpp
@@ -708,14 +708,14 @@ void FStarshipEditorStyle::FStyle::SetupGeneralStyles()
 		.SetHoveredForeground(FStyleColors::ForegroundHover)
 		.SetPressedForeground(FStyleColors::ForegroundHover)
 		.SetDisabledForeground(FStyleColors::White25)
-		.SetNormalPadding(FMargin(6, 6))
-		.SetPressedPadding(FMargin(6, 6));
+		.SetNormalPadding(FMargin(6, 4))
+		.SetPressedPadding(FMargin(6, 4));
 
 	Set("RoundedButton", RoundedButtonStyle);
 
 	FComboButtonStyle AddNewComboButtonStyle = FComboButtonStyle(FStarshipCoreStyle::GetCoreStyle().GetWidgetStyle<FComboButtonStyle>("SimpleComboButton"))
-		.SetButtonStyle(RoundedButtonStyle)
-		.SetDownArrowPadding(FMargin(1,0,0,0));
+		.SetButtonStyle(RoundedButtonStyle) 
+		.SetDownArrowImage(FSlateNoResource());
 
 	Set("RoundedComboButton", AddNewComboButtonStyle);
 
diff --git a/Engine/Source/Editor/Kismet/Private/BlueprintDetailsCustomization.cpp b/Engine/Source/Editor/Kismet/Private/BlueprintDetailsCustomization.cpp
--- a/Engine/Source/Editor/Kismet/Private/BlueprintDetailsCustomization.cpp
+++ b/Engine/Source/Editor/Kismet/Private/BlueprintDetailsCustomization.cpp
@@ -3748,6 +3748,7 @@ void FBlueprintGraphActionDetails::CustomizeDetails( IDetailLayoutBuilder& Detai
 
 			OutputsHeaderContentWidget->AddSlot()
 				.HAlign(HAlign_Right)
+				.Padding(FMargin(0,0,2,0))
 				[
 					SNew(SButton)
 					.ButtonStyle(FEditorStyle::Get(), "SimpleButton")
diff --git a/Engine/Source/Editor/LevelEditor/Private/SActorDetails.cpp b/Engine/Source/Editor/LevelEditor/Private/SActorDetails.cpp
--- a/Engine/Source/Editor/LevelEditor/Private/SActorDetails.cpp
+++ b/Engine/Source/Editor/LevelEditor/Private/SActorDetails.cpp
@@ -149,7 +149,7 @@ void SActorDetails::Construct(const FArguments& InArgs, UTypedElementSelectionSe
 	[
 		SNew(SVerticalBox)
 		+SVerticalBox::Slot()
-		.Padding(0.0f, 0.0f, 0.0f, 2.0f)
+		.Padding(0.0f, 4.0f, 0.0f, 4.0f)
 		.AutoHeight()
 		[
 			DetailsView->GetNameAreaWidget().ToSharedRef()
