Adding pause/eject/stop icons for PIE in the new style

diff --git a/Engine/Source/Editor/EditorStyle/Private/StarshipStyle.cpp b/Engine/Source/Editor/EditorStyle/Private/StarshipStyle.cpp
--- a/Engine/Source/Editor/EditorStyle/Private/StarshipStyle.cpp
+++ b/Engine/Source/Editor/EditorStyle/Private/StarshipStyle.cpp
@@ -5507,21 +5507,21 @@ void FStarshipEditorStyle::FStyle::SetupPersonaStyle()
 
 		Set( "PlayWorld.ResumePlaySession", new IMAGE_BRUSH( "Icons/icon_simulate_40x", Icon40x40 ) );
 		Set( "PlayWorld.ResumePlaySession.Small", new IMAGE_BRUSH( "Icons/icon_simulate_40x", Icon20x20 ) );
-		Set( "PlayWorld.PausePlaySession", new IMAGE_BRUSH( "Icons/icon_pause_40x", Icon40x40 ) );
-		Set( "PlayWorld.PausePlaySession.Small", new IMAGE_BRUSH( "Icons/icon_pause_40x", Icon20x20 ) );
+		Set( "PlayWorld.PausePlaySession", new IMAGE_BRUSH_SVG("Starship/MainToolbar/pause", Icon40x40));
+		Set( "PlayWorld.PausePlaySession.Small", new IMAGE_BRUSH_SVG("Starship/MainToolbar/pause", Icon20x20));
 		Set( "PlayWorld.SingleFrameAdvance", new IMAGE_BRUSH( "Icons/icon_advance_40x", Icon40x40 ) );
 		Set( "PlayWorld.SingleFrameAdvance.Small", new IMAGE_BRUSH( "Icons/icon_advance_40x", Icon20x20 ) );
 
-		Set( "PlayWorld.StopPlaySession", new IMAGE_BRUSH( "Icons/icon_stop_40x", Icon40x40 ) );
-		Set( "PlayWorld.StopPlaySession.Small", new IMAGE_BRUSH( "Icons/icon_stop_40x", Icon20x20 ) );
+		Set( "PlayWorld.StopPlaySession", new IMAGE_BRUSH_SVG("Starship/MainToolbar/stop", Icon40x40));
+		Set( "PlayWorld.StopPlaySession.Small", new IMAGE_BRUSH_SVG("Starship/MainToolbar/stop", Icon20x20));
 
 		Set("PlayWorld.LateJoinSession", new IMAGE_BRUSH("Icons/icon_simulate_40x", Icon40x40));
 		Set("PlayWorld.LateJoinSession.Small", new IMAGE_BRUSH("Icons/icon_simulate_40x", Icon20x20));
 
 		Set( "PlayWorld.PossessPlayer", new IMAGE_BRUSH( "Icons/icon_possess_40x", Icon40x40 ) );
 		Set( "PlayWorld.PossessPlayer.Small", new IMAGE_BRUSH( "Icons/icon_possess_40x", Icon20x20 ) );
-		Set( "PlayWorld.EjectFromPlayer", new IMAGE_BRUSH( "Icons/icon_eject_40x", Icon40x40 ) );
-		Set( "PlayWorld.EjectFromPlayer.Small", new IMAGE_BRUSH( "Icons/icon_eject_40x", Icon20x20 ) );
+		Set( "PlayWorld.EjectFromPlayer", new IMAGE_BRUSH_SVG("Starship/MainToolbar/eject", Icon40x40));
+		Set( "PlayWorld.EjectFromPlayer.Small", new IMAGE_BRUSH_SVG("Starship/MainToolbar/eject", Icon20x20));
 
 		Set( "PlayWorld.ShowCurrentStatement", new IMAGE_BRUSH( "Icons/icon_findnode_40x", Icon40x40 ) );
 		Set( "PlayWorld.ShowCurrentStatement.Small", new IMAGE_BRUSH( "Icons/icon_findnode_40x", Icon20x20 ) );
