Fix a few minor icon bugs

diff --git a/Engine/Source/Editor/EditorStyle/Private/StarshipStyle.cpp b/Engine/Source/Editor/EditorStyle/Private/StarshipStyle.cpp
--- a/Engine/Source/Editor/EditorStyle/Private/StarshipStyle.cpp
+++ b/Engine/Source/Editor/EditorStyle/Private/StarshipStyle.cpp
@@ -5927,7 +5927,6 @@ void FStarshipEditorStyle::FStyle::SetupClassIconsAndThumbnails()
 			TEXT("LevelSequence"),
 			TEXT("LightmassCharacterIndirectDetailVolume"),
 			TEXT("MassiveLODOverrideVolume"),
-			TEXT("Material"),
 			TEXT("MaterialParameterCollection"),
 			TEXT("MatineeActor"),
 			TEXT("MultiFont"),
@@ -6031,6 +6030,7 @@ void FStarshipEditorStyle::FStyle::SetupClassIconsAndThumbnails()
 			{TEXT("LightmassImportanceVolume")},
 			{TEXT("LightmassVolume")},
 			{TEXT("LiveLinkPreset")},
+			{TEXT("Material")},
 			{TEXT("MaterialFunction")},
 			{TEXT("MaterialInstanceActor")},
 			{TEXT("MaterialInstanceConstant")},
diff --git a/Engine/Source/Editor/Kismet/Private/SSCSEditor.cpp b/Engine/Source/Editor/Kismet/Private/SSCSEditor.cpp
--- a/Engine/Source/Editor/Kismet/Private/SSCSEditor.cpp
+++ b/Engine/Source/Editor/Kismet/Private/SSCSEditor.cpp
@@ -3596,6 +3596,7 @@ TSharedRef<SWidget> SSCS_RowWidget_ActorRoot::GenerateWidgetForColumn(const FNam
 		[
 			SNew(SImage)
 			.Image(GetIconBrush())
+			.ColorAndOpacity(FSlateColor::UseForeground())
 		]
 
 	+ SHorizontalBox::Slot()
