Updated Static Mesh Icon

diff --git a/Engine/Source/Editor/EditorStyle/Private/StarshipStyle.cpp b/Engine/Source/Editor/EditorStyle/Private/StarshipStyle.cpp
--- a/Engine/Source/Editor/EditorStyle/Private/StarshipStyle.cpp
+++ b/Engine/Source/Editor/EditorStyle/Private/StarshipStyle.cpp
@@ -5944,7 +5944,6 @@ void FStarshipEditorStyle::FStyle::SetupClassIconsAndThumbnails()
 			TEXT("SpotLightMovable"),
 			TEXT("SpotLightStatic"),
 			TEXT("SpotLightStationary"),
-			TEXT("StaticMesh"),
 			TEXT("UserDefinedEnum"),
 			TEXT("Cube"),
 			TEXT("Sphere"),
@@ -6096,6 +6095,7 @@ void FStarshipEditorStyle::FStyle::SetupClassIconsAndThumbnails()
 			{TEXT("SphereReflectionCapture")},
 			{TEXT("SpotLight")},
 			{TEXT("SpringArmComponent")},
+			{TEXT("StaticMesh")},
 			{TEXT("StaticMeshActor")},
 			{TEXT("StreamMediaSource")},
 			{TEXT("SubsurfaceProfile")},
