Fix Add Component drop down icons, giant splitters, and organizational cleanup

diff --git a/Engine/Plugins/Experimental/GeometryCache/Source/GeometryCache/Classes/GeometryCacheComponent.h b/Engine/Plugins/Experimental/GeometryCache/Source/GeometryCache/Classes/GeometryCacheComponent.h
--- a/Engine/Plugins/Experimental/GeometryCache/Source/GeometryCache/Classes/GeometryCacheComponent.h
+++ b/Engine/Plugins/Experimental/GeometryCache/Source/GeometryCache/Classes/GeometryCacheComponent.h
@@ -33,7 +33,7 @@ struct FTrackRenderData
 };
 
 /** GeometryCacheComponent, encapsulates a GeometryCache asset instance and implements functionality for rendering/and playback of GeometryCaches */
-UCLASS(ClassGroup = (Rendering, Common), hidecategories = (Object, LOD), meta = (BlueprintSpawnableComponent), Experimental, ClassGroup = Experimental)
+UCLASS(ClassGroup = (Rendering), hidecategories = (Object, LOD), meta = (BlueprintSpawnableComponent), Experimental, ClassGroup = Experimental)
 class GEOMETRYCACHE_API UGeometryCacheComponent : public UMeshComponent
 {
 	GENERATED_UCLASS_BODY()
diff --git a/Engine/Source/Editor/EditorStyle/Private/StarshipStyle.cpp b/Engine/Source/Editor/EditorStyle/Private/StarshipStyle.cpp
--- a/Engine/Source/Editor/EditorStyle/Private/StarshipStyle.cpp
+++ b/Engine/Source/Editor/EditorStyle/Private/StarshipStyle.cpp
@@ -1985,6 +1985,52 @@ void FStarshipEditorStyle::FStyle::SetupLevelGeneralStyles()
 		Set("Level.EmptyIcon16x", new IMAGE_BRUSH("Icons/Empty_16x", Icon16x16));
 		Set("Level.ColorIcon40x", new IMAGE_BRUSH("Icons/icon_levels_back_16px", Icon16x16));
 	}
+
+	// Spline component controls
+	{
+		Set("SplineComponentDetails.SelectFirst", FButtonStyle(Button)
+			.SetNormal(IMAGE_BRUSH("/Sequencer/Transport_Bar/Go_To_Front_24x_OFF", Icon24x24))
+			.SetHovered(IMAGE_BRUSH("/Sequencer/Transport_Bar/Go_To_Front_24x_OFF", Icon24x24, SelectionColor))
+			.SetPressed(IMAGE_BRUSH("/Sequencer/Transport_Bar/Go_To_Front_24x", Icon24x24, SelectionColor_Pressed))
+		);
+
+		Set("SplineComponentDetails.AddPrev", FButtonStyle(Button)
+			.SetNormal(IMAGE_BRUSH("/Sequencer/Transport_Bar/Step_Backwards_24x_OFF", Icon24x24))
+			.SetHovered(IMAGE_BRUSH("/Sequencer/Transport_Bar/Step_Backwards_24x_OFF", Icon24x24, SelectionColor))
+			.SetPressed(IMAGE_BRUSH("/Sequencer/Transport_Bar/Step_Backwards_24x", Icon24x24, SelectionColor_Pressed))
+		);
+
+		Set("SplineComponentDetails.SelectPrev", FButtonStyle(Button)
+			.SetNormal(IMAGE_BRUSH("/Sequencer/Transport_Bar/Backwards_24x_OFF", Icon24x24))
+			.SetHovered(IMAGE_BRUSH("/Sequencer/Transport_Bar/Backwards_24x_OFF", Icon24x24, SelectionColor))
+			.SetPressed(IMAGE_BRUSH("/Sequencer/Transport_Bar/Backwards_24x", Icon24x24, SelectionColor_Pressed))
+		);
+
+		Set("SplineComponentDetails.SelectAll", FButtonStyle(Button)
+			.SetNormal(IMAGE_BRUSH("/Sequencer/Transport_Bar/Loop_24x_OFF", Icon24x24))
+			.SetHovered(IMAGE_BRUSH("/Sequencer/Transport_Bar/Loop_24x_OFF", Icon24x24, SelectionColor))
+			.SetPressed(IMAGE_BRUSH("/Sequencer/Transport_Bar/Loop_24x", Icon24x24, SelectionColor_Pressed))
+		);
+
+		Set("SplineComponentDetails.SelectNext", FButtonStyle(Button)
+			.SetNormal(IMAGE_BRUSH("/Sequencer/Transport_Bar/Play_24x_OFF", Icon24x24))
+			.SetHovered(IMAGE_BRUSH("/Sequencer/Transport_Bar/Play_24x_OFF", Icon24x24, SelectionColor))
+			.SetPressed(IMAGE_BRUSH("/Sequencer/Transport_Bar/Play_24x", Icon24x24, SelectionColor_Pressed))
+		);
+
+		Set("SplineComponentDetails.AddNext", FButtonStyle(Button)
+			.SetNormal(IMAGE_BRUSH("/Sequencer/Transport_Bar/Step_Forward_24x_OFF", Icon24x24))
+			.SetHovered(IMAGE_BRUSH("/Sequencer/Transport_Bar/Step_Forward_24x_OFF", Icon24x24, SelectionColor))
+			.SetPressed(IMAGE_BRUSH("/Sequencer/Transport_Bar/Step_Forward_24x", Icon24x24, SelectionColor_Pressed))
+		);
+
+		Set("SplineComponentDetails.SelectLast", FButtonStyle(Button)
+			.SetNormal(IMAGE_BRUSH("/Sequencer/Transport_Bar/Go_To_End_24x_OFF", Icon24x24))
+			.SetHovered(IMAGE_BRUSH("/Sequencer/Transport_Bar/Go_To_End_24x_OFF", Icon24x24, SelectionColor))
+			.SetPressed(IMAGE_BRUSH("/Sequencer/Transport_Bar/Go_To_End_24x", Icon24x24, SelectionColor_Pressed))
+		);
+
+	}
 }
 
 void FStarshipEditorStyle::FStyle::SetupWorldBrowserStyles()
@@ -5843,51 +5889,48 @@ void FStarshipEditorStyle::FStyle::SetupClassIconsAndThumbnails()
 
 		// Component classes
 		Set("ClassIcon.AudioComponent", new IMAGE_BRUSH("Icons/ActorIcons/SoundActor_16x", Icon16x16));
-		Set("ClassIcon.CameraComponent", new IMAGE_BRUSH("Icons/AssetIcons/CameraActor_16x", Icon16x16));
 		Set("ClassIcon.BlueprintCore", new IMAGE_BRUSH("Icons/AssetIcons/Blueprint_16x", Icon16x16));
-		Set("ClassIcon.BrushComponent", new IMAGE_BRUSH("Icons/ActorIcons/Brush_16x", Icon16x16));
-		Set("ClassIcon.DecalComponent", new IMAGE_BRUSH("Icons/AssetIcons/DecalActor_16x", Icon16x16));
-		Set("ClassIcon.DirectionalLightComponent", new IMAGE_BRUSH("Icons/AssetIcons/DirectionalLight_16x", Icon16x16));
-		Set("ClassIcon.ExponentialHeightFogComponent", new IMAGE_BRUSH("Icons/AssetIcons/ExponentialHeightFog_16x", Icon16x16));
-		Set("ClassIcon.ForceFeedbackComponent", new IMAGE_BRUSH("Icons/AssetIcons/ForceFeedbackEffect_16x", Icon16x16));
-		Set("ClassIcon.LandscapeComponent", new IMAGE_BRUSH("Icons/AssetIcons/Landscape_16x", Icon16x16));
 		Set("ClassIcon.LightComponent", new IMAGE_BRUSH("Icons/ActorIcons/LightActor_16x", Icon16x16));
-		Set("ClassIcon.ParticleSystemComponent", new IMAGE_BRUSH("Icons/AssetIcons/ParticleSystem_16x", Icon16x16));
-		Set("ClassIcon.PointLightComponent", new IMAGE_BRUSH("Icons/AssetIcons/PointLight_16x", Icon16x16));
-		Set("ClassIcon.RB_RadialForceComponent", new IMAGE_BRUSH("Icons/AssetIcons/RadialForceActor_16x", Icon16x16));
-		Set("ClassIcon.SingleAnimSkeletalComponent", new IMAGE_BRUSH("Icons/AssetIcons/SkeletalMesh_16x", Icon16x16));
-		Set("ClassIcon.SkeletalMeshComponent", new IMAGE_BRUSH("Icons/AssetIcons/SkeletalMesh_16x", Icon16x16));
-		Set("ClassIcon.SpotLightComponent", new IMAGE_BRUSH("Icons/AssetIcons/SpotLight_16x", Icon16x16));
-		Set("ClassIcon.StaticMeshComponent", new IMAGE_BRUSH("Icons/AssetIcons/StaticMesh_16x", Icon16x16));
-		Set("ClassIcon.VectorFieldComponent", new IMAGE_BRUSH("Icons/ActorIcons/VectorFieldVolume_16x", Icon16x16));
 		Set("ClassIcon.ArrowComponent", new IMAGE_BRUSH("Icons/ActorIcons/Arrow_16px", Icon16x16));
-		Set("ClassIcon.AtmosphericFogComponent", new IMAGE_BRUSH("Icons/AssetIcons/AtmosphericFog_16x", Icon16x16));
-		Set("ClassIcon.SkyAtmosphereComponent", new IMAGE_BRUSH("Icons/AssetIcons/SkyAtmosphere_16x", Icon16x16));
-		Set("ClassIcon.VolumetricCloudComponent", new IMAGE_BRUSH("Icons/AssetIcons/SkyAtmosphere_16x", Icon16x16));
 		Set("ClassIcon.BoxComponent", new IMAGE_BRUSH("Icons/ActorIcons/Box_16px", Icon16x16));
 		Set("ClassIcon.CapsuleComponent", new IMAGE_BRUSH("Icons/ActorIcons/Capsule_16px", Icon16x16));
 		Set("ClassIcon.InstancedStaticMeshComponent", new IMAGE_BRUSH("Icons/ActorIcons/InstancedStaticMesh_16px", Icon16x16));
 		Set("ClassIcon.MaterialBillboardComponent", new IMAGE_BRUSH("Icons/ActorIcons/MaterialSprite_16px", Icon16x16));
-		Set("ClassIcon.SceneCaptureComponent2D", new IMAGE_BRUSH("Icons/AssetIcons/SceneCapture2D_16x", Icon16x16));
-		Set("ClassIcon.SceneCaptureComponent", new IMAGE_BRUSH("Icons/ActorIcons/SceneCapture_16px", Icon16x16));
 		Set("ClassIcon.SphereComponent", new IMAGE_BRUSH("Icons/ActorIcons/Sphere_16px", Icon16x16));
 		Set("ClassIcon.SplineComponent", new IMAGE_BRUSH("Icons/ActorIcons/Spline_16px", Icon16x16));
 		Set("ClassIcon.BillboardComponent", new IMAGE_BRUSH("Icons/ActorIcons/SpriteComponent_16px", Icon16x16));
-		Set("ClassIcon.TextRenderComponent", new IMAGE_BRUSH("Icons/AssetIcons/TextRenderActor_16x", Icon16x16));
 		Set("ClassIcon.TimelineComponent", new IMAGE_BRUSH("Icons/ActorIcons/TimelineComponent_16px", Icon16x16));
 		Set("ClassIcon.ChildActorComponent", new IMAGE_BRUSH("Icons/ActorIcons/ChildActorComponent_16px", Icon16x16));
-		Set("ClassIcon.ComponentMobilityStaticPip", new IMAGE_BRUSH("Icons/ActorIcons/ComponentMobilityStationary_7x16px", Icon7x16, FLinearColor(0.f, 0.f, 0.f, 0.f)));
-		Set("ClassIcon.ComponentMobilityStationaryPip", new IMAGE_BRUSH("Icons/ActorIcons/ComponentMobilityStationary_7x16px", Icon7x16));
-		Set("ClassIcon.ComponentMobilityMovablePip", new IMAGE_BRUSH("Icons/ActorIcons/ComponentMobilityMovable_7x16px", Icon7x16));
+
+		Set("ClassIcon.AtmosphericFogComponent", new IMAGE_BRUSH_SVG("Starship/AssetIcons/AtmosphericFog_16", Icon16x16));
+		Set("ClassIcon.BrushComponent", new IMAGE_BRUSH_SVG("Starship/AssetIcons/Brush_16", Icon16x16));
+		Set("ClassIcon.CableComponent", new IMAGE_BRUSH_SVG("Starship/AssetIcons/CableActor_16", Icon16x16));
+		Set("ClassIcon.CameraComponent", new IMAGE_BRUSH_SVG("Starship/AssetIcons/CameraActor_16", Icon16x16));
+		Set("ClassIcon.DecalComponent", new IMAGE_BRUSH_SVG("Starship/AssetIcons/DecalActor_16", Icon16x16));
+		Set("ClassIcon.DirectionalLightComponent", new IMAGE_BRUSH_SVG("Starship/AssetIcons/DirectionalLight_16", Icon16x16));
+		Set("ClassIcon.ExponentialHeightFogComponent", new IMAGE_BRUSH_SVG("Starship/AssetIcons/ExponentialHeightFog_16", Icon16x16));
+		Set("ClassIcon.ForceFeedbackComponent", new IMAGE_BRUSH_SVG("Starship/AssetIcons/ForceFeedbackEffect_16", Icon16x16));
+		Set("ClassIcon.LandscapeComponent", new IMAGE_BRUSH_SVG("Starship/AssetIcons/Landscape_16", Icon16x16));
+		Set("ClassIcon.ParticleSystemComponent", new IMAGE_BRUSH_SVG("Starship/AssetIcons/Emitter_16", Icon16x16));
+		Set("ClassIcon.PlanarReflectionComponent", new IMAGE_BRUSH_SVG("Starship/AssetIcons/PlaneReflectionCapture_16", Icon16x16));
+		Set("ClassIcon.PointLightComponent", new IMAGE_BRUSH_SVG("Starship/AssetIcons/PointLight_16", Icon16x16));
+		Set("ClassIcon.RectLightComponent", new IMAGE_BRUSH_SVG("Starship/AssetIcons/RectLight_16", Icon16x16));
+		Set("ClassIcon.RadialForceComponent", new IMAGE_BRUSH_SVG("Starship/AssetIcons/RadialForceActor_16", Icon16x16));
+		Set("ClassIcon.SceneCaptureComponent2D", new IMAGE_BRUSH_SVG("Starship/AssetIcons/SceneCapture2D_16", Icon16x16));
+		Set("ClassIcon.SceneCaptureComponent", new IMAGE_BRUSH_SVG("Starship/AssetIcons/SphereReflectionCapture_16", Icon16x16));
+		Set("ClassIcon.SingleAnimSkeletalComponent", new IMAGE_BRUSH_SVG("Starship/AssetIcons/SkeletalMesh_16", Icon16x16));
+		Set("ClassIcon.SkyAtmosphereComponent", new IMAGE_BRUSH_SVG("Starship/AssetIcons/SkyAtmosphere_16", Icon16x16));
+		Set("ClassIcon.SkeletalMeshComponent", new IMAGE_BRUSH_SVG("Starship/AssetIcons/SkeletalMesh_16", Icon16x16));
+		Set("ClassIcon.SpotLightComponent", new IMAGE_BRUSH_SVG("Starship/AssetIcons/SpotLight_16", Icon16x16));
+		Set("ClassIcon.StaticMeshComponent", new IMAGE_BRUSH_SVG("Starship/AssetIcons/StaticMesh_16", Icon16x16));
+		Set("ClassIcon.TextRenderComponent", new IMAGE_BRUSH_SVG("Starship/AssetIcons/TextRenderActor_16", Icon16x16));
+		Set("ClassIcon.VectorFieldComponent", new IMAGE_BRUSH_SVG("Starship/AssetIcons/VectorFieldVolume_16", Icon16x16));
+		Set("ClassIcon.VolumetricCloudComponent", new IMAGE_BRUSH_SVG("Starship/AssetIcons/VolumetricCloud_16", Icon16x16));
+
 		Set("ClassIcon.MovableMobilityIcon", new IMAGE_BRUSH("Icons/ActorIcons/Light_Movable_16x", Icon16x16));
 		Set("ClassIcon.StationaryMobilityIcon", new IMAGE_BRUSH("Icons/ActorIcons/Light_Adjustable_16x", Icon16x16));
 		Set("ClassIcon.ComponentMobilityHeaderIcon", new IMAGE_BRUSH("Icons/ActorIcons/ComponentMobilityHeader_7x16", Icon7x16));
 
-		// Factory classes
-		//Set( "ClassIcon.ActorFactoryBoxVolume", new IMAGE_BRUSH( "Icons/icon_volume_Box_16x", Icon16x16 ) );
-		//Set( "ClassIcon.ActorFactoryCylinderVolume", new IMAGE_BRUSH( "Icons/icon_volume_cylinder_16x", Icon16x16 ) );
-		//Set( "ClassIcon.ActorFactorySphereVolume", new IMAGE_BRUSH( "Icons/icon_volume_sphere_16x", Icon16x16 ) );
-
 		// Asset Type Classes
 		const TCHAR* AssetTypes[] = {
 
@@ -5895,7 +5938,6 @@ void FStarshipEditorStyle::FStyle::SetupClassIconsAndThumbnails()
 			TEXT("AIPerceptionComponent"),
 			TEXT("BlueprintInterface"),
 			TEXT("BlueprintMacroLibrary"),
-			TEXT("CableComponent"),
 			TEXT("CameraAnim"),
 			TEXT("Default"),
 			TEXT("DirectionalLightMovable"),
@@ -6035,7 +6077,7 @@ void FStarshipEditorStyle::FStyle::SetupClassIconsAndThumbnails()
 			{TEXT("PhysicsHandleComponent")},
 			{TEXT("PhysicsThruster")},
 			{TEXT("PhysicsVolume")},
-			{TEXT("PlaneReflectionCapture")},
+			{TEXT("PlanarReflectionCapture")},
 			{TEXT("PlatformMediaSource")},
 			{TEXT("PlayerController")},
 			{TEXT("PlayerStart")},
@@ -6048,7 +6090,6 @@ void FStarshipEditorStyle::FStyle::SetupClassIconsAndThumbnails()
 			{TEXT("ProceduralFoliageVolume")},
 			{TEXT("ProjectileMovementComponent")},
 			{TEXT("RadialForceActor")},
-			{TEXT("RadialForceComponent")},
 			{TEXT("RectLight")},
 			{TEXT("ReflectionCapture")},
 			{TEXT("ReverbEffect")},
diff --git a/Engine/Source/Editor/UnrealEd/Private/SComponentClassCombo.cpp b/Engine/Source/Editor/UnrealEd/Private/SComponentClassCombo.cpp
--- a/Engine/Source/Editor/UnrealEd/Private/SComponentClassCombo.cpp
+++ b/Engine/Source/Editor/UnrealEd/Private/SComponentClassCombo.cpp
@@ -19,6 +19,7 @@
 #include "SListViewSelectorDropdownMenu.h"
 #include "Misc/TextFilterExpressionEvaluator.h"
 #include "SEditorHeaderButton.h"
+#include "Widgets/Layout/SSeparator.h"
 
 #define LOCTEXT_NAMESPACE "ComponentClassCombo"
 
@@ -319,13 +320,9 @@ TSharedRef<ITableRow> SComponentClassCombo::GenerateAddComponentRow( FComponentC
 				.Style(&FEditorStyle::Get().GetWidgetStyle<FTableRowStyle>("TableView.NoHoverTableRow"))
 				.ShowSelection(false)
 			[
-				SNew(SBox)
-				.Padding(1.f)
-				[
-					SNew(SBorder)
-					.Padding(FEditorStyle::GetMargin(TEXT("Menu.Separator.Padding")))
-					.BorderImage(FEditorStyle::GetBrush(TEXT("Menu.Separator")))
-				]
+				SNew(SSeparator)
+				.SeparatorImage(FEditorStyle::Get().GetBrush("Menu.Separator"))
+				.Thickness(1.0f)
 			];
 	}
 	else
diff --git a/Engine/Source/Runtime/Engine/Classes/Components/RectLightComponent.h b/Engine/Source/Runtime/Engine/Classes/Components/RectLightComponent.h
--- a/Engine/Source/Runtime/Engine/Classes/Components/RectLightComponent.h
+++ b/Engine/Source/Runtime/Engine/Classes/Components/RectLightComponent.h
@@ -16,7 +16,7 @@ class FLightSceneProxy;
 /**
  * A light component which emits light from a rectangle.
  */
-UCLASS(Blueprintable, ClassGroup=(Lights,Common), hidecategories=(Object, LightShafts), editinlinenew, meta=(BlueprintSpawnableComponent))
+UCLASS(Blueprintable, ClassGroup=(Lights), hidecategories=(Object, LightShafts), editinlinenew, meta=(BlueprintSpawnableComponent))
 class ENGINE_API URectLightComponent : public ULocalLightComponent
 {
 	GENERATED_UCLASS_BODY()
diff --git a/Engine/Source/Runtime/Engine/Classes/Particles/ParticleSystemComponent.h b/Engine/Source/Runtime/Engine/Classes/Particles/ParticleSystemComponent.h
--- a/Engine/Source/Runtime/Engine/Classes/Particles/ParticleSystemComponent.h
+++ b/Engine/Source/Runtime/Engine/Classes/Particles/ParticleSystemComponent.h
@@ -425,7 +425,7 @@ public:
 /** 
  * A particle emitter.
  */
-UCLASS(ClassGroup=(Rendering, Common), hidecategories=Object, hidecategories=Physics, hidecategories=Collision, showcategories=Trigger, editinlinenew, meta=(BlueprintSpawnableComponent))
+UCLASS(ClassGroup=(Rendering), hidecategories=Object, hidecategories=Physics, hidecategories=Collision, showcategories=Trigger, editinlinenew, meta=(BlueprintSpawnableComponent))
 class ENGINE_API UParticleSystemComponent : public UFXSystemComponent
 {
 	friend class FParticleSystemWorldManager;
diff --git a/Engine/Source/Runtime/Experimental/GeometryCollectionEngine/Public/ChaosBlueprint.h b/Engine/Source/Runtime/Experimental/GeometryCollectionEngine/Public/ChaosBlueprint.h
--- a/Engine/Source/Runtime/Experimental/GeometryCollectionEngine/Public/ChaosBlueprint.h
+++ b/Engine/Source/Runtime/Experimental/GeometryCollectionEngine/Public/ChaosBlueprint.h
@@ -33,7 +33,7 @@ DECLARE_DYNAMIC_MULTICAST_DELEGATE_OneParam(FOnChaosTrailingEvents, const TArray
 
 
 /** Object allowing for retrieving Chaos Destruction data. */
-UCLASS(ClassGroup = (Chaos, Common), hidecategories = (Object, ActorComponent, Physics, Rendering, Mobility, LOD), ShowCategories = Trigger, meta = (BlueprintSpawnableComponent))
+UCLASS(ClassGroup = (Chaos), hidecategories = (Object, ActorComponent, Physics, Rendering, Mobility, LOD), ShowCategories = Trigger, meta = (BlueprintSpawnableComponent))
 class GEOMETRYCOLLECTIONENGINE_API UChaosDestructionListener : public USceneComponent
 {
 	GENERATED_UCLASS_BODY()
