Cleanup progress bar styles and fix slow task not showing the secondary progress bars

diff --git a/Engine/Source/Developer/Profiler/Private/Widgets/SDataGraph.cpp b/Engine/Source/Developer/Profiler/Private/Widgets/SDataGraph.cpp
--- a/Engine/Source/Developer/Profiler/Private/Widgets/SDataGraph.cpp
+++ b/Engine/Source/Developer/Profiler/Private/Widgets/SDataGraph.cpp
@@ -841,7 +841,7 @@ int32 SDataGraph::OnPaint( const FPaintArgs& Args, const FGeometry& AllottedGeom
 					OutDrawElements,
 					LayerId,
 					AllottedGeometry.ToPaintGeometry( FVector2D(LocalGraphSelectionX[Nx]-HalfGraphMarkerWidth,0.0f), FVector2D(GraphMarkerWidth, AllottedGeometry.GetLocalSize().Y) ),
-					FEditorStyle::GetBrush("ProgressBar.Background"),
+					FEditorStyle::GetBrush("Brushes.Foldout"),
 					DrawEffects,
 					FColor(64,64,255,128)
 				);
@@ -863,7 +863,7 @@ int32 SDataGraph::OnPaint( const FPaintArgs& Args, const FGeometry& AllottedGeom
 					OutDrawElements,
 					LayerId,
 					AllottedGeometry.ToPaintGeometry( FVector2D(GraphSelectionX0/*+HalfGraphMarkerWidth*/,0.0f), FVector2D(GraphSelectionW/*-GraphMarkerWidth*/, AllottedGeometry.GetLocalSize().Y) ),
-					FEditorStyle::GetBrush("ProgressBar.Background"),
+					FEditorStyle::GetBrush("Brushes.Foldout"),
 					DrawEffects,
 					FColor(64,64,255,32)
 				);
diff --git a/Engine/Source/Editor/ContentBrowser/Private/SAssetView.cpp b/Engine/Source/Editor/ContentBrowser/Private/SAssetView.cpp
--- a/Engine/Source/Editor/ContentBrowser/Private/SAssetView.cpp
+++ b/Engine/Source/Editor/ContentBrowser/Private/SAssetView.cpp
@@ -354,7 +354,6 @@ void SAssetView::Construct( const FArguments& InArgs )
 			[
 				SNew( SProgressBar )
 				.Percent( this, &SAssetView::GetIsWorkingProgressBarState )
-				.Style( FEditorStyle::Get(), "WorkingBar" )
 				.BorderPadding( FVector2D(0,0) )
 			]
 		]
diff --git a/Engine/Source/Editor/EditorStyle/Private/StarshipStyle.cpp b/Engine/Source/Editor/EditorStyle/Private/StarshipStyle.cpp
--- a/Engine/Source/Editor/EditorStyle/Private/StarshipStyle.cpp
+++ b/Engine/Source/Editor/EditorStyle/Private/StarshipStyle.cpp
@@ -609,29 +609,6 @@ void FStarshipEditorStyle::FStyle::SetupGeneralStyles()
 
 		Set( "SearchBox.ActiveBorder", new BOX_BRUSH( "Common/TextBox_Special_Active", FMargin(8.0f/32.0f) ) );
 	}
-	
-	// ProgressBar
-	{
-
-
-		Set( "ProgressBar.ThinBackground", new BOX_BRUSH( "Common/ProgressBar_Thin_Background", FMargin(5.f/12.f) ) );
-		Set( "ProgressBar.ThinFill", new BOX_BRUSH( "Common/ProgressBar_Thin_Fill", FMargin(5.f/12.f) ) );
-
-		// Legacy ProgressBar styles; kept here because other FEditorStyle controls (mis)use them
-		// todo: jdale - Widgets using these styles should be updated to use SlateStyle types once FEditorStyle has been obliterated from the Slate core
-		Set( "ProgressBar.Background", new BOX_BRUSH( "Common/ProgressBar_Background", FMargin(5.f/12.f) ) );
-		Set( "ProgressBar.Marquee", new IMAGE_BRUSH( "Common/ProgressBar_Marquee", FVector2D(20,12), FLinearColor::White, ESlateBrushTileType::Horizontal ) );
-		Set( "ProgressBar.BorderPadding", FVector2D(1,0) );
-	}
-
-	// WorkingBar
-	{
-		Set( "WorkingBar", FProgressBarStyle()
-			.SetBackgroundImage( FSlateNoResource() )
-			.SetFillImage( BOX_BRUSH( "Common/ProgressBar_Fill", FMargin(5.f/12.f), FLinearColor( 1.0f, 0.22f, 0.0f ) ) )
-			.SetMarqueeImage( IMAGE_BRUSH( "Common/WorkingBar_Marquee", FVector2D(20,2), FLinearColor::White, ESlateBrushTileType::Horizontal ) )
-			);
-	}
 
 	// Filtering/Searching feedback
 	{
diff --git a/Engine/Source/Editor/UnrealEd/Private/AutoReimport/ReimportFeedbackContext.cpp b/Engine/Source/Editor/UnrealEd/Private/AutoReimport/ReimportFeedbackContext.cpp
--- a/Engine/Source/Editor/UnrealEd/Private/AutoReimport/ReimportFeedbackContext.cpp
+++ b/Engine/Source/Editor/UnrealEd/Private/AutoReimport/ReimportFeedbackContext.cpp
@@ -326,8 +326,6 @@ public:
 						SAssignNew(ProgressBar, SProgressBar)
 						.BorderPadding(FVector2D::ZeroVector)
 						.Percent( this, &SReimportFeedback::GetProgressFraction )
-						.BackgroundImage( FEditorStyle::GetBrush("ProgressBar.ThinBackground") )
-						.FillImage( FEditorStyle::GetBrush("ProgressBar.ThinFill") )
 					]
 				]
 
diff --git a/Engine/Source/Editor/UnrealEd/Private/FeedbackContextEditor.cpp b/Engine/Source/Editor/UnrealEd/Private/FeedbackContextEditor.cpp
--- a/Engine/Source/Editor/UnrealEd/Private/FeedbackContextEditor.cpp
+++ b/Engine/Source/Editor/UnrealEd/Private/FeedbackContextEditor.cpp
@@ -189,7 +189,7 @@ private:
 		// Always show the first one
 		DynamicProgressIndices.Add(0);
 
-		for (int32 Index = 1; Index < ScopeStack->Num() && DynamicProgressIndices.Num() <= MaxNumSecondaryBars - 1; ++Index)
+		for (int32 Index = 1; Index < ScopeStack->Num() && DynamicProgressIndices.Num() <= MaxNumSecondaryBars; ++Index)
 		{
 			const auto* Scope = (*ScopeStack)[Index];
 
