Cleanup to Starship Style to use the right images for Android PreviewModes to match changes done in CL14141952

diff --git a/Engine/Source/Editor/EditorStyle/Private/StarshipStyle.cpp b/Engine/Source/Editor/EditorStyle/Private/StarshipStyle.cpp
--- a/Engine/Source/Editor/EditorStyle/Private/StarshipStyle.cpp
+++ b/Engine/Source/Editor/EditorStyle/Private/StarshipStyle.cpp
@@ -4492,33 +4492,19 @@ void FStarshipEditorStyle::FStyle::SetupLevelEditorStyle()
 		Set( "LevelEditor.Recompile.Small", new IMAGE_BRUSH_SVG( "Starship/MainToolbar/compile", Icon20x20 ) );
 
 		Set("LevelEditor.PreviewMode.Enabled", new IMAGE_BRUSH("Icons/icon_PreviewMode_SM5_Enabled_40x", Icon40x40));
-		Set("LevelEditor.PreviewMode.Enabled.Small", new IMAGE_BRUSH("Icons/icon_PreviewMode_SM5_Enabled_40x", Icon20x20));
 		Set("LevelEditor.PreviewMode.Disabled", new IMAGE_BRUSH("Icons/icon_PreviewMode_SM5_Disabled_40x", Icon40x40));
-		Set("LevelEditor.PreviewMode.Disabled.Small", new IMAGE_BRUSH("Icons/icon_PreviewMode_SM5_Disabled_40x", Icon20x20));
 		Set("LevelEditor.PreviewMode.SM5.Enabled", new IMAGE_BRUSH("Icons/icon_PreviewMode_SM5_Enabled_40x", Icon40x40));
-		Set("LevelEditor.PreviewMode.SM5.Enabled.Small", new IMAGE_BRUSH("Icons/icon_PreviewMode_SM5_Enabled_40x", Icon20x20));
 		Set("LevelEditor.PreviewMode.SM5.Disabled", new IMAGE_BRUSH("Icons/icon_PreviewMode_SM5_Enabled_40x", Icon40x40));
-		Set("LevelEditor.PreviewMode.SM5.Disabled.Small", new IMAGE_BRUSH("Icons/icon_PreviewMode_SM5_Enabled_40x", Icon20x20));
-		Set("LevelEditor.PreviewMode.SM4.Enabled", new IMAGE_BRUSH("Icons/icon_PreviewMode_SM4_Enabled_40x", Icon40x40));
-		Set("LevelEditor.PreviewMode.SM4.Enabled.Small", new IMAGE_BRUSH("Icons/icon_PreviewMode_SM4_Enabled_40x", Icon20x20));
-		Set("LevelEditor.PreviewMode.SM4.Disabled", new IMAGE_BRUSH("Icons/icon_PreviewMode_SM4_Disabled_40x", Icon40x40));
-		Set("LevelEditor.PreviewMode.SM4.Disabled.Small", new IMAGE_BRUSH("Icons/icon_PreviewMode_SM4_Disabled_40x", Icon20x20));
-		Set("LevelEditor.PreviewMode.AndroidES2.Enabled", new IMAGE_BRUSH("Icons/icon_PreviewMode_AndroidES2_Enabled_40x", Icon40x40));
-		Set("LevelEditor.PreviewMode.AndroidES2.Enabled.Small", new IMAGE_BRUSH("Icons/icon_PreviewMode_AndroidES2_Enabled_40x", Icon20x20));
-		Set("LevelEditor.PreviewMode.AndroidES2.Disabled", new IMAGE_BRUSH("Icons/icon_PreviewMode_AndroidES2_Disabled_40x", Icon40x40));
-		Set("LevelEditor.PreviewMode.AndroidES2.Disabled.Small", new IMAGE_BRUSH("Icons/icon_PreviewMode_AndroidES2_Disabled_40x", Icon20x20));
 		Set("LevelEditor.PreviewMode.AndroidES31.Enabled", new IMAGE_BRUSH("Icons/icon_PreviewMode_AndroidES31_Enabled_40x", Icon40x40));
-		Set("LevelEditor.PreviewMode.AndroidES31.Enabled.Small", new IMAGE_BRUSH("Icons/icon_PreviewMode_AndroidES31_Enabled_40x", Icon20x20));
 		Set("LevelEditor.PreviewMode.AndroidES31.Disabled", new IMAGE_BRUSH("Icons/icon_PreviewMode_AndroidES31_Disabled_40x", Icon40x40));
-		Set("LevelEditor.PreviewMode.AndroidES31.Disabled.Small", new IMAGE_BRUSH("Icons/icon_PreviewMode_AndroidES31_Disabled_40x", Icon20x20));
 		Set("LevelEditor.PreviewMode.AndroidVulkan.Enabled", new IMAGE_BRUSH("Icons/icon_PreviewMode_AndroidVulkan_Enabled_40x", Icon40x40));
-		Set("LevelEditor.PreviewMode.AndroidVulkan.Enabled.Small", new IMAGE_BRUSH("Icons/icon_PreviewMode_AndroidVulkan_Enabled_40x", Icon20x20));
 		Set("LevelEditor.PreviewMode.AndroidVulkan.Disabled", new IMAGE_BRUSH("Icons/icon_PreviewMode_AndroidVulkan_Disabled_40x", Icon40x40));
-		Set("LevelEditor.PreviewMode.AndroidVulkan.Disabled.Small", new IMAGE_BRUSH("Icons/icon_PreviewMode_AndroidVulkan_Disabled_40x", Icon20x20));
+		Set("LevelEditor.PreviewMode.AndroidVulkanSM5.Enabled", new IMAGE_BRUSH("Icons/icon_PreviewMode_AndroidVulkanSM5_Enabled_40x", Icon40x40));
+		Set("LevelEditor.PreviewMode.AndroidVulkanSM5.Disabled", new IMAGE_BRUSH("Icons/icon_PreviewMode_AndroidVulkanSM5_Disabled_40x", Icon40x40));
 		Set("LevelEditor.PreviewMode.iOS.Enabled", new IMAGE_BRUSH("Icons/icon_PreviewMode_iOS_Enabled_40x", Icon40x40));
-		Set("LevelEditor.PreviewMode.iOS.Enabled.Small", new IMAGE_BRUSH("Icons/icon_PreviewMode_iOS_Enabled_40x", Icon20x20));
 		Set("LevelEditor.PreviewMode.iOS.Disabled", new IMAGE_BRUSH("Icons/icon_PreviewMode_iOS_Disabled_40x", Icon40x40));
-		Set("LevelEditor.PreviewMode.iOS.Disabled.Small", new IMAGE_BRUSH("Icons/icon_PreviewMode_iOS_Disabled_40x", Icon20x20));
+		Set("LevelEditor.PreviewMode.iOSSM5.Enabled", new IMAGE_BRUSH("Icons/icon_PreviewMode_iOSSM5_Enabled_40x", Icon40x40));
+		Set("LevelEditor.PreviewMode.iOSSM5.Disabled", new IMAGE_BRUSH("Icons/icon_PreviewMode_iOSSM5_Disabled_40x", Icon40x40));
 
 		Set("LevelEditor.ViewOptions", new IMAGE_BRUSH("Icons/icon_view_40x", Icon40x40));
 		Set( "LevelEditor.ViewOptions.Small", new IMAGE_BRUSH( "Icons/icon_view_40x", Icon20x20 ) );
