Remove old project dialog styles

diff --git a/Engine/Source/Editor/EditorStyle/Private/StarshipStyle.cpp b/Engine/Source/Editor/EditorStyle/Private/StarshipStyle.cpp
--- a/Engine/Source/Editor/EditorStyle/Private/StarshipStyle.cpp
+++ b/Engine/Source/Editor/EditorStyle/Private/StarshipStyle.cpp
@@ -1583,83 +1583,9 @@ void FStarshipEditorStyle::FStyle::SetupGeneralStyles()
 	// GameProjectDialog
 	if (IncludeEditorSpecificStyles())
 	{
-		Set( "GameProjectDialog.NewProjectTitle", FTextBlockStyle(NormalText) 
-			.SetFont( DEFAULT_FONT( "BoldCondensed", 28 ) )
-			.SetShadowOffset( FVector2D( 1,1 ) )
-			.SetShadowColorAndOpacity( FLinearColor(0,0,0,0.9f) )
-		);
-
-		Set( "GameProjectDialog.RecentProjectsTitle", FTextBlockStyle(NormalText) 
-			.SetFont( DEFAULT_FONT( "BoldCondensed", 16 ) )
-			.SetColorAndOpacity( FLinearColor(0.8f,0.8f,0.8f,1) )
-			.SetShadowOffset( FVector2D( 0,1 ) )
-			.SetShadowColorAndOpacity( FLinearColor(0,0,0,0.9f) )
-		);
-
-		Set( "GameProjectDialog.ProjectNamePathLabels", FTextBlockStyle(NormalText)
-			.SetFont( DEFAULT_FONT( "Regular", 12 ) )
-			.SetShadowOffset( FVector2D( 0,1 ) )
-			.SetShadowColorAndOpacity( FLinearColor(0,0,0,0.9f) )
-		);
-
-		Set( "GameProjectDialog.ErrorLabelFont", FTextBlockStyle(NormalText)
-			.SetFont( DEFAULT_FONT( "Regular", 10 ) )
-			.SetShadowOffset( FVector2D::ZeroVector )
-		);
-
-		Set( "GameProjectDialog.ErrorLabelBorder", new FSlateColorBrush( FLinearColor(0.2f, 0.f, 0.f, 0.7f) ) );
-		Set( "GameProjectDialog.ErrorLabelCloseButton", new IMAGE_BRUSH( "Icons/Cross_12x", Icon12x12 ) );
-
-		Set( "GameProjectDialog.TemplateListView.TableRow", FTableRowStyle()
-			.SetEvenRowBackgroundBrush( FSlateNoResource() )
-			.SetEvenRowBackgroundHoveredBrush( IMAGE_BRUSH( "Common/Selection", Icon8x8, FLinearColor(1.0f, 1.0f, 1.0f, 0.1f) ) )
-			.SetOddRowBackgroundBrush( FSlateNoResource() )
-			.SetOddRowBackgroundHoveredBrush( IMAGE_BRUSH( "Common/Selection", Icon8x8, FLinearColor(1.0f, 1.0f, 1.0f, 0.1f) ) )
-			.SetSelectorFocusedBrush( BORDER_BRUSH( "Common/Selector", FMargin(4.f/16.f), SelectorColor ) )
-			.SetActiveBrush( IMAGE_BRUSH( "Common/Selection", Icon8x8, SelectionColor ) )
-			.SetActiveHoveredBrush( IMAGE_BRUSH( "Common/Selection", Icon8x8, SelectionColor ) )
-			.SetInactiveBrush( IMAGE_BRUSH( "Common/Selection", Icon8x8, SelectionColor ) )
-			.SetInactiveHoveredBrush( IMAGE_BRUSH( "Common/Selection", Icon8x8, SelectionColor )  )
-			.SetTextColor( DefaultForeground )
-			.SetSelectedTextColor( InvertedForeground )
-			);
-
 		Set( "GameProjectDialog.DefaultGameThumbnail", new IMAGE_BRUSH( "GameProjectDialog/default_game_thumbnail", Icon64x64 ) );
 		Set( "GameProjectDialog.BlankProjectThumbnail", new IMAGE_BRUSH( "GameProjectDialog/blank_project_thumbnail", Icon128x128 ) );
 		Set( "GameProjectDialog.BlankProjectPreview", new IMAGE_BRUSH( "GameProjectDialog/blank_project_preview", FVector2D(400, 200) ) );
-		Set( "GameProjectDialog.BasicCodeThumbnail", new IMAGE_BRUSH( "GameProjectDialog/basic_code_thumbnail", Icon128x128 ) );
-		Set( "GameProjectDialog.CodeIcon", new IMAGE_BRUSH( "GameProjectDialog/feature_code_32x", FVector2D(32,32) ) );
-		Set( "GameProjectDialog.CodeImage", new IMAGE_BRUSH( "GameProjectDialog/feature_code", FVector2D(96,96) ) );
-		Set( "GameProjectDialog.CodeImage_64", new IMAGE_BRUSH( "GameProjectDialog/feature_code", Icon64x64) );
-		Set( "GameProjectDialog.CodeImage_48x", new IMAGE_BRUSH("GameProjectDialog/feature_code", Icon40x40));
-		Set( "GameProjectDialog.BlueprintIcon", new IMAGE_BRUSH( "GameProjectDialog/feature_blueprint_32x", FVector2D(32,32) ) );
-		Set( "GameProjectDialog.BlueprintImage", new IMAGE_BRUSH( "GameProjectDialog/feature_blueprint", FVector2D(96,96) ) );
-		Set( "GameProjectDialog.BlueprintImage_64", new IMAGE_BRUSH( "GameProjectDialog/feature_blueprint", Icon64x64) );
-		Set( "GameProjectDialog.CodeBorder", new BOX_BRUSH( "GameProjectDialog/feature_border", FMargin(4.0f/16.0f), FLinearColor(0.570, 0.359, 0.081, 1.f) ) );
-		Set( "GameProjectDialog.FeatureText", FTextBlockStyle(NormalText)
-			.SetFont( DEFAULT_FONT( "Regular", 14 ) )
-			.SetShadowOffset( FVector2D( 0,1 ) )
-			.SetShadowColorAndOpacity( FLinearColor(0,0,0,0.9f) ) );
-		Set( "GameProjectDialog.TemplateItemTitle", FTextBlockStyle(NormalText)
-			.SetFont( DEFAULT_FONT( "Regular", 10 ) )
-			.SetShadowOffset( FVector2D( 0,1 ) )
-			.SetShadowColorAndOpacity( FLinearColor(0,0,0,0.9f) ) );
-		
-		Set( "GameProjectDialog.Tab", FCheckBoxStyle()
-			.SetCheckBoxType( ESlateCheckBoxType::ToggleButton )
-			.SetUncheckedImage( BOX_BRUSH( "/GameProjectDialog/Tab_Inactive", 4 / 16.0f ) )
-			.SetUncheckedPressedImage( BOX_BRUSH( "/GameProjectDialog/Tab_Active", 4 / 16.0f ) )
-			.SetUncheckedHoveredImage( BOX_BRUSH( "/GameProjectDialog/Tab_Active", 4 / 16.0f ) )
-			.SetCheckedHoveredImage( BOX_BRUSH( "/GameProjectDialog/Tab_Active", 4 / 16.0f ) )
-			.SetCheckedPressedImage( BOX_BRUSH( "/GameProjectDialog/Tab_Active", 4 / 16.0f ) )
-			.SetCheckedImage( BOX_BRUSH( "/GameProjectDialog/Tab_Active", 4 / 16.0f ) )
-			);
-
-		Set( "GameProjectDialog.TabBackground", new BOX_BRUSH( "/GameProjectDialog/tab_background", 4 / 16.0f ) );
-
-		Set( "GameProjectDialog.FolderIconClosed", new IMAGE_BRUSH( "Icons/FolderClosed", FVector2D(18, 16) ) );
-		Set( "GameProjectDialog.FolderIconOpen", new IMAGE_BRUSH( "Icons/FolderOpen", FVector2D(18, 16) ) );
-		Set( "GameProjectDialog.ProjectFileIcon", new IMAGE_BRUSH( "Icons/doc_16x", FVector2D(18, 16) ) );
 
 	}
 
