Adding a new custom style for comment bubbles and node title text editing that matches the current graph editor look.

diff --git a/Engine/Source/Editor/EditorStyle/Private/StarshipStyle.cpp b/Engine/Source/Editor/EditorStyle/Private/StarshipStyle.cpp
--- a/Engine/Source/Editor/EditorStyle/Private/StarshipStyle.cpp
+++ b/Engine/Source/Editor/EditorStyle/Private/StarshipStyle.cpp
@@ -3664,12 +3664,12 @@ void FStarshipEditorStyle::FStyle::SetupGraphEditorStyles()
 
 			Set( "BTEditor.Debugger.PausePlaySession", new IMAGE_BRUSH( "Icons/icon_pause_40x", Icon40x40 ) );
 			Set( "BTEditor.Debugger.PausePlaySession.Small", new IMAGE_BRUSH( "Icons/icon_pause_40x", Icon20x20 ) );
-			Set( "BTEditor.Debugger.ResumePlaySession", new IMAGE_BRUSH( "Icons/icon_simulate_40x", Icon40x40 ) );
-			Set( "BTEditor.Debugger.ResumePlaySession.Small", new IMAGE_BRUSH( "Icons/icon_simulate_40x", Icon20x20 ) );
+			Set("BTEditor.Debugger.ResumePlaySession", new IMAGE_BRUSH_SVG("Starship/MainToolbar/simulate", Icon40x40));
+			Set("BTEditor.Debugger.ResumePlaySession.Small", new IMAGE_BRUSH_SVG("Starship/MainToolbar/simulate", Icon20x20));
 			Set( "BTEditor.Debugger.StopPlaySession", new IMAGE_BRUSH( "Icons/icon_stop_40x", Icon40x40 ) );
 			Set( "BTEditor.Debugger.StopPlaySession.Small", new IMAGE_BRUSH( "Icons/icon_stop_40x", Icon20x20 ) );
-			Set( "BTEditor.Debugger.LateJoinSession", new IMAGE_BRUSH("Icons/icon_simulate_40x", Icon40x40) );
-			Set( "BTEditor.Debugger.LateJoinSession.Small", new IMAGE_BRUSH("Icons/icon_simulate_40x", Icon20x20) );
+			Set("BTEditor.Debugger.LateJoinSession", new IMAGE_BRUSH_SVG("Starship/MainToolbar/simulate", Icon40x40));
+			Set("BTEditor.Debugger.LateJoinSession.Small", new IMAGE_BRUSH_SVG("Starship/MainToolbar/simulate", Icon20x20));
 
 			Set( "BTEditor.Debugger.CurrentValues", new IMAGE_BRUSH( "BehaviorTree/Debugger_Current_40x", Icon40x40 ) );
 			Set( "BTEditor.Debugger.CurrentValues.Small", new IMAGE_BRUSH( "BehaviorTree/Debugger_Current_40x", Icon20x20 ) );
@@ -3751,17 +3751,18 @@ void FStarshipEditorStyle::FStyle::SetupGraphEditorStyles()
 		{
 			FTextBlockStyle GraphNodeTitle = FTextBlockStyle(NormalText)
 				.SetFont( DEFAULT_FONT( "Bold", 10 ) )
-				.SetColorAndOpacity( FLinearColor(218.0f/255.0f,218.0f/255.0f,218.0f/255.0f) )
+				.SetColorAndOpacity( FStyleColors::Foreground )
 				.SetShadowOffset( FVector2D::ZeroVector )
 				.SetShadowColorAndOpacity( FLinearColor(0.f,0.f,0.f, 0.7f) );
 			Set( "Graph.Node.NodeTitle", GraphNodeTitle );
 
-			FEditableTextBoxStyle GraphNodeTitleEditableText = FEditableTextBoxStyle()
+			FEditableTextBoxStyle GraphNodeTitleEditableText = FEditableTextBoxStyle(NormalEditableTextBoxStyle)
 				.SetFont(NormalText.Font)
-				.SetBackgroundImageNormal( BOX_BRUSH( "Common/TextBox", FMargin(4.0f/16.0f) ) )
-				.SetBackgroundImageHovered( BOX_BRUSH( "Common/TextBox_Hovered", FMargin(4.0f/16.0f) ) )
-				.SetBackgroundImageFocused( BOX_BRUSH( "Common/TextBox_Hovered", FMargin(4.0f/16.0f) ) )
-				.SetBackgroundImageReadOnly( BOX_BRUSH( "Common/TextBox_ReadOnly", FMargin(4.0f/16.0f) ) )
+				.SetForegroundColor(FStyleColors::Input)
+				.SetBackgroundImageNormal(FSlateRoundedBoxBrush(FStyleColors::Foreground, FStyleColors::Secondary, 1.0f))
+				.SetBackgroundImageHovered(FSlateRoundedBoxBrush(FStyleColors::Foreground, FStyleColors::Hover, 1.0f))
+				.SetBackgroundImageFocused(FSlateRoundedBoxBrush(FStyleColors::Foreground, FStyleColors::Primary, 1.0f))
+				.SetBackgroundImageReadOnly(FSlateRoundedBoxBrush(FStyleColors::Header, FStyleColors::InputOutline, 1.0f))
 				.SetScrollBarStyle( ScrollBar );
 			Set( "Graph.Node.NodeTitleEditableText", GraphNodeTitleEditableText );
 
@@ -3784,12 +3785,8 @@ void FStarshipEditorStyle::FStyle::SetupGraphEditorStyles()
 				.SetShadowColorAndOpacity( FLinearColor(0.f,0.f,0.f, 0.7f) );
 			Set( "Graph.CommentBlock.Title", GraphCommentBlockTitle );
 
-			FEditableTextBoxStyle GraphCommentBlockTitleEditableText = FEditableTextBoxStyle()
+			FEditableTextBoxStyle GraphCommentBlockTitleEditableText = FEditableTextBoxStyle(GraphNodeTitleEditableText)
 				.SetFont(GraphCommentBlockTitle.Font)
-				.SetBackgroundImageNormal( BOX_BRUSH( "Common/TextBox", FMargin(4.0f/16.0f) ) )
-				.SetBackgroundImageHovered( BOX_BRUSH( "Common/TextBox_Hovered", FMargin(4.0f/16.0f) ) )
-				.SetBackgroundImageFocused( BOX_BRUSH( "Common/TextBox_Hovered", FMargin(4.0f/16.0f) ) )
-				.SetBackgroundImageReadOnly( BOX_BRUSH( "Common/TextBox_ReadOnly", FMargin(4.0f/16.0f) ) )
 				.SetScrollBarStyle( ScrollBar );
 			Set( "Graph.CommentBlock.TitleEditableText", GraphCommentBlockTitleEditableText );
 
@@ -5251,7 +5248,7 @@ void FStarshipEditorStyle::FStyle::SetupPersonaStyle()
 		Set( "BlueprintEditor.ResetCamera", new IMAGE_BRUSH_SVG("Starship/Common/ResetCamera", Icon16x16));
 		Set( "BlueprintEditor.ShowFloor", new IMAGE_BRUSH_SVG( "Starship/Common/Floor", Icon16x16));
 		Set( "BlueprintEditor.ShowGrid", new IMAGE_BRUSH_SVG( "Starship/Common/Grid", Icon16x16 ) );
-		Set( "BlueprintEditor.EnableSimulation", new IMAGE_BRUSH( "Icons/icon_Enable_Simulation_40px", Icon40x40));
+		Set( "BlueprintEditor.EnableSimulation", new IMAGE_BRUSH_SVG("Starship/MainToolbar/simulate", Icon20x20));
 		Set( "BlueprintEditor.EnableProfiling", new IMAGE_BRUSH_SVG("Starship/Common/Statistics", Icon20x20) );
 		Set( "SCS.NativeComponent", new IMAGE_BRUSH( "Icons/NativeSCSComponent", Icon20x20 ));
 		Set( "SCS.Component", new IMAGE_BRUSH( "Icons/SCSComponent", Icon20x20 ));
@@ -5400,11 +5397,11 @@ void FStarshipEditorStyle::FStyle::SetupPersonaStyle()
 
 		Set("PlayToolBar", PlayToolbar);
 
-		Set( "PlayWorld.Simulate", new IMAGE_BRUSH( "Icons/icon_simulate_40x", Icon40x40 ) );
-		Set( "PlayWorld.Simulate.Small", new IMAGE_BRUSH( "Icons/icon_simulate_40x", Icon20x20 ) );
+		Set("PlayWorld.Simulate", new IMAGE_BRUSH_SVG("Starship/MainToolbar/simulate", Icon40x40));
+		Set("PlayWorld.Simulate.Small", new IMAGE_BRUSH_SVG("Starship/MainToolbar/simulate", Icon20x20));
 
-		Set( "PlayWorld.RepeatLastPlay", new IMAGE_BRUSH( "Icons/icon_simulate_40x", Icon40x40 ) );
-		Set( "PlayWorld.RepeatLastPlay.Small", new IMAGE_BRUSH( "Icons/icon_simulate_40x", Icon20x20 ) );
+		Set("PlayWorld.RepeatLastPlay", new IMAGE_BRUSH_SVG("Starship/MainToolbar/simulate", Icon40x40));
+		Set("PlayWorld.RepeatLastPlay.Small", new IMAGE_BRUSH_SVG("Starship/MainToolbar/simulate", Icon20x20));
 		Set( "PlayWorld.PlayInViewport", new IMAGE_BRUSH_SVG("Starship/MainToolbar/play", Icon20x20));
 
 		Set( "PlayWorld.PlayInEditorFloating", new IMAGE_BRUSH( "Icons/icon_playInWindow_40x", Icon40x40 ) );
@@ -5424,8 +5421,8 @@ void FStarshipEditorStyle::FStyle::SetupPersonaStyle()
 		Set( "PlayWorld.PlayInCameraLocation", new IMAGE_BRUSH( "Icons/icon_PlayCameraLocation_40x", Icon40x40 ) );
 		Set( "PlayWorld.PlayInDefaultPlayerStart", new IMAGE_BRUSH( "Icons/icon_PlayDefaultPlayerStart_40x", Icon40x40 ) );
 
-		Set( "PlayWorld.ResumePlaySession", new IMAGE_BRUSH( "Icons/icon_simulate_40x", Icon40x40 ) );
-		Set( "PlayWorld.ResumePlaySession.Small", new IMAGE_BRUSH( "Icons/icon_simulate_40x", Icon20x20 ) );
+		Set("PlayWorld.ResumePlaySession", new IMAGE_BRUSH_SVG("Starship/MainToolbar/simulate", Icon40x40));
+		Set("PlayWorld.ResumePlaySession.Small", new IMAGE_BRUSH_SVG("Starship/MainToolbar/simulate", Icon20x20));
 		Set( "PlayWorld.PausePlaySession", new IMAGE_BRUSH_SVG("Starship/MainToolbar/pause", Icon40x40));
 		Set( "PlayWorld.PausePlaySession.Small", new IMAGE_BRUSH_SVG("Starship/MainToolbar/pause", Icon20x20));
 		Set( "PlayWorld.SingleFrameAdvance", new IMAGE_BRUSH( "Icons/icon_advance_40x", Icon40x40 ) );
@@ -5434,8 +5431,8 @@ void FStarshipEditorStyle::FStyle::SetupPersonaStyle()
 		Set( "PlayWorld.StopPlaySession", new IMAGE_BRUSH_SVG("Starship/MainToolbar/stop", Icon40x40));
 		Set( "PlayWorld.StopPlaySession.Small", new IMAGE_BRUSH_SVG("Starship/MainToolbar/stop", Icon20x20));
 
-		Set("PlayWorld.LateJoinSession", new IMAGE_BRUSH("Icons/icon_simulate_40x", Icon40x40));
-		Set("PlayWorld.LateJoinSession.Small", new IMAGE_BRUSH("Icons/icon_simulate_40x", Icon20x20));
+		Set("PlayWorld.LateJoinSession", new IMAGE_BRUSH_SVG("Starship/MainToolbar/simulate", Icon40x40));
+		Set("PlayWorld.LateJoinSession.Small", new IMAGE_BRUSH_SVG("Starship/MainToolbar/simulate", Icon20x20));
 
 		Set( "PlayWorld.PossessPlayer", new IMAGE_BRUSH( "Icons/icon_possess_40x", Icon40x40 ) );
 		Set( "PlayWorld.PossessPlayer.Small", new IMAGE_BRUSH( "Icons/icon_possess_40x", Icon20x20 ) );
@@ -6666,13 +6663,6 @@ void FStarshipEditorStyle::FStyle::SetupToolkitStyles()
 	{
 		Set( "MaterialInstanceEditor.Tabs.Properties", new IMAGE_BRUSH( "/Icons/icon_tab_SelectionDetails_16x", Icon16x16 ) );
 		Set( "MaterialInstanceEditor.Tabs.Parents", new IMAGE_BRUSH( "/Icons/layers_16x", Icon16x16 ) );
-		Set( "MaterialInstanceEditor.LayersBorder", new FSlateColorBrush(FLinearColor(0.1f, 0.1f, 0.1f, 1.0f)));
-		Set("MaterialInstanceEditor.StackHeader", new FSlateColorBrush(FLinearColor(0.025f, 0.025f, 0.025f, 1.0f)));
-		Set("MaterialInstanceEditor.StackBody", new FSlateColorBrush(FLinearColor(0.045f, 0.045f, 0.045f, 1.0f)));
-		Set("MaterialInstanceEditor.StackBodyBlend", new FSlateColorBrush(FLinearColor(0.06f, 0.06f, 0.06f, 1.0f)));
-		Set("MaterialInstanceEditor.StackBodyDragged", new FSlateColorBrush(FLinearColor(0.0f, 0.0f, 0.0f, 1.0f)));
-		Set("MaterialInstanceEditor.StackBody_Highlighted", new BOX_BRUSH("Common/TextBox_Special_Active", FMargin(8.0f / 32.0f)));
-
 	}
 	
 	// Sound Class Editor
diff --git a/Engine/Source/Editor/GraphEditor/Private/SCommentBubble.cpp b/Engine/Source/Editor/GraphEditor/Private/SCommentBubble.cpp
--- a/Engine/Source/Editor/GraphEditor/Private/SCommentBubble.cpp
+++ b/Engine/Source/Editor/GraphEditor/Private/SCommentBubble.cpp
@@ -9,6 +9,7 @@
 #include "Widgets/Input/SCheckBox.h"
 #include "SGraphPanel.h"
 #include "ScopedTransaction.h"
+#include "Styling/StyleColors.h"
 
 namespace SCommentBubbleDefs
 {
@@ -256,15 +257,13 @@ void SCommentBubble::UpdateBubble()
 						[
 							SAssignNew(TextBlock, SMultiLineEditableTextBox)
 							.Text(MakeAttributeLambda([this] { return CachedCommentText; }))
+							.Style(FAppStyle::Get(), "Graph.Node.NodeTitleInlineEditableText")
 							.HintText( NSLOCTEXT( "CommentBubble", "EditCommentHint", "Click to edit" ))
 							.IsReadOnly(this, &SCommentBubble::IsReadOnly)
 							.SelectAllTextWhenFocused( true )
 							.RevertTextOnEscape( true )
 							.ClearKeyboardFocusOnCommit( true )
 							.ModiferKeyForNewLine(EModifierKey::Shift)
-							.ForegroundColor(this, &SCommentBubble::GetTextForegroundColor)
-							.ReadOnlyForegroundColor(this, &SCommentBubble::GetReadOnlyTextForegroundColor)
-							.BackgroundColor(this, &SCommentBubble::GetTextBackgroundColor)
 							.OnTextCommitted(this, &SCommentBubble::OnCommentTextCommitted)
 						]
 						+SHorizontalBox::Slot()
@@ -401,22 +400,6 @@ FSlateColor SCommentBubble::GetBubbleColor() const
 	}
 	return ReturnColor;
 }
-
-FSlateColor SCommentBubble::GetTextBackgroundColor() const
-{
-	return TextBlock->HasKeyboardFocus() ? FLinearColor::White : SCommentBubbleDefs::TextClearBackground;
-}
-
-FSlateColor SCommentBubble::GetTextForegroundColor() const
-{
-	return TextBlock->HasKeyboardFocus() ? FSlateColor::UseStyle() : FLinearColor::Black;
-}
-
-FSlateColor SCommentBubble::GetReadOnlyTextForegroundColor() const
-{
-	return TextBlock->HasKeyboardFocus() ? FLinearColor::White : FLinearColor::Black;
-}
-
 void SCommentBubble::OnCommentTextCommitted( const FText& NewText, ETextCommit::Type CommitInfo )
 {
 	if (CommitInfo == ETextCommit::OnEnter || CommitInfo == ETextCommit::OnUserMovedFocus)
diff --git a/Engine/Source/Editor/GraphEditor/Private/SGraphNode.cpp b/Engine/Source/Editor/GraphEditor/Private/SGraphNode.cpp
--- a/Engine/Source/Editor/GraphEditor/Private/SGraphNode.cpp
+++ b/Engine/Source/Editor/GraphEditor/Private/SGraphNode.cpp
@@ -732,7 +732,7 @@ void SGraphNode::SetupErrorReporting()
 TSharedRef<SWidget> SGraphNode::CreateTitleWidget(TSharedPtr<SNodeTitle> NodeTitle)
 {
 	SAssignNew(InlineEditableText, SInlineEditableTextBlock)
-		.Style(FEditorStyle::Get(), "Graph.Node.NodeTitleInlineEditableText")
+		.Style(FAppStyle::Get(), "Graph.Node.NodeTitleInlineEditableText")
 		.Text(NodeTitle.Get(), &SNodeTitle::GetHeadTitle)
 		.OnVerifyTextChanged(this, &SGraphNode::OnVerifyNameTextChanged)
 		.OnTextCommitted(this, &SGraphNode::OnNameTextCommited)
diff --git a/Engine/Source/Editor/GraphEditor/Private/SGraphNodeComment.cpp b/Engine/Source/Editor/GraphEditor/Private/SGraphNodeComment.cpp
--- a/Engine/Source/Editor/GraphEditor/Private/SGraphNodeComment.cpp
+++ b/Engine/Source/Editor/GraphEditor/Private/SGraphNodeComment.cpp
@@ -141,7 +141,7 @@ void SGraphNodeComment::UpdateGraphNode()
 
 	CachedFontSize = CommentNode->GetFontSize();
 
-	CommentStyle = FEditorStyle::Get().GetWidgetStyle<FInlineEditableTextBlockStyle>("Graph.CommentBlock.TitleInlineEditableText");
+	CommentStyle = FAppStyle::Get().GetWidgetStyle<FInlineEditableTextBlockStyle>("Graph.CommentBlock.TitleInlineEditableText");
 	CommentStyle.EditableTextBoxStyle.Font.Size = CachedFontSize;
 	CommentStyle.TextStyle.Font.Size = CachedFontSize;
 
diff --git a/Engine/Source/Editor/GraphEditor/Public/SCommentBubble.h b/Engine/Source/Editor/GraphEditor/Public/SCommentBubble.h
--- a/Engine/Source/Editor/GraphEditor/Public/SCommentBubble.h
+++ b/Engine/Source/Editor/GraphEditor/Public/SCommentBubble.h
@@ -133,15 +133,6 @@ protected:
 	/** Returns the color of the main bubble */
 	FSlateColor GetBubbleColor() const;
 
-	/** Returns the current background color for the textbox */
-	FSlateColor GetTextBackgroundColor() const;
-	
-	/** Returns the current foreground color for the textbox */
-	FSlateColor GetTextForegroundColor() const;
-
-	/** Returns the current foreground color for the textbox when it is read only*/
-	FSlateColor GetReadOnlyTextForegroundColor() const;
-
 	/** Returns the foreground color for the text and buttons, taking into account the bubble color */
 	FSlateColor GetForegroundColor() const override { return ForegroundColor; }
 
