Fix to warning message on Slate Styling.

#rnx

diff --git a/Engine/Plugins/Editor/EditorDebugTools/Source/EditorDebugTools/Private/EditorDebugToolsStyle.cpp b/Engine/Plugins/Editor/EditorDebugTools/Source/EditorDebugTools/Private/EditorDebugToolsStyle.cpp
--- a/Engine/Plugins/Editor/EditorDebugTools/Source/EditorDebugTools/Private/EditorDebugToolsStyle.cpp
+++ b/Engine/Plugins/Editor/EditorDebugTools/Source/EditorDebugTools/Private/EditorDebugToolsStyle.cpp
@@ -8,6 +8,7 @@
 
 TSharedPtr< FSlateStyleSet > FEditorDebugToolsStyle::StyleInstance = NULL;
 
+#define IMAGE_BRUSH( RelativePath, ... ) FSlateImageBrush( FPaths::EngineContentDir() / "Editor/Slate"/ RelativePath + TEXT(".png"), __VA_ARGS__ )
 #define IMAGE_PLUGIN_BRUSH( RelativePath, ... ) FSlateImageBrush( Style->RootToContentDir( RelativePath, TEXT(".png") ), __VA_ARGS__ )
 
 void FEditorDebugToolsStyle::Initialize()
@@ -46,7 +47,7 @@ TSharedRef< FSlateStyleSet > FEditorDebugToolsStyle::Create()
 	Style->Set("GammaReference", new IMAGE_PLUGIN_BRUSH("GammaReference", FVector2D(256, 128)));
 
 	// Debug Tools Window
-	Style->Set("DebugTools.TabIcon", new IMAGE_PLUGIN_BRUSH("icon_tab_DebugTools_16x", Icon16x16));
+	Style->Set("DebugTools.TabIcon", new IMAGE_BRUSH("Icons/icon_tab_DebugTools_16x", Icon16x16));
 	
 	// Modules Window
 	Style->Set("Modules.TabIcon", new IMAGE_PLUGIN_BRUSH("icon_tab_Modules_16px", Icon16x16));
diff --git a/Engine/Source/Developer/LogVisualizer/Private/LogVisualizerStyle.cpp b/Engine/Source/Developer/LogVisualizer/Private/LogVisualizerStyle.cpp
--- a/Engine/Source/Developer/LogVisualizer/Private/LogVisualizerStyle.cpp
+++ b/Engine/Source/Developer/LogVisualizer/Private/LogVisualizerStyle.cpp
@@ -63,7 +63,7 @@ TSharedRef< FSlateStyleSet > FLogVisualizerStyle::Create()
 	Style.Set("TextLogs.Text", LogsText);
 
 	{
-		Style.Set("LogVisualizerApp.TabIcon", new IMAGE_BRUSH("Icons/icon_tab_DebugTools_40x", Icon16x16));
+		Style.Set("LogVisualizerApp.TabIcon", new IMAGE_BRUSH("Icons/icon_tab_DebugTools_16x", Icon16x16));
 		
 		Style.Set("LogVisualizer.LogBar.Background", new BOX_BRUSH("Common/ProgressBar_Background", FMargin(5.f / 12.f)));
 		Style.Set("LogVisualizer.LogBar.Selected", new BOX_BRUSH("Common/TaskGraph_Selected", FMargin(5.f / 12.f)));
diff --git a/Engine/Source/Editor/EditorStyle/Private/StarshipStyle.cpp b/Engine/Source/Editor/EditorStyle/Private/StarshipStyle.cpp
--- a/Engine/Source/Editor/EditorStyle/Private/StarshipStyle.cpp
+++ b/Engine/Source/Editor/EditorStyle/Private/StarshipStyle.cpp
@@ -4441,12 +4441,12 @@ void FStarshipEditorStyle::FStyle::SetupLevelEditorStyle()
 		Set( "ToolPalette.DockingTab", FCheckBoxStyle()
 			.SetCheckBoxType(ESlateCheckBoxType::ToggleButton)
 			.SetPadding( FMargin(16.0f, 2.0f, 16.0f, 2.0f ) )
-			.SetCheckedImage(       BOX_BRUSH("Docking/Tab_Shape",  2.f/8.0f, FLinearColor(FColor(62, 62, 62)) ) )
-			.SetCheckedHoveredImage(BOX_BRUSH("Docking/Tab_Shape",  2.f/8.0f, FLinearColor(FColor(62, 62, 62)) ) )
-			.SetCheckedPressedImage(BOX_BRUSH("Docking/Tab_Shape",  2.f/8.0f, FLinearColor(FColor(62, 62, 62)) ) )
-			.SetUncheckedImage(     BOX_BRUSH("Docking/Tab_Shape",  2.f/8.0f, FLinearColor(FColor(45, 45, 45)) ) )
-			.SetUncheckedHoveredImage(BOX_BRUSH("Docking/Tab_Shape",2.f/8.0f, FLinearColor(FColor(54, 54, 54)) ) )
-			.SetUncheckedPressedImage(BOX_BRUSH("Docking/Tab_Shape",2.f/8.0f, FLinearColor(FColor(54, 54, 54)) ) )
+			.SetCheckedImage(         CORE_BOX_BRUSH("Docking/Tab_Shape",  2.f/8.0f, FLinearColor(FColor(62, 62, 62)) ) )
+			.SetCheckedHoveredImage(  CORE_BOX_BRUSH("Docking/Tab_Shape",  2.f/8.0f, FLinearColor(FColor(62, 62, 62)) ) )
+			.SetCheckedPressedImage(  CORE_BOX_BRUSH("Docking/Tab_Shape",  2.f/8.0f, FLinearColor(FColor(62, 62, 62)) ) )
+			.SetUncheckedImage(       CORE_BOX_BRUSH("Docking/Tab_Shape",  2.f/8.0f, FLinearColor(FColor(45, 45, 45)) ) )
+			.SetUncheckedHoveredImage(CORE_BOX_BRUSH("Docking/Tab_Shape",2.f/8.0f, FLinearColor(FColor(54, 54, 54)) ) )
+			.SetUncheckedPressedImage(CORE_BOX_BRUSH("Docking/Tab_Shape",2.f/8.0f, FLinearColor(FColor(54, 54, 54)) ) )
 			.SetUndeterminedImage(        FSlateNoResource() )
 			.SetUndeterminedHoveredImage( FSlateNoResource() )
 			.SetUndeterminedPressedImage( FSlateNoResource() )
@@ -7457,18 +7457,21 @@ void FStarshipEditorStyle::FStyle::SetupAutomationStyles()
 
 			// some platforms may specify a "rooted" path in the platform extensions directory, so look for that case here, and use a different path for the brush
 			FString NormalIconPath = PlatformInfo.GetIconPath(EPlatformIconSize::Normal);
-			if (NormalIconPath.StartsWith(TEXT("/Platforms/")))
+			if(!NormalIconPath.IsEmpty())
 			{
+				if (NormalIconPath.StartsWith(TEXT("/Platforms/")))
+				{
 #define PLATFORM_IMAGE_BRUSH( PlatformPath, ... ) FSlateImageBrush( PlatformPath.Replace(TEXT("/Platforms/"), *FPaths::EnginePlatformExtensionsDir()) + TEXT(".png") , __VA_ARGS__ )
-				Set(PlatformInfo.GetIconStyleName(EPlatformIconSize::Normal), new PLATFORM_IMAGE_BRUSH(NormalIconPath, Icon24x24));
-				Set(PlatformInfo.GetIconStyleName(EPlatformIconSize::Large), new PLATFORM_IMAGE_BRUSH(PlatformInfo.GetIconPath(EPlatformIconSize::Large), Icon64x64));
-				Set(PlatformInfo.GetIconStyleName(EPlatformIconSize::XLarge), new PLATFORM_IMAGE_BRUSH(PlatformInfo.GetIconPath(EPlatformIconSize::XLarge), Icon128x128));
-			}
-			else
-			{
-				Set(PlatformInfo.GetIconStyleName(EPlatformIconSize::Normal), new IMAGE_BRUSH(*NormalIconPath, Icon24x24));
-				Set(PlatformInfo.GetIconStyleName(EPlatformIconSize::Large), new IMAGE_BRUSH(*PlatformInfo.GetIconPath(EPlatformIconSize::Large), Icon64x64));
-				Set(PlatformInfo.GetIconStyleName(EPlatformIconSize::XLarge), new IMAGE_BRUSH(*PlatformInfo.GetIconPath(EPlatformIconSize::XLarge), Icon128x128));
+					Set(PlatformInfo.GetIconStyleName(EPlatformIconSize::Normal), new PLATFORM_IMAGE_BRUSH(NormalIconPath, Icon24x24));
+					Set(PlatformInfo.GetIconStyleName(EPlatformIconSize::Large), new PLATFORM_IMAGE_BRUSH(PlatformInfo.GetIconPath(EPlatformIconSize::Large), Icon64x64));
+					Set(PlatformInfo.GetIconStyleName(EPlatformIconSize::XLarge), new PLATFORM_IMAGE_BRUSH(PlatformInfo.GetIconPath(EPlatformIconSize::XLarge), Icon128x128));
+				}
+				else
+				{
+					Set(PlatformInfo.GetIconStyleName(EPlatformIconSize::Normal), new IMAGE_BRUSH(*NormalIconPath, Icon24x24));
+					Set(PlatformInfo.GetIconStyleName(EPlatformIconSize::Large), new IMAGE_BRUSH(*PlatformInfo.GetIconPath(EPlatformIconSize::Large), Icon64x64));
+					Set(PlatformInfo.GetIconStyleName(EPlatformIconSize::XLarge), new IMAGE_BRUSH(*PlatformInfo.GetIconPath(EPlatformIconSize::XLarge), Icon128x128));
+				}
 			}
 		}
 #endif
