Adding back AllPlatforms Brushes to fix mac references

diff --git a/Engine/Source/Editor/EditorStyle/Private/StarshipStyle.cpp b/Engine/Source/Editor/EditorStyle/Private/StarshipStyle.cpp
--- a/Engine/Source/Editor/EditorStyle/Private/StarshipStyle.cpp
+++ b/Engine/Source/Editor/EditorStyle/Private/StarshipStyle.cpp
@@ -7451,6 +7451,9 @@ void FStarshipEditorStyle::FStyle::SetupAutomationStyles()
 
 #if DDPI_HAS_EXTENDED_PLATFORMINFO_DATA
 
+		Set( "Launcher.Platform.AllPlatforms", new IMAGE_BRUSH( "Launcher/All_Platforms_24x", Icon24x24) );
+		Set( "Launcher.Platform.AllPlatforms.Large", new IMAGE_BRUSH( "Launcher/All_Platforms_128x", Icon64x64) );
+		Set( "Launcher.Platform.AllPlatforms.XLarge", new IMAGE_BRUSH( "Launcher/All_Platforms_128x", Icon128x128) );
 		for (auto Pair : FDataDrivenPlatformInfoRegistry::GetAllPlatformInfos())
 		{
 			const FDataDrivenPlatformInfo& PlatformInfo = Pair.Value;
