Fix up missing icons and buttons in level browser

diff --git a/Engine/Source/Editor/EditorStyle/Private/StarshipStyle.cpp b/Engine/Source/Editor/EditorStyle/Private/StarshipStyle.cpp
--- a/Engine/Source/Editor/EditorStyle/Private/StarshipStyle.cpp
+++ b/Engine/Source/Editor/EditorStyle/Private/StarshipStyle.cpp
@@ -1869,7 +1869,7 @@ void FStarshipEditorStyle::FStyle::SetupLevelGeneralStyles()
 		Set("Level.NotVisibleHighlightIcon16x", new CORE_IMAGE_BRUSH_SVG("Starship/Common/hidden", Icon16x16));
 
 		Set("Level.LightingScenarioIcon16x", new IMAGE_BRUSH_SVG("Starship/AssetIcons/PointLight_16", Icon16x16));
-		Set("Level.LightingScenarioNotIcon16x", new IMAGE_BRUSH("Icons/icon_levels_LightingScenarioNot_16px", Icon16x16));
+		Set("Level.LightingScenarioNotIcon16x", new IMAGE_BRUSH_SVG("Starship/Common/LightBulbOff", Icon16x16));
 		Set("Level.LockedIcon16x", new IMAGE_BRUSH("Icons/icon_locked_16px", Icon16x16));
 		Set("Level.LockedHighlightIcon16x", new IMAGE_BRUSH("Icons/icon_locked_highlight_16px", Icon16x16));
 		Set("Level.UnlockedIcon16x", new IMAGE_BRUSH("Icons/icon_levels_unlocked_16px", Icon16x16));
@@ -1885,7 +1885,7 @@ void FStarshipEditorStyle::FStyle::SetupLevelGeneralStyles()
 		Set("Level.ScriptIcon16x", new IMAGE_BRUSH("Icons/icon_levels_Blueprint_16px", Icon16x16));
 		Set("Level.ScriptHighlightIcon16x", new IMAGE_BRUSH("Icons/icon_levels_Blueprint_hi_16px", Icon16x16));
 		Set("Level.EmptyIcon16x", new IMAGE_BRUSH("Icons/Empty_16x", Icon16x16));
-		Set("Level.ColorIcon40x", new IMAGE_BRUSH("Icons/icon_levels_back_16px", Icon16x16));
+		Set("Level.ColorIcon", new FSlateRoundedBoxBrush(FStyleColors::White, 4.0f, Icon16x16));
 	}
 
 	// Spline component controls
@@ -1957,9 +1957,7 @@ void FStarshipEditorStyle::FStyle::SetupWorldBrowserStyles()
 		Set("WorldBrowser.DetailsButtonBrush", new IMAGE_BRUSH("Icons/icon_levels_detailsbutton_40x", Icon16x16));
 		Set("WorldBrowser.CompositionButtonBrush", new IMAGE_BRUSH("Icons/icon_levels_compositionbutton_16x", Icon16x16));
 
-		Set("WorldBrowser.FolderClosed", new IMAGE_BRUSH("Icons/FolderClosed", Icon16x16));
-		Set("WorldBrowser.FolderOpen", new IMAGE_BRUSH("Icons/FolderOpen", Icon16x16));
-		Set("WorldBrowser.NewFolderIcon", new IMAGE_BRUSH("Icons/icon_AddFolder_16x", Icon16x16));
+		Set("WorldBrowser.NewFolderIcon", new CORE_IMAGE_BRUSH_SVG("Starship/Common/folder-plus", Icon16x16));
 
 		Set("WorldBrowser.StatusBarText", FTextBlockStyle(NormalText)
 			.SetFont(DEFAULT_FONT("BoldCondensed", 12))
diff --git a/Engine/Source/Editor/WorldBrowser/Private/SWorldHierarchyImpl.cpp b/Engine/Source/Editor/WorldBrowser/Private/SWorldHierarchyImpl.cpp
--- a/Engine/Source/Editor/WorldBrowser/Private/SWorldHierarchyImpl.cpp
+++ b/Engine/Source/Editor/WorldBrowser/Private/SWorldHierarchyImpl.cpp
@@ -190,14 +190,16 @@ void SWorldHierarchyImpl::Construct(const FArguments& InArgs)
 	TSharedRef<SWidget> CreateNewFolderButton = SNullWidget::NullWidget;
 	if (!bFoldersOnlyMode)
 	{
-		CreateNewFolderButton = SNew(SButton)
-			.ButtonStyle(FEditorStyle::Get(), "HoverHintOnly")
+		CreateNewFolderButton = 
+			SNew(SButton)
+			.ButtonStyle(FAppStyle::Get(), "SimpleButton")
 			.ToolTipText(LOCTEXT("CreateFolderTooltip", "Create a new folder containing the current selection"))
 			.OnClicked(this, &SWorldHierarchyImpl::OnCreateFolderClicked)
 			.Visibility(WorldModel->HasFolderSupport() ? EVisibility::Visible : EVisibility::Collapsed)
 			[
 				SNew(SImage)
 				.Image(FEditorStyle::GetBrush("WorldBrowser.NewFolderIcon"))
+				.ColorAndOpacity(FSlateColor::UseForeground())
 			];
 	}
 
@@ -223,6 +225,7 @@ void SWorldHierarchyImpl::Construct(const FArguments& InArgs)
 			// Create New Folder icon
 			+ SHorizontalBox::Slot()
 			.AutoWidth()
+			.VAlign(VAlign_Center)
 			.Padding(4.0f, 0.0f, 0.0f, 0.0f)
 			[
 				CreateNewFolderButton
diff --git a/Engine/Source/Editor/WorldBrowser/Private/SWorldHierarchyItem.cpp b/Engine/Source/Editor/WorldBrowser/Private/SWorldHierarchyItem.cpp
--- a/Engine/Source/Editor/WorldBrowser/Private/SWorldHierarchyItem.cpp
+++ b/Engine/Source/Editor/WorldBrowser/Private/SWorldHierarchyItem.cpp
@@ -21,6 +21,7 @@
 #include "WorldTreeItemTypes.h"
 #include "LevelFolders.h"
 #include "Framework/Application/SlateApplication.h"
+#include "Styling/StyleColors.h"
 
 #define LOCTEXT_NAMESPACE "WorldBrowser"
 
@@ -104,6 +105,7 @@ TSharedRef< SWidget > SWorldHierarchyItem::GenerateWidgetForColumn( const FName&
 				[
 					SNew(SImage)
 					.Image(this, &SWorldHierarchyItem::GetLevelIconBrush)
+					.ColorAndOpacity(FStyleColors::AccentFolder)
 				]
 			]
 
@@ -196,8 +198,7 @@ TSharedRef< SWidget > SWorldHierarchyItem::GenerateWidgetForColumn( const FName&
 				[
 					SNew(SImage)
 					.ColorAndOpacity(this, &SWorldHierarchyItem::GetDrawColor)
-					.Image(this, &SWorldHierarchyItem::GetLevelColorBrush)
-					.ColorAndOpacity(FSlateColor::UseForeground())
+					.Image(FEditorStyle::GetBrush("Level.ColorIcon"))
 				]
 			;
 		}
@@ -230,15 +231,9 @@ TSharedRef< SWidget > SWorldHierarchyItem::GenerateWidgetForColumn( const FName&
 				.HAlign(HAlign_Center)
 				.VAlign(VAlign_Center)
 				[
-					SNew(SHorizontalBox)
-					+ SHorizontalBox::Slot()
-					.FillWidth(1.0f)
-					.VAlign(VAlign_Center)
-					[
-						SNew(SImage)
-						.Image(this, &SWorldHierarchyItem::GetSCCStateImage)
-						.ToolTipText(this, &SWorldHierarchyItem::GetSCCStateTooltip)
-					]
+					SNew(SImage)
+					.Image(this, &SWorldHierarchyItem::GetSCCStateImage)
+					.ToolTipText(this, &SWorldHierarchyItem::GetSCCStateTooltip)
 				]
 			;
 		}
@@ -760,7 +755,7 @@ const FSlateBrush* SWorldHierarchyItem::GetSCCStateImage() const
 		FSourceControlStatePtr SourceControlState = ISourceControlModule::Get().GetProvider().GetState(PackageName, EStateCacheUsage::Use);
 		if(SourceControlState.IsValid())
 		{
-			return FEditorStyle::GetBrush(SourceControlState->GetIcon().GetStyleSetName());
+			return SourceControlState->GetIcon().GetIcon();
 		}
 	}
 
@@ -798,9 +793,4 @@ const FSlateBrush* SWorldHierarchyItem::GetLevelKismetBrush() const
 	}	
 }
 
-const FSlateBrush* SWorldHierarchyItem::GetLevelColorBrush() const
-{
-	return FEditorStyle::GetBrush("Level.ColorIcon40x");
-}
-
 #undef LOCTEXT_NAMESPACE
diff --git a/Engine/Source/Editor/WorldBrowser/Private/WorldTreeItemTypes.cpp b/Engine/Source/Editor/WorldBrowser/Private/WorldTreeItemTypes.cpp
--- a/Engine/Source/Editor/WorldBrowser/Private/WorldTreeItemTypes.cpp
+++ b/Engine/Source/Editor/WorldBrowser/Private/WorldTreeItemTypes.cpp
@@ -898,10 +898,10 @@ namespace WorldHierarchy
 	{
 		if (Children.Num() > 0 && Flags.bExpanded)
 		{
-			return FEditorStyle::Get().GetBrush("WorldBrowser.FolderOpen");
+			return FEditorStyle::Get().GetBrush("Icons.FolderOpen");
 		}
 
-		return FEditorStyle::Get().GetBrush("WorldBrowser.FolderClosed");
+		return FEditorStyle::Get().GetBrush("Icons.FolderClosed");
 	}
 
 	void FFolderTreeItem::GenerateContextMenu(UToolMenu* Menu, const SWorldHierarchyImpl& Hierarchy)
