Restoring lost Curve Editor Icon

diff --git a/Engine/Source/Editor/EditorStyle/Private/StarshipStyle.cpp b/Engine/Source/Editor/EditorStyle/Private/StarshipStyle.cpp
--- a/Engine/Source/Editor/EditorStyle/Private/StarshipStyle.cpp
+++ b/Engine/Source/Editor/EditorStyle/Private/StarshipStyle.cpp
@@ -6977,6 +6977,8 @@ void FStarshipEditorStyle::FStyle::SetupToolkitStyles()
 
 		Set("GenericCurveEditor.Pin_Active", new IMAGE_BRUSH("Common/PushPin_Down", Icon16x16));
 		Set("GenericCurveEditor.Pin_Inactive", new IMAGE_BRUSH("Common/PushPin_Up", Icon16x16));
+		
+		Set("GenericCurveEditor.Select", new IMAGE_BRUSH("GenericCurveEditor/Icons/SelectButton", Icon16x16));
 	}
 
 	// Generic Curve Editor Tools
