Re-enable brush edit mode

diff --git a/Engine/Plugins/Editor/GeometryMode/Source/GeometryMode/Private/GeometryModeModule.cpp b/Engine/Plugins/Editor/GeometryMode/Source/GeometryMode/Private/GeometryModeModule.cpp
--- a/Engine/Plugins/Editor/GeometryMode/Source/GeometryMode/Private/GeometryModeModule.cpp
+++ b/Engine/Plugins/Editor/GeometryMode/Source/GeometryMode/Private/GeometryModeModule.cpp
@@ -20,8 +20,8 @@ void FGeometryModeModule::StartupModule()
 	FEditorModeRegistry::Get().RegisterMode<FEdModeGeometry>(
 		FGeometryEditingModes::EM_Geometry,
 		NSLOCTEXT("EditorModes", "GeometryMode", "Brush Editing"),
-		FSlateIcon(FEditorStyle::GetStyleSetName(), "LevelEditor.BspMode", "LevelEditor.BspMode.Small"),
-		true, 500
+		FSlateIcon(FEditorStyle::GetStyleSetName(), "LevelEditor.BrushEdit"),
+		true
 	);
 }
 
diff --git a/Engine/Source/Editor/EditorStyle/Private/StarshipStyle.cpp b/Engine/Source/Editor/EditorStyle/Private/StarshipStyle.cpp
--- a/Engine/Source/Editor/EditorStyle/Private/StarshipStyle.cpp
+++ b/Engine/Source/Editor/EditorStyle/Private/StarshipStyle.cpp
@@ -4492,6 +4492,8 @@ void FStarshipEditorStyle::FStyle::SetupLevelEditorStyle()
 
 		Set( "LevelEditor.WorldProperties.Tab", new IMAGE_BRUSH_SVG( "Starship/Common/WorldSettings", Icon16x16 ) );
 
+		Set("LevelEditor.BrushEdit", new IMAGE_BRUSH_SVG("Starship/MainToolbar/brush_edit", Icon20x20));
+
 		Set( "LevelEditor.OpenPlaceActors", new IMAGE_BRUSH_SVG( "Starship/Common/PlaceActors", Icon20x20 ) );
 		Set( "LevelEditor.OpenContentBrowser", new IMAGE_BRUSH_SVG( "Starship/MainToolbar/content", Icon20x20 ) );
 		Set( "LevelEditor.OpenMarketplace", new IMAGE_BRUSH_SVG( "Starship/MainToolbar/marketplace", Icon20x20) );
