UMG reskin cleanup

diff --git a/Engine/Source/Editor/EditorStyle/Private/StarshipStyle.cpp b/Engine/Source/Editor/EditorStyle/Private/StarshipStyle.cpp
--- a/Engine/Source/Editor/EditorStyle/Private/StarshipStyle.cpp
+++ b/Engine/Source/Editor/EditorStyle/Private/StarshipStyle.cpp
@@ -5174,7 +5174,9 @@ void FStarshipEditorStyle::FStyle::SetupPersonaStyle()
 			Set( "Kismet.Tabs.BlueprintDefaults", new IMAGE_BRUSH( "Icons/icon_BlueprintEditor_Defaults_40x", Icon16x16 ) );
 		}
 
+		// this should be moved to a generic icon in the core style. Icons.Star
 		Set("Kismet.Palette.Favorites", new IMAGE_BRUSH_SVG("Starship/Common/Favorite", Icon16x16));
+
 		const FCheckBoxStyle KismetFavoriteToggleStyle = FCheckBoxStyle()
 			.SetCheckBoxType(ESlateCheckBoxType::CheckBox)
 			.SetUncheckedImage( IMAGE_BRUSH("Icons/EmptyStar_16x", Icon10x10, FLinearColor(0.8f, 0.8f, 0.8f, 1.f)) )
@@ -5280,6 +5282,7 @@ void FStarshipEditorStyle::FStyle::SetupPersonaStyle()
 		Set("BlendSpaceEditor.ArrowLeft", new IMAGE_BRUSH(TEXT("Persona/BlendSpace/arrow_left_12x"), FVector2D(25.0f, 13.0f)));
 		Set("BlendSpaceEditor.PreviewIcon", new IMAGE_BRUSH(TEXT("Persona/BlendSpace/preview_21x"), FVector2D(21.0f, 21.0f)));
 		Set("BlendSpaceEditor.LabelBackground", new FSlateRoundedBoxBrush(FStyleColors::Background, FStyleColors::Foreground, 1.0f));
+		Set("BlendSpaceEditor.ZoomToFit", new IMAGE_BRUSH("GenericCurveEditor/Icons/FramingSelected_48x", Icon16x16));
 
 		const FButtonStyle BlueprintContextTargetsButtonStyle = FButtonStyle()
 			.SetNormal(IMAGE_BRUSH("Common/TreeArrow_Collapsed_Hovered", Icon10x10, FLinearColor(0.2f, 0.2f, 0.2f, 1.f)))
@@ -7607,29 +7610,22 @@ void FStarshipEditorStyle::FStyle::SetupAutomationStyles()
 
 void FStarshipEditorStyle::FStyle::SetupUMGEditorStyles()
 {
-	const FLinearColor IconColor = FLinearColor::Black;
-
 	Set("WidgetDesigner.LayoutTransform", new IMAGE_BRUSH("Icons/UMG/Layout_TransformMode_16x", Icon16x16));
 	Set("WidgetDesigner.LayoutTransform.Small", new IMAGE_BRUSH("Icons/UMG/Layout_TransformMode_16x", Icon16x16));
 	Set("WidgetDesigner.RenderTransform", new IMAGE_BRUSH("Icons/UMG/Render_TransformMode_16x", Icon16x16));
 	Set("WidgetDesigner.RenderTransform.Small", new IMAGE_BRUSH("Icons/UMG/Render_TransformMode_16x", Icon16x16));
 	Set("WidgetDesigner.ToggleOutlines", new IMAGE_BRUSH("Icons/UMG/ToggleOutlines.Small", Icon16x16));
-	Set("WidgetDesigner.ToggleOutlines.Small", new IMAGE_BRUSH("Icons/UMG/ToggleOutlines.Small", Icon16x16));
-	Set("WidgetDesigner.ToggleRespectLocks", new IMAGE_BRUSH("Icons/UMG/ToggleRespectLocks.Small", Icon16x16));
-	Set("WidgetDesigner.ToggleRespectLocks.Small", new IMAGE_BRUSH("Icons/UMG/ToggleRespectLocks.Small", Icon16x16));
-	Set("WidgetDesigner.ToggleLocalizationPreview", new IMAGE_BRUSH("Icons/icon_localization_white_16x", Icon16x16, FLinearColor::Black));
-	Set("WidgetDesigner.ToggleLocalizationPreview.Small", new IMAGE_BRUSH("Icons/icon_localization_white_16x", Icon16x16, FLinearColor::Black));
+	Set("WidgetDesigner.ToggleRespectLocks", new CORE_IMAGE_BRUSH_SVG("Starship/Common/lock", Icon16x16));
 
-	Set("WidgetDesigner.LocationGridSnap", new IMAGE_BRUSH("Old/LevelEditor/LocationGridSnap", Icon14x14, IconColor));
-	Set("WidgetDesigner.RotationGridSnap", new IMAGE_BRUSH("Old/LevelEditor/RotationGridSnap", Icon14x14, IconColor));
+	Set("WidgetDesigner.ToggleLocalizationPreview", new IMAGE_BRUSH("Icons/icon_localization_white_16x", Icon16x16));
 
-	Set("WidgetDesigner.ZoomToFit", new IMAGE_BRUSH("Icons/UMG/Fit_16x", Icon16x16));
-	Set("WidgetDesigner.ZoomToFit.Small", new IMAGE_BRUSH("Icons/UMG/Fit_16x", Icon16x16));
+	Set("WidgetDesigner.LocationGridSnap", new IMAGE_BRUSH_SVG("Starship/EditorViewport/grid", Icon16x16));
+	Set("WidgetDesigner.RotationGridSnap", new IMAGE_BRUSH("Old/LevelEditor/RotationGridSnap", Icon16x16));
 
 	Set("WidgetDesigner.WidgetVisible", new IMAGE_BRUSH("/Icons/icon_layer_visible", Icon16x16));
 	Set("WidgetDesigner.WidgetHidden", new IMAGE_BRUSH("/Icons/icon_layer_not_visible", Icon16x16));
 
-	Set("UMGEditor.ZoomToFit", new IMAGE_BRUSH("Icons/UMG/Fit_16x", Icon16x16, FLinearColor(.05f, .05f, .05f, 1.f)));
+	Set("UMGEditor.ZoomToFit", new IMAGE_BRUSH("GenericCurveEditor/Icons/FramingSelected_48x", Icon16x16));
 
 	Set("UMGEditor.ScreenOutline", new BOX_BRUSH(TEXT("Icons/UMG/ScreenOutline"), FMargin(0.25f) ));
 
@@ -7665,24 +7661,20 @@ void FStarshipEditorStyle::FStyle::SetupUMGEditorStyles()
 	Set("UMGEditor.AnchoredWidget", new BOX_BRUSH("Common/Button", FVector2D(32, 32), 8.0f / 32.0f));
 	Set("UMGEditor.AnchoredWidgetAlignment", new IMAGE_BRUSH("Icons/icon_tab_DeviceManager_16x", Icon8x8));
 	
-	Set("UMGEditor.PaletteHeader", FTableRowStyle()
-		.SetEvenRowBackgroundBrush( BOX_BRUSH( "PropertyView/DetailCategoryMiddle", FMargin( 4/16.0f, 8.0f/16.0f, 4/16.0f, 4/16.0f ), FLinearColor(0.6f,0.6f,0.6f,1.0f) ) )
-		.SetEvenRowBackgroundHoveredBrush( BOX_BRUSH( "PropertyView/DetailCategoryMiddle", FMargin( 4/16.0f, 8.0f/16.0f, 4/16.0f, 4/16.0f ), FLinearColor(0.3f,0.3f,0.3f,1.0f) ) )
-		.SetOddRowBackgroundBrush( BOX_BRUSH( "PropertyView/DetailCategoryMiddle", FMargin( 4/16.0f, 8.0f/16.0f, 4/16.0f, 4/16.0f ), FLinearColor(0.6f,0.6f,0.6f,1.0f) ) )
-		.SetOddRowBackgroundHoveredBrush( BOX_BRUSH( "PropertyView/DetailCategoryMiddle", FMargin( 4/16.0f, 8.0f/16.0f, 4/16.0f, 4/16.0f ), FLinearColor(0.3f,0.3f,0.3f,1.0f) ) )
-		.SetSelectorFocusedBrush( BORDER_BRUSH( "Common/Selector", FMargin(4.f/16.f), SelectorColor ) )
-		.SetActiveBrush( IMAGE_BRUSH( "Common/Selection", Icon8x8, SelectionColor ) )
-		.SetActiveHoveredBrush( IMAGE_BRUSH( "Common/Selection", Icon8x8, SelectionColor ) )
-		.SetInactiveBrush( IMAGE_BRUSH( "Common/Selection", Icon8x8, SelectionColor_Inactive ) )
-		.SetInactiveHoveredBrush( IMAGE_BRUSH( "Common/Selection", Icon8x8, SelectionColor_Inactive ) )
-		.SetTextColor( DefaultForeground )
-		.SetSelectedTextColor( InvertedForeground )
-		);
 
-	Set("UMGEditor.PaletteItem", FTableRowStyle(NormalTableRowStyle)
-		.SetEvenRowBackgroundBrush(BOX_BRUSH("PropertyView/DetailCategoryMiddle", FMargin(4 / 16.0f, 8.0f / 16.0f, 4 / 16.0f, 4 / 16.0f)))
-		.SetOddRowBackgroundBrush(BOX_BRUSH("PropertyView/DetailCategoryMiddle", FMargin(4 / 16.0f, 8.0f / 16.0f, 4 / 16.0f, 4 / 16.0f)))
-		);
+	Set("UMGEditor.PaletteHeader", FTableRowStyle()
+		.SetEvenRowBackgroundBrush(FSlateColorBrush(FStyleColors::Header))
+		.SetEvenRowBackgroundHoveredBrush(FSlateColorBrush(FStyleColors::Header))
+		.SetOddRowBackgroundBrush(FSlateColorBrush(FStyleColors::Header))
+		.SetOddRowBackgroundHoveredBrush(FSlateColorBrush(FStyleColors::Header))
+		.SetSelectorFocusedBrush(FSlateNoResource())
+		.SetActiveBrush(FSlateNoResource())
+		.SetActiveHoveredBrush(FSlateNoResource())
+		.SetInactiveBrush(FSlateNoResource())
+		.SetInactiveHoveredBrush(FSlateNoResource())
+		.SetTextColor(DefaultForeground)
+		.SetSelectedTextColor(InvertedForeground)
+	);
 
 	// Style of the favorite toggle
 	const FCheckBoxStyle UMGEditorFavoriteToggleStyle = FCheckBoxStyle()
@@ -7840,7 +7832,8 @@ void FStarshipEditorStyle::FStyle::SetupStatusBarStyle()
 	Set("StatusBar.ContentBrowserShadow", new BOX_BRUSH("Starship/StatusBar/drawer-shadow-bottom", FMargin(10/64.0f, 20/64.f, 10/64.f, 0), FLinearColor(0,0,0,1)));
 
 	Set("StatusBar.Height", 32.0f);
-	Set("StatusBar.Background", new FSlateColorBrush(FStyleColors::Panel));
+	Set("StatusBar.ProgressOverlay", new FSlateRoundedBoxBrush(FStyleColors::Transparent, 2.0f, FStyleColors::Panel, 1.0f, FVector2D(20, 8)));
+
 	Set("StatusBar.HelpIcon", new CORE_IMAGE_BRUSH_SVG("Starship/Common/help", Icon16x16, FStyleColors::Foreground));
 
 	FToolBarStyle StatusBarToolBarStyle = FStarshipCoreStyle::GetCoreStyle().GetWidgetStyle<FToolBarStyle>("SlimToolBar");
diff --git a/Engine/Source/Editor/Persona/Private/SAnimationBlendSpaceGridWidget.cpp b/Engine/Source/Editor/Persona/Private/SAnimationBlendSpaceGridWidget.cpp
--- a/Engine/Source/Editor/Persona/Private/SAnimationBlendSpaceGridWidget.cpp
+++ b/Engine/Source/Editor/Persona/Private/SAnimationBlendSpaceGridWidget.cpp
@@ -205,7 +205,7 @@ void SBlendSpaceGridWidget::Construct(const FArguments& InArgs)
 									.ButtonColorAndOpacity_Lambda([this]() -> FLinearColor { return bStretchToFit ? FEditorStyle::GetSlateColor("SelectionColor").GetSpecifiedColor() : FLinearColor::White; })
 									[
 										SNew(SImage)
-										.Image(FEditorStyle::GetBrush("WidgetDesigner.ZoomToFit"))
+										.Image(FEditorStyle::GetBrush("BlendSpaceEditor.ZoomToFit"))
 										.ColorAndOpacity(FSlateColor::UseForeground())
 									]
 								]
@@ -1326,7 +1326,7 @@ void SBlendSpaceGridWidget::MakeViewContextMenuEntries(FMenuBuilder& InMenuBuild
 		InMenuBuilder.AddMenuEntry(
 			LOCTEXT("StretchFittingText", "Stretch Grid to Fit"),
 			LOCTEXT("StretchFittingTextToolTip", "Whether to stretch the grid to fit or to fit the grid to the largest axis"),
-			FSlateIcon("EditorStyle", "WidgetDesigner.ZoomToFit"),
+			FSlateIcon("EditorStyle", "BlendSpaceEditor.ZoomToFit"),
 			FUIAction(
 				FExecuteAction::CreateLambda([this](){ bStretchToFit = !bStretchToFit; }),
 				FCanExecuteAction(),
diff --git a/Engine/Source/Editor/UMGEditor/Private/Designer/SDesignerToolBar.cpp b/Engine/Source/Editor/UMGEditor/Private/Designer/SDesignerToolBar.cpp
--- a/Engine/Source/Editor/UMGEditor/Private/Designer/SDesignerToolBar.cpp
+++ b/Engine/Source/Editor/UMGEditor/Private/Designer/SDesignerToolBar.cpp
@@ -38,8 +38,8 @@ TSharedRef< SWidget > SDesignerToolBar::MakeToolBar(const TSharedPtr< FExtender
 {
 	FToolBarBuilder ToolbarBuilder( CommandList, FMultiBoxCustomization::None, InExtenders );
 
-	// Use a custom style
-	FName ToolBarStyle = "LegacyViewportMenu";
+	// Use a custom style	
+	FName ToolBarStyle = "EditorViewportToolBar";
 	ToolbarBuilder.SetStyle(&FEditorStyle::Get(), ToolBarStyle);
 	ToolbarBuilder.SetLabelVisibility(EVisibility::Collapsed);
 
@@ -48,7 +48,6 @@ TSharedRef< SWidget > SDesignerToolBar::MakeToolBar(const TSharedPtr< FExtender
 		FUICommandInfo* ToggleLocalizationPreviewCommand = FDesignerCommands::Get().ToggleLocalizationPreview.Get();
 
 		ToolbarBuilder.AddWidget(SNew(SViewportToolBarComboMenu)
-			.Style(ToolBarStyle)
 			.BlockLocation(EMultiBlockLocation::Start)
 			.Cursor(EMouseCursor::Default)
 			.IsChecked(this, &SDesignerToolBar::IsLocalizationPreviewChecked)
@@ -67,8 +66,10 @@ TSharedRef< SWidget > SDesignerToolBar::MakeToolBar(const TSharedPtr< FExtender
 	ToolbarBuilder.SetIsFocusable( false );
 
 	ToolbarBuilder.BeginSection("View");
+	ToolbarBuilder.BeginBlockGroup();
 	ToolbarBuilder.AddToolBarButton(FDesignerCommands::Get().ToggleOutlines, NAME_None, TAttribute<FText>(), TAttribute<FText>(), TAttribute<FSlateIcon>(), "ToggleOutlines");
 	ToolbarBuilder.AddToolBarButton(FDesignerCommands::Get().ToggleRespectLocks, NAME_None, TAttribute<FText>(), TAttribute<FText>(), TAttribute<FSlateIcon>(), "ToggleRespectLocks");
+	ToolbarBuilder.EndBlockGroup();
 	ToolbarBuilder.EndSection();
 
 	ToolbarBuilder.BeginSection("Transform");
diff --git a/Engine/Source/Editor/UMGEditor/Private/Designer/SDesignerView.cpp b/Engine/Source/Editor/UMGEditor/Private/Designer/SDesignerView.cpp
--- a/Engine/Source/Editor/UMGEditor/Private/Designer/SDesignerView.cpp
+++ b/Engine/Source/Editor/UMGEditor/Private/Designer/SDesignerView.cpp
@@ -516,7 +516,7 @@ EActiveTimerReturnType SDesignerView::EnsureTick(double InCurrentTime, float InD
 
 TSharedRef<SWidget> SDesignerView::CreateOverlayUI()
 {
-	const FToolBarStyle& ToolBarStyle = FEditorStyle::Get().GetWidgetStyle<FToolBarStyle>("LegacyViewportMenu");
+	const FToolBarStyle& ToolBarStyle = FEditorStyle::Get().GetWidgetStyle<FToolBarStyle>("EditorViewportToolBar");
 
 	return SNew(SOverlay)
 
@@ -593,30 +593,31 @@ TSharedRef<SWidget> SDesignerView::CreateOverlayUI()
 		]
 
 		+ SHorizontalBox::Slot()
+		.Padding(0.0f, 1.0f)
 		.AutoWidth()
-		.VAlign(VAlign_Center)
 		[
 			SNew(SDesignerToolBar)
 			.CommandList(CommandList)
 		]
 		+ SHorizontalBox::Slot()
-			.AutoWidth()
+		.Padding(0.0f, 1.0f)
+		.AutoWidth()
+		[
+			SNew(SButton)
+			.ButtonStyle(&ToolBarStyle.ButtonStyle)
+			.ToolTipText(LOCTEXT("ZoomToFit_ToolTip", "Zoom To Fit"))
+			.OnClicked(this, &SDesignerView::HandleZoomToFitClicked)
+			.ContentPadding(ToolBarStyle.ButtonPadding)
 			.VAlign(VAlign_Center)
 			[
-				SNew(SButton)
-				.ButtonStyle(&ToolBarStyle.ButtonStyle)
-				.ToolTipText(LOCTEXT("ZoomToFit_ToolTip", "Zoom To Fit"))
-				.OnClicked(this, &SDesignerView::HandleZoomToFitClicked)
-				.ContentPadding(ToolBarStyle.ButtonPadding)
-				[
-					SNew(SImage)
-					.Image(FEditorStyle::GetBrush("UMGEditor.ZoomToFit"))
-				]
+				SNew(SImage)
+				.Image(FEditorStyle::GetBrush("UMGEditor.ZoomToFit"))
+				.ColorAndOpacity(FSlateColor::UseForeground())
 			]
+		]
 
 		+ SHorizontalBox::Slot()
 			.AutoWidth()
-			.VAlign(VAlign_Center)
 			[
 				SNew(SButton)
 				.ButtonStyle(&ToolBarStyle.ButtonStyle)
@@ -624,14 +625,15 @@ TSharedRef<SWidget> SDesignerView::CreateOverlayUI()
 				.OnClicked(this, &SDesignerView::HandleSwapAspectRatioClicked)
 				.ContentPadding(ToolBarStyle.ButtonPadding)
 				.IsEnabled(this, &SDesignerView::GetAspectRatioSwitchEnabled)
+				.VAlign(VAlign_Center)
 				[
 					SNew(SImage)
 					.Image(this, &SDesignerView::GetAspectRatioSwitchImage)
+					.ColorAndOpacity(FSlateColor::UseForeground())
 				]
 			]
 		+ SHorizontalBox::Slot()
 			.AutoWidth()
-			.VAlign(VAlign_Center)
 			[
 				SNew(SButton)
 				.ButtonStyle(&ToolBarStyle.ButtonStyle)
@@ -639,20 +641,21 @@ TSharedRef<SWidget> SDesignerView::CreateOverlayUI()
 				.OnClicked(this, &SDesignerView::HandleFlipSafeZonesClicked)
 				.ContentPadding(ToolBarStyle.ButtonPadding)
 				.IsEnabled(this, &SDesignerView::GetFlipDeviceEnabled)
+				.VAlign(VAlign_Center)
 				[
 					SNew(SImage)
 					.Image(FEditorStyle::Get().GetBrush("UMGEditor.Mirror"))
+					.ColorAndOpacity(FSlateColor::UseForeground())
 				]
 			]
 
 		// Preview Screen Size
 		+ SHorizontalBox::Slot()
+		.Padding(2.0f,0.0f)
 		.AutoWidth()
-		.VAlign(VAlign_Center)
 		[
 			SNew(SComboButton)
 			.ButtonStyle(&ToolBarStyle.ButtonStyle)
-			.ForegroundColor(FLinearColor::Black)
 			.OnGetMenuContent(this, &SDesignerView::GetResolutionsMenu)
 			.ContentPadding(ToolBarStyle.ButtonPadding)
 			.ButtonContent()
@@ -660,17 +663,17 @@ TSharedRef<SWidget> SDesignerView::CreateOverlayUI()
 				SNew(STextBlock)
 				.Text(LOCTEXT("ScreenSize", "Screen Size"))
 				.TextStyle(&ToolBarStyle.LabelStyle)
+				.ColorAndOpacity(FSlateColor::UseForeground())
 			]
 		]
 
 		// Screen Fill Size Rule
 		+ SHorizontalBox::Slot()
+		.Padding(2.0f, 0.0f)
 		.AutoWidth()
-		.VAlign(VAlign_Center)
 		[
 			SNew(SComboButton)
 			.ButtonStyle(&ToolBarStyle.ButtonStyle)
-			.ForegroundColor(FLinearColor::Black)
 			.OnGetMenuContent(this, &SDesignerView::GetScreenSizingFillMenu)
 			.ContentPadding(ToolBarStyle.ButtonPadding)
 			.ButtonContent()
@@ -678,21 +681,23 @@ TSharedRef<SWidget> SDesignerView::CreateOverlayUI()
 				SNew(STextBlock)
 				.Text(this, &SDesignerView::GetScreenSizingFillText)
 				.TextStyle(&ToolBarStyle.LabelStyle)
+				.ColorAndOpacity(FSlateColor::UseForeground())
 			]
 		]
 		+ SHorizontalBox::Slot()
 		.AutoWidth()
-		.VAlign(VAlign_Center)
 		.Padding(FMargin(2, 0))
+		.VAlign(VAlign_Center)
 		[
 			SNew(STextBlock)
 			.Visibility(this, &SDesignerView::GetCustomResolutionEntryVisibility)
 			.Text(LOCTEXT("Width", "Width"))
+			.ColorAndOpacity(FSlateColor::UseForeground())
 		]
 		+ SHorizontalBox::Slot()
 		.AutoWidth()
-		.VAlign(VAlign_Center)
 		.Padding(FMargin(2, 0))
+		.VAlign(VAlign_Center)
 		[
 			SNew(SNumericEntryBox<int32>)
 			.AllowSpin(true)
@@ -708,17 +713,18 @@ TSharedRef<SWidget> SDesignerView::CreateOverlayUI()
 		]
 		+ SHorizontalBox::Slot()
 		.AutoWidth()
-		.VAlign(VAlign_Center)
 		.Padding(FMargin(2, 0))
+		.VAlign(VAlign_Center)
 		[
 			SNew(STextBlock)
 			.Visibility(this, &SDesignerView::GetCustomResolutionEntryVisibility)
 			.Text(LOCTEXT("Height", "Height"))
+			.ColorAndOpacity(FSlateColor::UseForeground())
 		]
 		+ SHorizontalBox::Slot()
 		.AutoWidth()
-		.VAlign(VAlign_Center)
 		.Padding(FMargin(2, 0))
+		.VAlign(VAlign_Center)
 		[
 			SNew(SNumericEntryBox<int32>)
 			.AllowSpin(true)
diff --git a/Engine/Source/Editor/UMGEditor/Private/Details/SWidgetDetailsView.cpp b/Engine/Source/Editor/UMGEditor/Private/Details/SWidgetDetailsView.cpp
--- a/Engine/Source/Editor/UMGEditor/Private/Details/SWidgetDetailsView.cpp
+++ b/Engine/Source/Editor/UMGEditor/Private/Details/SWidgetDetailsView.cpp
@@ -122,6 +122,7 @@ void SWidgetDetailsView::Construct(const FArguments& InArgs, TSharedPtr<FWidgetB
 					[
 						SNew(SImage)
 						.Image(this, &SWidgetDetailsView::GetNameIcon)
+						.ColorAndOpacity(FSlateColor::UseForeground())
 					]
 
 					+ SHorizontalBox::Slot()
@@ -157,6 +158,7 @@ void SWidgetDetailsView::Construct(const FArguments& InArgs, TSharedPtr<FWidgetB
 
 					+ SHorizontalBox::Slot()
 					.AutoWidth()
+					.VAlign(VAlign_Center)
 					.Padding(15,0,0,0)
 					[
 						SAssignNew(ClassLinkArea, SBox)
diff --git a/Engine/Source/Editor/UMGEditor/Private/Hierarchy/SHierarchyViewItem.cpp b/Engine/Source/Editor/UMGEditor/Private/Hierarchy/SHierarchyViewItem.cpp
--- a/Engine/Source/Editor/UMGEditor/Private/Hierarchy/SHierarchyViewItem.cpp
+++ b/Engine/Source/Editor/UMGEditor/Private/Hierarchy/SHierarchyViewItem.cpp
@@ -1324,7 +1324,7 @@ void SHierarchyViewItem::Construct(const FArguments& InArgs, const TSharedRef< S
 			.VAlign(VAlign_Center)
 			[
 				SNew(SImage)
-				.ColorAndOpacity(FLinearColor(1,1,1,0.5))
+				.ColorAndOpacity(FSlateColor::UseForeground())
 				.Image(Model->GetImage())
 				.ToolTipText(Model->GetImageToolTipText())
 			]
diff --git a/Engine/Source/Editor/UMGEditor/Private/Palette/SPaletteView.cpp b/Engine/Source/Editor/UMGEditor/Private/Palette/SPaletteView.cpp
--- a/Engine/Source/Editor/UMGEditor/Private/Palette/SPaletteView.cpp
+++ b/Engine/Source/Editor/UMGEditor/Private/Palette/SPaletteView.cpp
@@ -69,6 +69,11 @@ void SPaletteViewItem::OnFavoriteToggled(ECheckBoxState InNewState)
 	}
 }
 
+EVisibility SPaletteViewItem::GetFavoritedStateVisibility() const
+{
+	return GetFavoritedState() == ECheckBoxState::Checked || IsHovered() ? EVisibility::Visible : EVisibility::Hidden;
+}
+
 void SPaletteViewItem::Construct(const FArguments& InArgs, TSharedPtr<FWidgetTemplateViewModel> InWidgetViewModel)
 {
 	WidgetViewModel = InWidgetViewModel;
@@ -86,13 +91,14 @@ void SPaletteViewItem::Construct(const FArguments& InArgs, TSharedPtr<FWidgetTem
 				.IsChecked(this, &SPaletteViewItem::GetFavoritedState)
 				.OnCheckStateChanged(this, &SPaletteViewItem::OnFavoriteToggled)
 				.Style(FEditorStyle::Get(), "UMGEditor.Palette.FavoriteToggleStyle")
+				.Visibility(this, &SPaletteViewItem::GetFavoritedStateVisibility)
 			]
 			+ SHorizontalBox::Slot()
 			.AutoWidth()
 			.VAlign(VAlign_Center)
 			[
 				SNew(SImage)
-				.ColorAndOpacity(FLinearColor(1, 1, 1, 0.5))
+				.ColorAndOpacity(FSlateColor::UseForeground())
 				.Image(WidgetViewModel->Template->GetIcon())
 			]
 
@@ -161,7 +167,12 @@ void SPaletteView::Construct(const FArguments& InArgs, TSharedPtr<FWidgetBluepri
 		[
 			SNew(SScrollBorder, WidgetTemplatesView.ToSharedRef())
 			[
-				WidgetTemplatesView.ToSharedRef()
+				SNew(SBorder)
+				.BorderImage(FAppStyle::Get().GetBrush("Brushes.Recessed"))
+				.Padding(0)
+				[
+					WidgetTemplatesView.ToSharedRef()
+				]
 			]
 		]
 	];
diff --git a/Engine/Source/Editor/UMGEditor/Private/Palette/SPaletteView.h b/Engine/Source/Editor/UMGEditor/Private/Palette/SPaletteView.h
--- a/Engine/Source/Editor/UMGEditor/Private/Palette/SPaletteView.h
+++ b/Engine/Source/Editor/UMGEditor/Private/Palette/SPaletteView.h
@@ -60,6 +60,8 @@ private:
 	 */
 	void OnFavoriteToggled(ECheckBoxState InNewState);
 
+	EVisibility GetFavoritedStateVisibility() const;
+
 	virtual FReply OnMouseButtonDoubleClick(const FGeometry& InMyGeometry, const FPointerEvent& InMouseEvent) override;
 
 private:
diff --git a/Engine/Source/Editor/UMGEditor/Private/Palette/SPaletteViewModel.cpp b/Engine/Source/Editor/UMGEditor/Private/Palette/SPaletteViewModel.cpp
--- a/Engine/Source/Editor/UMGEditor/Private/Palette/SPaletteViewModel.cpp
+++ b/Engine/Source/Editor/UMGEditor/Private/Palette/SPaletteViewModel.cpp
@@ -51,7 +51,6 @@ TSharedRef<ITableRow> FWidgetTemplateViewModel::BuildRow(const TSharedRef<STable
 {
 	return SNew(STableRow<TSharedPtr<FWidgetViewModel>>, OwnerTable)
 		.Padding(2.0f)
-		.Style(FEditorStyle::Get(), "UMGEditor.PaletteItem")
 		.OnDragDetected(this, &FWidgetTemplateViewModel::OnDraggingWidgetTemplateItem)
 		[
 			SNew(SPaletteViewItem, SharedThis(this))
@@ -80,13 +79,13 @@ TSharedRef<ITableRow> FWidgetHeaderViewModel::BuildRow(const TSharedRef<STableVi
 {
 	return SNew(STableRow<TSharedPtr<FWidgetViewModel>>, OwnerTable)
 		.Style(FEditorStyle::Get(), "UMGEditor.PaletteHeader")
-		.Padding(2.0f)
+		.Padding(5.0f)
 		.ShowSelection(false)
 		[
 			SNew(STextBlock)
+			.TransformPolicy(ETextTransformPolicy::ToUpper)
 			.Text(GroupName)
-			.Font(FEditorStyle::GetFontStyle("DetailsView.CategoryFontStyle"))
-			.ShadowOffset(FVector2D(1.0f, 1.0f))
+			.Font(FAppStyle::Get().GetFontStyle("SmallFontBold"))
 		];
 }
 
diff --git a/Engine/Source/Editor/UMGEditor/Private/TabFactory/AnimationTabSummoner.cpp b/Engine/Source/Editor/UMGEditor/Private/TabFactory/AnimationTabSummoner.cpp
--- a/Engine/Source/Editor/UMGEditor/Private/TabFactory/AnimationTabSummoner.cpp
+++ b/Engine/Source/Editor/UMGEditor/Private/TabFactory/AnimationTabSummoner.cpp
@@ -27,6 +27,7 @@
 #include "Kismet2/BlueprintEditorUtils.h"
 #include "Misc/TextFilter.h"
 #include "Kismet2/Kismet2NameValidators.h"
+#include "SEditorHeaderButton.h"
 
 #define LOCTEXT_NAMESPACE "UMG"
 
@@ -285,33 +286,9 @@ public:
 					.VAlign( VAlign_Center )
 					.AutoWidth()
 					[
-						SNew(SButton)
-						.ButtonStyle(FEditorStyle::Get(), "FlatButton.Success")
-						.ForegroundColor(FEditorStyle::Get().GetSlateColor("Foreground"))
-						.ContentPadding(FMargin(2.0f, 1.0f))
+						SNew(SEditorHeaderButton)
 						.OnClicked( this, &SUMGAnimationList::OnNewAnimationClicked )
-						.HAlign(HAlign_Center)
-						.VAlign(VAlign_Center)
-						[
-							SNew( SHorizontalBox )
-							+ SHorizontalBox::Slot()
-							.VAlign(VAlign_Center)
-							.AutoWidth()
-							[
-								SNew(STextBlock)
-								.TextStyle(FEditorStyle::Get(), "NormalText.Important")
-								.Font(FEditorStyle::Get().GetFontStyle("FontAwesome.10"))
-								.Text(FText::FromString(FString(TEXT("\xf067"))) /*fa-plus*/)
-							]
-
-							+ SHorizontalBox::Slot()
-							.Padding( 2.0f, 0.0f )
-							[
-								SNew( STextBlock )
-								.TextStyle(FEditorStyle::Get(), "NormalText.Important")
-								.Text( LOCTEXT("NewAnimationButtonText", "Animation") )
-							]
-						]
+						.Text( LOCTEXT("NewAnimationButtonText", "Animation") )
 					]
 					+ SHorizontalBox::Slot()
 					.Padding(2.0f, 0.0f)
diff --git a/Engine/Source/Editor/UnrealEd/Private/SViewportToolBarComboMenu.cpp b/Engine/Source/Editor/UnrealEd/Private/SViewportToolBarComboMenu.cpp
--- a/Engine/Source/Editor/UnrealEd/Private/SViewportToolBarComboMenu.cpp
+++ b/Engine/Source/Editor/UnrealEd/Private/SViewportToolBarComboMenu.cpp
@@ -34,14 +34,9 @@ void SViewportToolBarComboMenu::Construct( const FArguments& InArgs )
 		.ToolTipText(InArgs._ToggleButtonToolTip)
 		.IsChecked(InArgs._IsChecked)
 		[
-			SNew( SBox )
-			.HAlign(HAlign_Center)
-			.VAlign(VAlign_Center)
-			[
-				SNew(SImage)
-				.Image(Icon.GetIcon())
-				.ColorAndOpacity(FSlateColor::UseForeground())
-			]
+			SNew(SImage)
+			.Image(Icon.GetIcon())
+			.ColorAndOpacity(FSlateColor::UseForeground())
 		];
 	}
 
@@ -88,12 +83,14 @@ void SViewportToolBarComboMenu::Construct( const FArguments& InArgs )
 
 			+SHorizontalBox::Slot()
 			.VAlign(VAlign_Center)
+			.Padding(2.0f,0.0f)
 			.AutoWidth()
 			[
 				ToggleControl.ToSharedRef()
 			]
 
 			+SHorizontalBox::Slot()
+			.Padding(2.0f, 0.0f)
 			.VAlign(VAlign_Center)
 			.AutoWidth()
 			[
