Fix redundant usage of styles related to various "favorites" systems

diff --git a/Engine/Plugins/Editor/GameplayTagsEditor/Source/GameplayTagsEditor/Private/SAddNewGameplayTagWidget.cpp b/Engine/Plugins/Editor/GameplayTagsEditor/Source/GameplayTagsEditor/Private/SAddNewGameplayTagWidget.cpp
--- a/Engine/Plugins/Editor/GameplayTagsEditor/Source/GameplayTagsEditor/Private/SAddNewGameplayTagWidget.cpp
+++ b/Engine/Plugins/Editor/GameplayTagsEditor/Source/GameplayTagsEditor/Private/SAddNewGameplayTagWidget.cpp
@@ -185,7 +185,7 @@ const FSlateBrush* SAddNewGameplayTagWidget::OnGetTagSourceFavoriteImage() const
 	const FName ActiveTagSource = *TagSourcesComboBox->GetSelectedItem().Get();
 	const bool bIsFavoriteTagSource = FGameplayTagSource::GetFavoriteName() == ActiveTagSource;
 
-	return FEditorStyle::GetBrush(bIsFavoriteTagSource ? TEXT("PropertyWindow.Favorites_Enabled") : TEXT("PropertyWindow.Favorites_Disabled"));
+	return FEditorStyle::GetBrush(bIsFavoriteTagSource ? TEXT("Icons.Star") : TEXT("PropertyWindow.Favorites_Disabled"));
 }
 
 void SAddNewGameplayTagWidget::Tick( const FGeometry& AllottedGeometry, const double InCurrentTime, const float InDeltaTime )
diff --git a/Engine/Plugins/Experimental/ActorPalette/Source/ActorPalette/Private/ActorPaletteViewport.cpp b/Engine/Plugins/Experimental/ActorPalette/Source/ActorPalette/Private/ActorPaletteViewport.cpp
--- a/Engine/Plugins/Experimental/ActorPalette/Source/ActorPalette/Private/ActorPaletteViewport.cpp
+++ b/Engine/Plugins/Experimental/ActorPalette/Source/ActorPalette/Private/ActorPaletteViewport.cpp
@@ -66,7 +66,7 @@ public:
 						.Image_Lambda([this]()
 						{
 							return FEditorStyle::GetBrush(GetDefault<UActorPaletteSettings>()->FavoritesList.Contains(AssetData.ObjectPath.ToString()) ?
-								TEXT("PropertyWindow.Favorites_Enabled") :
+								TEXT("Icons.Star") :
 								TEXT("PropertyWindow.Favorites_Disabled"));
 						})
 					]
diff --git a/Engine/Source/Editor/ContentBrowser/Private/PathContextMenu.cpp b/Engine/Source/Editor/ContentBrowser/Private/PathContextMenu.cpp
--- a/Engine/Source/Editor/ContentBrowser/Private/PathContextMenu.cpp
+++ b/Engine/Source/Editor/ContentBrowser/Private/PathContextMenu.cpp
@@ -198,7 +198,7 @@ void FPathContextMenu::MakePathViewContextMenu(UToolMenu* Menu)
 					"AddToFavorites",
 					LOCTEXT("AddToFavorites", "Add To Favorites"),
 					LOCTEXT("AddToFavoritesTooltip", "Adds this folder to the favorites section for easy access."),
-					FSlateIcon(FEditorStyle::GetStyleSetName(), "PropertyWindow.Favorites_Enabled"),
+					FSlateIcon(FEditorStyle::GetStyleSetName(), "Icons.Star"),
 					FUIAction(FExecuteAction::CreateSP(this, &FPathContextMenu::ExecuteFavorite))
 				);
 			}
diff --git a/Engine/Source/Editor/EditorStyle/Private/StarshipStyle.cpp b/Engine/Source/Editor/EditorStyle/Private/StarshipStyle.cpp
--- a/Engine/Source/Editor/EditorStyle/Private/StarshipStyle.cpp
+++ b/Engine/Source/Editor/EditorStyle/Private/StarshipStyle.cpp
@@ -3118,7 +3118,6 @@ void FStarshipEditorStyle::FStyle::SetupPropertyEditorStyles()
 
 		Set( "PropertyWindow.FilterSearch", new IMAGE_BRUSH( "Old/FilterSearch", Icon16x16 ) );
 		Set( "PropertyWindow.FilterCancel", new IMAGE_BRUSH( "Old/FilterCancel", Icon16x16 ) );
-		Set( "PropertyWindow.Favorites_Enabled", new IMAGE_BRUSH( "Icons/Star_16x", Icon16x16 ) );
 		Set( "PropertyWindow.Favorites_Disabled", new IMAGE_BRUSH( "Icons/EmptyStar_16x", Icon16x16 ) );
 		Set( "PropertyWindow.Locked", new CORE_IMAGE_BRUSH_SVG( "Starship/Common/lock", Icon16x16 ) );
 		Set( "PropertyWindow.Unlocked", new CORE_IMAGE_BRUSH_SVG( "Starship/Common/lock-unlocked", Icon16x16 ) );
@@ -5174,9 +5173,6 @@ void FStarshipEditorStyle::FStyle::SetupPersonaStyle()
 			Set( "Kismet.Tabs.BlueprintDefaults", new IMAGE_BRUSH( "Icons/icon_BlueprintEditor_Defaults_40x", Icon16x16 ) );
 		}
 
-		// this should be moved to a generic icon in the core style. Icons.Star
-		Set("Kismet.Palette.Favorites", new IMAGE_BRUSH_SVG("Starship/Common/Favorite", Icon16x16));
-
 		const FCheckBoxStyle KismetFavoriteToggleStyle = FCheckBoxStyle()
 			.SetCheckBoxType(ESlateCheckBoxType::CheckBox)
 			.SetUncheckedImage( IMAGE_BRUSH("Icons/EmptyStar_16x", Icon10x10, FLinearColor(0.8f, 0.8f, 0.8f, 1.f)) )
diff --git a/Engine/Source/Editor/Kismet/Private/SBlueprintFavoritesPalette.cpp b/Engine/Source/Editor/Kismet/Private/SBlueprintFavoritesPalette.cpp
--- a/Engine/Source/Editor/Kismet/Private/SBlueprintFavoritesPalette.cpp
+++ b/Engine/Source/Editor/Kismet/Private/SBlueprintFavoritesPalette.cpp
@@ -289,7 +289,7 @@ void SBlueprintFavoritesPalette::Construct(FArguments const& InArgs, TWeakPtr<FB
 {
 	SBlueprintSubPalette::FArguments SuperArgs;
 	SuperArgs._Title       = LOCTEXT("PaletteTitle", "Favorites");
-	SuperArgs._Icon        = FEditorStyle::GetBrush("Kismet.Palette.Favorites");
+	SuperArgs._Icon        = FEditorStyle::GetBrush("Icons.Star");
 	SuperArgs._ToolTipText = LOCTEXT("PaletteToolTip", "A listing of your favorite and most used nodes.");
 
 	static FString const ShowFreqUsedConfigKey("bShowFrequentlyUsed");
diff --git a/Engine/Source/Editor/PropertyEditor/Private/SDetailsView.cpp b/Engine/Source/Editor/PropertyEditor/Private/SDetailsView.cpp
--- a/Engine/Source/Editor/PropertyEditor/Private/SDetailsView.cpp
+++ b/Engine/Source/Editor/PropertyEditor/Private/SDetailsView.cpp
@@ -262,7 +262,7 @@ void SDetailsView::Construct(const FArguments& InArgs, const FDetailsViewArgs& I
 				[
 					SNew(SImage)
 					.ColorAndOpacity(this, &SDetailsView::GetToggleFavoritesColor)
-					.Image(FAppStyle::Get().GetBrush("PropertyWindow.Favorites_Enabled"))
+					.Image(FAppStyle::Get().GetBrush("Icons.Star"))
 				]
 			];
 	}
diff --git a/Engine/Source/Editor/PropertyEditor/Private/SPropertyTreeViewImpl.cpp b/Engine/Source/Editor/PropertyEditor/Private/SPropertyTreeViewImpl.cpp
--- a/Engine/Source/Editor/PropertyEditor/Private/SPropertyTreeViewImpl.cpp
+++ b/Engine/Source/Editor/PropertyEditor/Private/SPropertyTreeViewImpl.cpp
@@ -381,7 +381,7 @@ const FSlateBrush* SPropertyTreeViewImpl::OnGetFavoriteButtonImageResource() con
 {
 	if( bFavoritesEnabled )
 	{
-		return FEditorStyle::GetBrush(TEXT("PropertyWindow.Favorites_Enabled"));
+		return FEditorStyle::GetBrush(TEXT("Icons.Star"));
 	}
 	else
 	{
diff --git a/Engine/Source/Editor/PropertyEditor/Private/UserInterface/PropertyEditor/SPropertyEditorTableRow.cpp b/Engine/Source/Editor/PropertyEditor/Private/UserInterface/PropertyEditor/SPropertyEditorTableRow.cpp
--- a/Engine/Source/Editor/PropertyEditor/Private/UserInterface/PropertyEditor/SPropertyEditorTableRow.cpp
+++ b/Engine/Source/Editor/PropertyEditor/Private/UserInterface/PropertyEditor/SPropertyEditorTableRow.cpp
@@ -189,7 +189,7 @@ const FSlateBrush* SPropertyEditorTableRow::OnGetFavoriteImage() const
 {
 	if( PropertyEditor->IsFavorite() )
 	{
-		return FEditorStyle::GetBrush(TEXT("PropertyWindow.Favorites_Enabled"));
+		return FEditorStyle::GetBrush(TEXT("Icons.Star"));
 	}
 
 	return FEditorStyle::GetBrush(TEXT("PropertyWindow.Favorites_Disabled"));
