Adding missing icon import lines

diff --git a/Engine/Source/Editor/EditorStyle/Private/StarshipStyle.cpp b/Engine/Source/Editor/EditorStyle/Private/StarshipStyle.cpp
--- a/Engine/Source/Editor/EditorStyle/Private/StarshipStyle.cpp
+++ b/Engine/Source/Editor/EditorStyle/Private/StarshipStyle.cpp
@@ -7026,6 +7026,13 @@ void FStarshipEditorStyle::FStyle::SetupToolkitStyles()
 		Set("CurveEditorTools.ActivateMultiScaleTool.Small", new IMAGE_BRUSH("GenericCurveEditor/Icons/ToolsMultiScale_48x", Icon24x24));
 	}
 
+	// General Curve Icons
+	{
+		Set("Curve.ZoomToFit", new IMAGE_BRUSH_SVG("Starship/Common/ZoomToFit", Icon20x20));
+		Set("Curve.ZoomToFitHorizontal", new IMAGE_BRUSH_SVG("Starship/Common/ZoomToFitHorizontal", Icon20x20));
+		Set("Curve.ZoomToFitVertical", new IMAGE_BRUSH_SVG("Starship/Common/ZoomToFitVertical", Icon20x20));
+	}
+
 	// PhysicsAssetEditor
 	{
 		Set( "PhysicsAssetEditor.Tabs.Properties", new IMAGE_BRUSH( "/Icons/icon_tab_SelectionDetails_16x", Icon16x16 ) );
