Updating the default project thumbnail

diff --git a/Engine/Source/Editor/DetailCustomizations/Private/GeneralProjectSettingsDetails.cpp b/Engine/Source/Editor/DetailCustomizations/Private/GeneralProjectSettingsDetails.cpp
--- a/Engine/Source/Editor/DetailCustomizations/Private/GeneralProjectSettingsDetails.cpp
+++ b/Engine/Source/Editor/DetailCustomizations/Private/GeneralProjectSettingsDetails.cpp
@@ -32,7 +32,7 @@ void FGeneralProjectSettingsDetails::CustomizeDetails(IDetailLayoutBuilder& Deta
 	FString ProjectThumbnail_AutomaticImagePath = FPaths::ProjectSavedDir() / TEXT("AutoScreenshot.png");
 	if (!FPaths::FileExists(ProjectThumbnail_AutomaticImagePath))
 	{
-		ProjectThumbnail_AutomaticImagePath = FPaths::EngineContentDir() / TEXT("Editor") / TEXT("Slate") / TEXT("GameProjectDialog") / TEXT("default_game_thumbnail_192x.png");
+		ProjectThumbnail_AutomaticImagePath = FPaths::EngineContentDir() / TEXT("Editor") / TEXT("Slate") / TEXT("GameProjectDialog") / TEXT("default_game_thumbnail.png");
 	}
 
 	ImageCategory.AddCustomRow(ProjectThumbnailDesc)
diff --git a/Engine/Source/Editor/EditorStyle/Private/SlateEditorStyle.cpp b/Engine/Source/Editor/EditorStyle/Private/SlateEditorStyle.cpp
--- a/Engine/Source/Editor/EditorStyle/Private/SlateEditorStyle.cpp
+++ b/Engine/Source/Editor/EditorStyle/Private/SlateEditorStyle.cpp
@@ -1926,7 +1926,7 @@ void FSlateEditorStyle::FStyle::SetupGeneralStyles()
 			.SetSelectedTextColor( InvertedForeground )
 			);
 
-		Set( "GameProjectDialog.DefaultGameThumbnail", new IMAGE_BRUSH( "GameProjectDialog/default_game_thumbnail_128x", Icon128x128 ) );
+		Set( "GameProjectDialog.DefaultGameThumbnail", new IMAGE_BRUSH( "GameProjectDialog/default_game_thumbnail", Icon128x128 ) );
 		Set( "GameProjectDialog.DefaultGameThumbnail.Small", new IMAGE_BRUSH( "GameProjectDialog/default_game_thumbnail", Icon128x128 ) );
 		Set( "GameProjectDialog.BlankProjectThumbnail", new IMAGE_BRUSH( "GameProjectDialog/blank_project_thumbnail", Icon128x128 ) );
 		Set( "GameProjectDialog.BlankProjectPreview", new IMAGE_BRUSH( "GameProjectDialog/blank_project_preview", FVector2D(400, 200) ) );
