Level Editor Menu Clean Up, Part 2

diff --git a/Engine/Plugins/Animation/GameplayInsights/Source/GameplayInsights/Private/GameplayInsightsModule.cpp b/Engine/Plugins/Animation/GameplayInsights/Source/GameplayInsights/Private/GameplayInsightsModule.cpp
--- a/Engine/Plugins/Animation/GameplayInsights/Source/GameplayInsights/Private/GameplayInsightsModule.cpp
+++ b/Engine/Plugins/Animation/GameplayInsights/Source/GameplayInsights/Private/GameplayInsightsModule.cpp
@@ -25,6 +25,7 @@
 #include "LevelEditorMenuContext.h"
 #include "Engine/Selection.h"
 #include "SSCSEditorMenuContext.h"
+#include "GameplayInsightsStyle.h"
 #include "SSCSEditor.h"
 #endif
 
@@ -82,6 +83,8 @@ void FGameplayInsightsModule::StartupModule()
 
 		FInsightsMajorTabConfig TimingProfilerConfig;
 		TimingProfilerConfig.TabLabel = LOCTEXT("AnimationInsightsTabName", "Animation Insights");
+		TimingProfilerConfig.TabIcon = FSlateIcon(FGameplayInsightsStyle::Get().GetStyleSetName(), "AnimationInsights.TabIcon");
+
 		TimingProfilerConfig.TabTooltip = LOCTEXT("AnimationInsightsTabTooltip", "Open the Animation Insights tab.");
 		TimingProfilerConfig.Layout = FTabManager::NewLayout("GameplayInsightsTimingLayout_v1.2")
 		->AddArea
diff --git a/Engine/Plugins/Animation/GameplayInsights/Source/GameplayInsights/Private/GameplayInsightsStyle.h b/Engine/Plugins/Animation/GameplayInsights/Source/GameplayInsights/Private/GameplayInsightsStyle.h
--- a/Engine/Plugins/Animation/GameplayInsights/Source/GameplayInsights/Private/GameplayInsightsStyle.h
+++ b/Engine/Plugins/Animation/GameplayInsights/Source/GameplayInsights/Private/GameplayInsightsStyle.h
@@ -4,10 +4,8 @@
 
 #include "Styling/SlateStyle.h"
 #include "Styling/SlateStyleRegistry.h"
+#include "Styling/SlateStyleMacros.h"
 
-#define IMAGE_BRUSH(RelativePath, ...) FSlateImageBrush(RootToContentDir(RelativePath, TEXT(".png")), __VA_ARGS__)
-#define BOX_BRUSH(RelativePath, ...) FSlateBoxBrush(RootToContentDir(RelativePath, TEXT(".png")), __VA_ARGS__)
-#define BORDER_BRUSH(RelativePath, ...) FSlateBorderBrush(RootToContentDir(RelativePath, TEXT(".png")), __VA_ARGS__)
 #define TTF_FONT(RelativePath, ...) FSlateFontInfo(RootToContentDir(RelativePath, TEXT(".ttf")), __VA_ARGS__)
 #define OTF_FONT(RelativePath, ...) FSlateFontInfo(RootToContentDir(RelativePath, TEXT(".otf")), __VA_ARGS__)
 
@@ -32,6 +30,8 @@ public:
 
 		const FSlateColor DefaultForeground(FLinearColor(0.72f, 0.72f, 0.72f, 1.f));
 
+		Set("AnimationInsights.TabIcon", new IMAGE_BRUSH_SVG("AnimationProfiler", Icon16x16));
+
 		Set("SchematicViewRootLeft", new BOX_BRUSH("SchematicViewRootLeft", FMargin(4/16.0f, 8.0f/16.0f, 4/16.0f, 4/16.0f)));
 		Set("SchematicViewRootMid", new BOX_BRUSH("SchematicViewRootMid", FMargin(4/16.0f, 8.0f/16.0f, 4/16.0f, 4/16.0f)));
 
diff --git a/Engine/Plugins/Cameras/CameraShakePreviewer/Source/CameraShakePreviewer/Private/CameraShakePreviewerModule.cpp b/Engine/Plugins/Cameras/CameraShakePreviewer/Source/CameraShakePreviewer/Private/CameraShakePreviewerModule.cpp
--- a/Engine/Plugins/Cameras/CameraShakePreviewer/Source/CameraShakePreviewer/Private/CameraShakePreviewerModule.cpp
+++ b/Engine/Plugins/Cameras/CameraShakePreviewer/Source/CameraShakePreviewer/Private/CameraShakePreviewerModule.cpp
@@ -86,7 +86,7 @@ void FCameraShakePreviewerModule::RegisterEditorTab()
 		TSharedPtr<FTabManager> LevelEditorTabManager = LevelEditorModule.GetLevelEditorTabManager();
 
 		const IWorkspaceMenuStructure& MenuStructure = WorkspaceMenu::GetMenuStructure();
-		const FSlateIcon Icon(FEditorStyle::GetStyleSetName(), "LevelViewport.ToggleActorPilotCameraView");
+		const FSlateIcon Icon(FAppStyle::GetAppStyleSetName(), "LevelViewport.ToggleActorPilotCameraView");
 
 		LevelEditorTabManager->RegisterTabSpawner("CameraShakePreviewer", FOnSpawnTab::CreateStatic(&FCameraShakePreviewerModule::CreateCameraShakePreviewerTab))
 			.SetDisplayName(LOCTEXT("CameraShakePreviewer", "Camera Shake Previewer"))
diff --git a/Engine/Plugins/Developer/Concert/ConcertSync/ConcertSyncClient/Source/ConcertSyncUI/Private/ConcertFrontendStyle.cpp b/Engine/Plugins/Developer/Concert/ConcertSync/ConcertSyncClient/Source/ConcertSyncUI/Private/ConcertFrontendStyle.cpp
--- a/Engine/Plugins/Developer/Concert/ConcertSync/ConcertSyncClient/Source/ConcertSyncUI/Private/ConcertFrontendStyle.cpp
+++ b/Engine/Plugins/Developer/Concert/ConcertSync/ConcertSyncClient/Source/ConcertSyncUI/Private/ConcertFrontendStyle.cpp
@@ -9,6 +9,7 @@
 
 #define IMAGE_PLUGIN_BRUSH( RelativePath, ... ) FSlateImageBrush( FConcertFrontendStyle::InContent( RelativePath, ".png" ), __VA_ARGS__ )
 #define IMAGE_BRUSH(RelativePath, ...) FSlateImageBrush(StyleSet->RootToContentDir(RelativePath, TEXT(".png")), __VA_ARGS__)
+#define IMAGE_PLUGIN_BRUSH_SVG( RelativePath, ... ) FSlateVectorImageBrush( FConcertFrontendStyle::InContent(RelativePath, ".svg"), __VA_ARGS__)
 
 FString FConcertFrontendStyle::InContent(const FString& RelativePath, const ANSICHAR* Extension)
 {
@@ -53,7 +54,7 @@ void FConcertFrontendStyle::Initialize()
 	StyleSet->Set("Concert.ModifiedByOther", new IMAGE_PLUGIN_BRUSH("Icons/icon_ConcertModifiedByOther_16x", Icon16x16));
 
 	// Multi-user Tab/Menu icons
-	StyleSet->Set("Concert.MultiUser", new IMAGE_PLUGIN_BRUSH("Icons/icon_MultiUser_32x", Icon16x16));
+	StyleSet->Set("Concert.MultiUser", new IMAGE_PLUGIN_BRUSH_SVG("Icons/icon_MultiUser", Icon16x16));
 
 	// Maps the UI Command name in Multi-User module. (UI_COMMAND does magic icon mapping when style name and command name matches)
 	StyleSet->Set("Concert.OpenBrowser",  new IMAGE_PLUGIN_BRUSH("Icons/icon_MultiUser_32x", Icon16x16));
diff --git a/Engine/Plugins/Developer/TraceDataFiltering/Source/TraceDataFiltering/Private/EventFilterStyle.cpp b/Engine/Plugins/Developer/TraceDataFiltering/Source/TraceDataFiltering/Private/EventFilterStyle.cpp
--- a/Engine/Plugins/Developer/TraceDataFiltering/Source/TraceDataFiltering/Private/EventFilterStyle.cpp
+++ b/Engine/Plugins/Developer/TraceDataFiltering/Source/TraceDataFiltering/Private/EventFilterStyle.cpp
@@ -4,16 +4,17 @@
 #include "Styling/SlateStyleRegistry.h"
 #include "Styling/SlateTypes.h"
 #include "Styling/CoreStyle.h"
+#include "Styling/StarshipCoreStyle.h"
+#include "Styling/SlateStyleMacros.h"
+
+// This is to fix the issue that SlateStyleMacros like IMAGE_BRUSH look for RootToContentDir but StyleSet->RootToContentDir is how this style is set up
+#define RootToContentDir StyleSet->RootToContentDir
 
 TSharedPtr< FSlateStyleSet > FEventFilterStyle::StyleSet = nullptr;
 
 FTextBlockStyle FEventFilterStyle::NormalText;
 
-#define IMAGE_BRUSH( RelativePath, ... ) FSlateImageBrush( StyleSet->RootToContentDir( RelativePath, TEXT(".png") ), __VA_ARGS__ )
-#define BOX_BRUSH( RelativePath, ... ) FSlateBoxBrush( StyleSet->RootToContentDir( RelativePath, TEXT(".png") ), __VA_ARGS__ )
-#define BORDER_BRUSH( RelativePath, ... ) FSlateBorderBrush( StyleSet->RootToContentDir( RelativePath, TEXT(".png") ), __VA_ARGS__ )
-#define DEFAULT_FONT(...) FCoreStyle::GetDefaultFontStyle(__VA_ARGS__)
-#define ICON_FONT(...) FSlateFontInfo(StyleSet->RootToContentDir("Fonts/FontAwesome", TEXT(".ttf")), __VA_ARGS__)
+#define ICON_FONT(...) FSlateFontInfo(RootToContentDir("Fonts/FontAwesome", TEXT(".ttf")), __VA_ARGS__)
 
 // Const icon sizes
 static const FVector2D Icon8x8(8.0f, 8.0f);
@@ -76,7 +77,7 @@ void FEventFilterStyle::Initialize()
 		StyleSet->Set("EventFilter.State.Pending", new IMAGE_BRUSH("Common/CheckBox_Undetermined", Icon16x16));
 		StyleSet->Set("EventFilter.State.Pending_Hovered", new IMAGE_BRUSH("Common/CheckBox_Undetermined_Hovered", Icon16x16));
 
-		StyleSet->Set("EventFilter.TabIcon", new IMAGE_BRUSH("/Icons/icon_Genericfinder_16x", Icon16x16));
+		StyleSet->Set("EventFilter.TabIcon", new IMAGE_BRUSH_SVG("Starship/Common/TraceDataFiltering", Icon16x16));
 	}
 
 	FButtonStyle Button = FButtonStyle()
diff --git a/Engine/Plugins/Editor/AssetManagerEditor/Source/AssetManagerEditor/Private/AssetManagerEditorModule.cpp b/Engine/Plugins/Editor/AssetManagerEditor/Source/AssetManagerEditor/Private/AssetManagerEditorModule.cpp
--- a/Engine/Plugins/Editor/AssetManagerEditor/Source/AssetManagerEditor/Private/AssetManagerEditorModule.cpp
+++ b/Engine/Plugins/Editor/AssetManagerEditor/Source/AssetManagerEditor/Private/AssetManagerEditorModule.cpp
@@ -499,7 +499,7 @@ void FAssetManagerEditorModule::StartupModule()
 			.SetDisplayName(LOCTEXT("AssetAuditTitle", "Asset Audit"))
 			.SetTooltipText(LOCTEXT("AssetAuditTooltip", "Open Asset Audit window, allows viewing detailed information about assets."))
 			.SetGroup(WorkspaceMenu::GetMenuStructure().GetDeveloperToolsAuditCategory())
-			.SetIcon(FSlateIcon(FEditorStyle::Get().GetStyleSetName(), "GenericEditor.Tabs.Properties"));
+			.SetIcon(FSlateIcon(FAppStyle::GetAppStyleSetName(), "LevelEditor.Audit"));
 
 		FGlobalTabmanager::Get()->RegisterNomadTabSpawner(ReferenceViewerTabName, FOnSpawnTab::CreateRaw(this, &FAssetManagerEditorModule::SpawnReferenceViewerTab))
 			.SetDisplayName(LOCTEXT("ReferenceViewerTitle", "Reference Viewer"))
diff --git a/Engine/Plugins/Editor/EditorDebugTools/Source/EditorDebugTools/Private/EditorDebugTools.cpp b/Engine/Plugins/Editor/EditorDebugTools/Source/EditorDebugTools/Private/EditorDebugTools.cpp
--- a/Engine/Plugins/Editor/EditorDebugTools/Source/EditorDebugTools/Private/EditorDebugTools.cpp
+++ b/Engine/Plugins/Editor/EditorDebugTools/Source/EditorDebugTools/Private/EditorDebugTools.cpp
@@ -68,7 +68,7 @@ void FEditorDebugToolsModule::StartupModule()
 	FTabSpawnerEntry& Spawner = FGlobalTabmanager::Get()->RegisterNomadTabSpawner(EditorDebugToolsTabName, FOnSpawnTab::CreateStatic(&CreateDebugToolsTab))
 		.SetDisplayName(NSLOCTEXT("Toolbox", "DebugTools", "Debug Tools"))
 		.SetTooltipText(NSLOCTEXT("Toolbox", "DebugToolsTooltipText", "Open the Debug Tools tab."))
-		.SetIcon(FSlateIcon(FEditorDebugToolsStyle::GetStyleSetName(), "DebugTools.TabIcon"))
+		.SetIcon(FSlateIcon(FAppStyle::GetAppStyleSetName(), "Debug"))
 		.SetGroup(WorkspaceMenu::GetMenuStructure().GetDeveloperToolsDebugCategory());
 
 	if (CanShowModulesTab())
@@ -76,7 +76,7 @@ void FEditorDebugToolsModule::StartupModule()
 		FGlobalTabmanager::Get()->RegisterNomadTabSpawner(ModulesTabName, FOnSpawnTab::CreateStatic(&CreateModulesTab))
 			.SetDisplayName(NSLOCTEXT("Toolbox", "Modules", "Modules"))
 			.SetTooltipText(NSLOCTEXT("Toolbox", "ModulesTooltipText", "Open the Modules tab."))
-			.SetIcon(FSlateIcon(FEditorDebugToolsStyle::GetStyleSetName(), "Modules.TabIcon"))
+			.SetIcon(FSlateIcon(FAppStyle::GetAppStyleSetName(), "Modules"))
 			.SetGroup(WorkspaceMenu::GetMenuStructure().GetDeveloperToolsDebugCategory());
 	}
 }
diff --git a/Engine/Plugins/Editor/EditorDebugTools/Source/EditorDebugTools/Private/EditorDebugToolsStyle.cpp b/Engine/Plugins/Editor/EditorDebugTools/Source/EditorDebugTools/Private/EditorDebugToolsStyle.cpp
--- a/Engine/Plugins/Editor/EditorDebugTools/Source/EditorDebugTools/Private/EditorDebugToolsStyle.cpp
+++ b/Engine/Plugins/Editor/EditorDebugTools/Source/EditorDebugTools/Private/EditorDebugToolsStyle.cpp
@@ -46,12 +46,6 @@ TSharedRef< FSlateStyleSet > FEditorDebugToolsStyle::Create()
 	// Gamma reference.
 	Style->Set("GammaReference", new IMAGE_PLUGIN_BRUSH("GammaReference", FVector2D(256, 128)));
 
-	// Debug Tools Window
-	Style->Set("DebugTools.TabIcon", new LOCAL_IMAGE_BRUSH("Icons/icon_tab_DebugTools_16x", Icon16x16));
-	
-	// Modules Window
-	Style->Set("Modules.TabIcon", new IMAGE_PLUGIN_BRUSH("icon_tab_Modules_16px", Icon16x16));
-
 	return Style;
 }
 
diff --git a/Engine/Plugins/Editor/MaterialAnalyzer/Source/MaterialAnalyzerModule.cpp b/Engine/Plugins/Editor/MaterialAnalyzer/Source/MaterialAnalyzerModule.cpp
--- a/Engine/Plugins/Editor/MaterialAnalyzer/Source/MaterialAnalyzerModule.cpp
+++ b/Engine/Plugins/Editor/MaterialAnalyzer/Source/MaterialAnalyzerModule.cpp
@@ -46,7 +46,7 @@ void FMaterialAnalyzerModule::StartupModule()
 		.SetGroup(WorkspaceMenu::GetMenuStructure().GetDeveloperToolsAuditCategory())
 		.SetDisplayName(LOCTEXT("TabTitle", "Material Analyzer"))
 		.SetTooltipText(LOCTEXT("TooltipText", "Opens Material Analyzer tool."))
-		.SetIcon(FSlateIcon(FEditorStyle::GetStyleSetName(), "MaterialEditor.ToggleMaterialStats.Tab"));
+		.SetIcon(FSlateIcon(FAppStyle::GetAppStyleSetName(), "MaterialEditor.ToggleMaterialStats.Tab"));
 }
 
 TSharedRef<SDockTab> FMaterialAnalyzerModule::SpawnMaterialAnalyzerTab(const FSpawnTabArgs& SpawnTabArgs)
diff --git a/Engine/Plugins/Editor/StylusInput/Source/StylusInput/Private/StylusInputModule.cpp b/Engine/Plugins/Editor/StylusInput/Source/StylusInput/Private/StylusInputModule.cpp
--- a/Engine/Plugins/Editor/StylusInput/Source/StylusInput/Private/StylusInputModule.cpp
+++ b/Engine/Plugins/Editor/StylusInput/Source/StylusInput/Private/StylusInputModule.cpp
@@ -55,6 +55,7 @@ void UStylusInputSubsystem::Initialize(FSubsystemCollectionBase& Collection)
 		FOnSpawnTab::CreateUObject(this, &UStylusInputSubsystem::OnSpawnPluginTab))
 		.SetDisplayName(LOCTEXT("DebugTabTitle", "Stylus Input Debug"))
 		.SetTooltipText(LOCTEXT("DebugTabTooltip", "Debug panel to display current values of stylus inputs."))
+		.SetIcon(FSlateIcon(FAppStyle::GetAppStyleSetName(), "StylusInputDebug.TabIcon"))
 		.SetGroup(MenuStructure.GetDeveloperToolsDebugCategory());
 }
 
diff --git a/Engine/Plugins/Experimental/ControlRig/Source/ControlRigEditor/Private/ControlRigEditorModule.cpp b/Engine/Plugins/Experimental/ControlRig/Source/ControlRigEditor/Private/ControlRigEditorModule.cpp
--- a/Engine/Plugins/Experimental/ControlRig/Source/ControlRigEditor/Private/ControlRigEditorModule.cpp
+++ b/Engine/Plugins/Experimental/ControlRig/Source/ControlRigEditor/Private/ControlRigEditorModule.cpp
@@ -279,7 +279,7 @@ void FControlRigEditorModule::StartupModule()
 			.SetDisplayName(NSLOCTEXT("UnrealEditor", "HierarchicalProfilerTab", "Hierarchical Profiler"))
 			.SetTooltipText(NSLOCTEXT("UnrealEditor", "HierarchicalProfilerTooltip", "Open the Hierarchical Profiler tab."))
 			.SetGroup(WorkspaceMenu::GetMenuStructure().GetDeveloperToolsProfilingCategory())
-			.SetIcon(FSlateIcon(TEXT("ControlRigEditorStyle"), TEXT("ControlRig.RigUnit")));
+			.SetIcon(FSlateIcon(TEXT("ControlRigEditorStyle"), TEXT("HierarchicalProfiler.TabIcon")));
 
 		FGlobalTabmanager::Get()->RegisterNomadTabSpawner(IControlRigEditorModule::ControlRigPoseTab, FOnSpawnTab::CreateStatic(&SpawnPoseTab))
 			.SetDisplayName(NSLOCTEXT("UnrealEditor", "ControlRigPoseTab", "Control Rig Pose"))
diff --git a/Engine/Plugins/Experimental/ControlRig/Source/ControlRigEditor/Private/ControlRigEditorStyle.h b/Engine/Plugins/Experimental/ControlRig/Source/ControlRigEditor/Private/ControlRigEditorStyle.h
--- a/Engine/Plugins/Experimental/ControlRig/Source/ControlRigEditor/Private/ControlRigEditorStyle.h
+++ b/Engine/Plugins/Experimental/ControlRig/Source/ControlRigEditor/Private/ControlRigEditorStyle.h
@@ -54,6 +54,8 @@ public:
 			Set("ControlRig.TabIcon", new IMAGE_BRUSH("Slate/ControlRigTab_16x", Icon16x16));
 			Set("ControlRig.RigUnit", new IMAGE_BRUSH("Slate/ControlRigUnit_16x", Icon16x16));
 
+			Set("HierarchicalProfiler.TabIcon", new IMAGE_BRUSH_SVG("Slate/HierarchicalProfiler", Icon16x16));
+
 			// icons for control units
 			Set("ControlRig.ControlUnitOn", new IMAGE_BRUSH("Slate/ControlUnit_On", Icon32x32));
 			Set("ControlRig.ControlUnitOff", new IMAGE_BRUSH("Slate/ControlUnit_Off", Icon32x32));
@@ -108,6 +110,7 @@ public:
 			Set("ControlRig.Hierarchy.Menu", TTF_FONT("Fonts/Roboto-Regular", 12));
 		}
 
+
 		FSlateStyleRegistry::RegisterSlateStyle(*this);
 	}
 
diff --git a/Engine/Plugins/Experimental/RemoteSession/Source/RemoteSessionEditor/Private/RemoteSessionEditorStyle.cpp b/Engine/Plugins/Experimental/RemoteSession/Source/RemoteSessionEditor/Private/RemoteSessionEditorStyle.cpp
--- a/Engine/Plugins/Experimental/RemoteSession/Source/RemoteSessionEditor/Private/RemoteSessionEditorStyle.cpp
+++ b/Engine/Plugins/Experimental/RemoteSession/Source/RemoteSessionEditor/Private/RemoteSessionEditorStyle.cpp
@@ -4,6 +4,7 @@
 
 #include "Brushes/SlateImageBrush.h"
 #include "Styling/SlateStyle.h"
+#include "Styling/SlateStyleMacros.h"
 #include "Styling/SlateStyleRegistry.h"
 
 namespace RemoteSessionStyle
@@ -18,14 +19,17 @@ namespace RemoteSessionStyle
 	static TUniquePtr<FSlateStyleSet> StyleInstance;
 }
 
-#define IMAGE_BRUSH(RelativePath, ...) FSlateImageBrush(RemoteSessionStyle::StyleInstance->RootToContentDir(RelativePath, TEXT(".png")), __VA_ARGS__)
+// This is to fix the issue that SlateStyleMacros like IMAGE_BRUSH look for RootToContentDir but StyleSet->RootToContentDir is how this style is set up
+#define RootToContentDir RemoteSessionStyle::StyleInstance->RootToContentDir
 
 void FRemoteSessionEditorStyle::Register()
 {
 	RemoteSessionStyle::StyleInstance = MakeUnique<FSlateStyleSet>(RemoteSessionStyle::NAME_StyleName);
+	RemoteSessionStyle::StyleInstance->SetCoreContentRoot(FPaths::EngineContentDir() / TEXT("Slate"));
 	RemoteSessionStyle::StyleInstance->SetContentRoot(FPaths::EnginePluginsDir() / TEXT("Experimental/RemoteSession/Content/Editor/Icons/"));
 
 
+	RemoteSessionStyle::StyleInstance->Set("RemoteSessionStream", new IMAGE_BRUSH_SVG("RemoteSessionStream", RemoteSessionStyle::Icon16x16));
 	RemoteSessionStyle::StyleInstance->Set("TabIcons.RemoteSession.Small", new IMAGE_BRUSH("RemoteSession_Stream_16x", RemoteSessionStyle::Icon16x16));
 	RemoteSessionStyle::StyleInstance->Set("RemoteSessionStream.Stream", new IMAGE_BRUSH("RemoteSession_Stream_40x", RemoteSessionStyle::Icon40x40));
 	RemoteSessionStyle::StyleInstance->Set("RemoteSessionStream.Stop", new IMAGE_BRUSH("RemoteSession_Stop_40x", RemoteSessionStyle::Icon40x40));
diff --git a/Engine/Plugins/Experimental/RemoteSession/Source/RemoteSessionEditor/Private/Widgets/SRemoteSessionStream.cpp b/Engine/Plugins/Experimental/RemoteSession/Source/RemoteSessionEditor/Private/Widgets/SRemoteSessionStream.cpp
--- a/Engine/Plugins/Experimental/RemoteSession/Source/RemoteSessionEditor/Private/Widgets/SRemoteSessionStream.cpp
+++ b/Engine/Plugins/Experimental/RemoteSession/Source/RemoteSessionEditor/Private/Widgets/SRemoteSessionStream.cpp
@@ -80,7 +80,7 @@ void SRemoteSessionStream::RegisterNomadTabSpawner()
 			.SetDisplayName(LOCTEXT("TabTitle", "Remote Session Stream"))
 			.SetTooltipText(LOCTEXT("TooltipText", "Stream a particular UMG to a Remote Session app."))
 			.SetGroup(WorkspaceMenu::GetMenuStructure().GetToolsCategory())
-			.SetIcon(FSlateIcon(FRemoteSessionEditorStyle::GetStyleSetName(), "TabIcons.RemoteSession.Small"));
+			.SetIcon(FSlateIcon(FRemoteSessionEditorStyle::GetStyleSetName(), "RemoteSessionStream"));
 	};
 
 	FLevelEditorModule& LevelEditorModule = FModuleManager::LoadModuleChecked<FLevelEditorModule>(TEXT("LevelEditor"));
diff --git a/Engine/Plugins/FX/Niagara/Source/NiagaraEditor/Private/NiagaraEditorStyle.cpp b/Engine/Plugins/FX/Niagara/Source/NiagaraEditor/Private/NiagaraEditorStyle.cpp
--- a/Engine/Plugins/FX/Niagara/Source/NiagaraEditor/Private/NiagaraEditorStyle.cpp
+++ b/Engine/Plugins/FX/Niagara/Source/NiagaraEditor/Private/NiagaraEditorStyle.cpp
@@ -51,6 +51,7 @@ NIAGARAEDITOR_API FString RelativePathToPluginPath(const FString& RelativePath,
 #define BOX_PLUGIN_BRUSH( RelativePath, ... ) FSlateBoxBrush( RelativePathToPluginPath( RelativePath, ".png"), __VA_ARGS__ )
 #define BORDER_BRUSH( RelativePath, ... ) FSlateBorderBrush( Style->RootToContentDir( RelativePath, TEXT(".png") ), __VA_ARGS__ )
 #define DEFAULT_FONT(...) FCoreStyle::GetDefaultFontStyle(__VA_ARGS__)
+#define IMAGE_PLUGIN_BRUSH_SVG( RelativePath, ... ) FSlateVectorImageBrush( RelativePathToPluginPath( RelativePath, ".svg" ), __VA_ARGS__ )
 
 const FVector2D Icon8x8(8.0f, 8.0f);
 const FVector2D Icon12x12(12.0f, 12.0f);
@@ -354,7 +355,7 @@ TSharedRef< FSlateStyleSet > FNiagaraEditorStyle::Create()
 	{
 		const FVector2D Icon24x24(24.0f, 24.0f);
 
-		Style->Set("NiagaraEditor.Debugger.TabIcon", new IMAGE_CORE_BRUSH("Icons/icon_tab_DebugTools_16x", Icon16x16));
+		Style->Set("NiagaraEditor.Debugger.TabIcon", new IMAGE_PLUGIN_BRUSH_SVG("Icons/NiagaraDebugger", Icon16x16));
 		Style->Set("NiagaraEditor.Debugger.PlayIcon", new IMAGE_PLUGIN_BRUSH("Icons/Debugger/Play", Icon24x24));
 		Style->Set("NiagaraEditor.Debugger.SpeedIcon", new IMAGE_PLUGIN_BRUSH("Icons/Debugger/Speed", Icon24x24));
 		Style->Set("NiagaraEditor.Debugger.PauseIcon", new IMAGE_PLUGIN_BRUSH("Icons/Debugger/Pause", Icon24x24));
diff --git a/Engine/Plugins/MovieScene/MatineeToLevelSequence/Source/MatineeToLevelSequence/Private/MatineeToLevelSequenceModule.cpp b/Engine/Plugins/MovieScene/MatineeToLevelSequence/Source/MatineeToLevelSequence/Private/MatineeToLevelSequenceModule.cpp
--- a/Engine/Plugins/MovieScene/MatineeToLevelSequence/Source/MatineeToLevelSequence/Private/MatineeToLevelSequenceModule.cpp
+++ b/Engine/Plugins/MovieScene/MatineeToLevelSequence/Source/MatineeToLevelSequence/Private/MatineeToLevelSequenceModule.cpp
@@ -104,7 +104,7 @@ protected:
 				TSharedPtr<FTabManager> LevelEditorTabManager = LevelEditorModule.GetLevelEditorTabManager();
 
 				const IWorkspaceMenuStructure& MenuStructure = WorkspaceMenu::GetMenuStructure();
-				const FSlateIcon Icon(FSlateIcon(FEditorStyle::GetStyleSetName(), "LevelEditor.EditMatinee", "LevelEditor.EditMatinee.Small"));
+				const FSlateIcon Icon(FAppStyle::GetAppStyleSetName(), "LevelViewport.ToggleActorPilotCameraView");
 
 				LevelEditorTabManager->RegisterTabSpawner("MatineeToLevelSequence", FOnSpawnTab::CreateRaw(this, &FMatineeToLevelSequenceModule::CreateCameraShakeConverterTab))
 				.SetDisplayName(LOCTEXT("MatineeCameraShakeConverter", "Matinee Camera Shake Converter"))
diff --git a/Engine/Plugins/VirtualProduction/Takes/Source/TakeRecorder/Private/TakeRecorderModule.cpp b/Engine/Plugins/VirtualProduction/Takes/Source/TakeRecorder/Private/TakeRecorderModule.cpp
--- a/Engine/Plugins/VirtualProduction/Takes/Source/TakeRecorder/Private/TakeRecorderModule.cpp
+++ b/Engine/Plugins/VirtualProduction/Takes/Source/TakeRecorder/Private/TakeRecorderModule.cpp
@@ -146,7 +146,7 @@ namespace
 			.SetGroup(WorkspaceMenu::GetMenuStructure().GetLevelEditorCinematicsCategory())
 			.SetDisplayName(ITakeRecorderModule::TakeRecorderTabLabel)
 			.SetTooltipText(LOCTEXT("TakeRecorderTab_Tooltip", "Open the main Take Recorder UI."))
-			.SetIcon(FSlateIcon(FTakeRecorderStyle::StyleName, "TakeRecorder.TabIcon"));
+			.SetIcon(FSlateIcon(FAppStyle::GetAppStyleSetName(), "SequenceRecorder.TabIcon"));
 
 		FTabSpawnerEntry& TBTabSpawner = FGlobalTabmanager::Get()->RegisterNomadTabSpawner(ITakeRecorderModule::TakesBrowserTabName, FOnSpawnTab::CreateStatic(SpawnTakesBrowserTab));
 
diff --git a/Engine/Plugins/VirtualProduction/Takes/Source/TakeRecorder/Private/Widgets/STakeRecorderTabContent.cpp b/Engine/Plugins/VirtualProduction/Takes/Source/TakeRecorder/Private/Widgets/STakeRecorderTabContent.cpp
--- a/Engine/Plugins/VirtualProduction/Takes/Source/TakeRecorder/Private/Widgets/STakeRecorderTabContent.cpp
+++ b/Engine/Plugins/VirtualProduction/Takes/Source/TakeRecorder/Private/Widgets/STakeRecorderTabContent.cpp
@@ -71,7 +71,7 @@ void STakeRecorderTabContent::SetupForRecording(UTakePreset* BasePreset)
 
 	CurrentMode    = ETakeRecorderPanelMode::NewRecording;
 	TitleAttribute = ITakeRecorderModule::TakeRecorderTabLabel;
-	IconAttribute  = FTakeRecorderStyle::Get().GetBrush("TakeRecorder.TabIcon");
+	IconAttribute  = FAppStyle::Get().GetBrush("SequenceRecorder.TabIcon");
 
 	// Null out the tab content to ensure that all references have been cleaned up before constructing the new one
 	ChildSlot [ SNullWidget::NullWidget ];
diff --git a/Engine/Source/Developer/LogVisualizer/Private/LogVisualizerStyle.cpp b/Engine/Source/Developer/LogVisualizer/Private/LogVisualizerStyle.cpp
--- a/Engine/Source/Developer/LogVisualizer/Private/LogVisualizerStyle.cpp
+++ b/Engine/Source/Developer/LogVisualizer/Private/LogVisualizerStyle.cpp
@@ -4,12 +4,12 @@
 #include "Styling/SlateStyleRegistry.h"
 #include "Styling/SlateTypes.h"
 #include "Styling/CoreStyle.h"
+#include "Styling/StarshipCoreStyle.h"
 #include "EditorStyleSet.h"
+#include "Styling/SlateStyleMacros.h"
 
-#define IMAGE_BRUSH( RelativePath, ... ) FSlateImageBrush( FPaths::EngineContentDir() / "Editor/Slate"/ RelativePath + TEXT(".png"), __VA_ARGS__ )
-#define BOX_BRUSH( RelativePath, ... ) FSlateBoxBrush( FPaths::EngineContentDir() / "Editor/Slate"/ RelativePath + TEXT(".png"), __VA_ARGS__ )
-#define BORDER_BRUSH( RelativePath, ... ) FSlateBorderBrush( FPaths::EngineContentDir() / "Editor/Slate"/ RelativePath + TEXT(".png"), __VA_ARGS__ )
-#define DEFAULT_FONT(...) FCoreStyle::GetDefaultFontStyle(__VA_ARGS__)
+// This is to fix the issue that SlateStyleMacros like IMAGE_BRUSH look for RootToContentDir but StyleSet->RootToContentDir is how this style is set up
+#define RootToContentDir Style.RootToContentDir
 
 TSharedPtr< FSlateStyleSet > FLogVisualizerStyle::StyleInstance = nullptr;
 
@@ -63,7 +63,7 @@ TSharedRef< FSlateStyleSet > FLogVisualizerStyle::Create()
 	Style.Set("TextLogs.Text", LogsText);
 
 	{
-		Style.Set("LogVisualizerApp.TabIcon", new IMAGE_BRUSH("Icons/icon_tab_DebugTools_16x", Icon16x16));
+		Style.Set("LogVisualizerApp.TabIcon", new IMAGE_BRUSH_SVG("Starship/Common/VisualLogger", Icon16x16));
 		
 		Style.Set("LogVisualizer.LogBar.Background", new BOX_BRUSH("Common/ProgressBar_Background", FMargin(5.f / 12.f)));
 		Style.Set("LogVisualizer.LogBar.Selected", new BOX_BRUSH("Common/TaskGraph_Selected", FMargin(5.f / 12.f)));
diff --git a/Engine/Source/Developer/SlateReflector/Private/SlateReflectorModule.cpp b/Engine/Source/Developer/SlateReflector/Private/SlateReflectorModule.cpp
--- a/Engine/Source/Developer/SlateReflector/Private/SlateReflectorModule.cpp
+++ b/Engine/Source/Developer/SlateReflector/Private/SlateReflectorModule.cpp
@@ -135,7 +135,7 @@ public:
 			FTabSpawnerEntry& SpawnerEntry = FGlobalTabmanager::Get()->RegisterNomadTabSpawner("WidgetReflector", FOnSpawnTab::CreateRaw(this, &FSlateReflectorModuleImpl::MakeWidgetReflectorTab) )
 				.SetDisplayName(LOCTEXT("WidgetReflectorTitle", "Widget Reflector"))
 				.SetTooltipText(LOCTEXT("WidgetReflectorTooltipText", "Open the Widget Reflector tab."))
-				.SetIcon(FSlateIcon(FCoreStyle::Get().GetStyleSetName(), "WidgetReflector.TabIcon"));
+				.SetIcon(FSlateIcon(FWidgetReflectorStyle::GetStyleSetName(), "WidgetReflector.TabIcon"));
 
 			if (WorkspaceGroup.IsValid())
 			{
diff --git a/Engine/Source/Developer/SlateReflector/Private/Styling/WidgetReflectorStyle.cpp b/Engine/Source/Developer/SlateReflector/Private/Styling/WidgetReflectorStyle.cpp
--- a/Engine/Source/Developer/SlateReflector/Private/Styling/WidgetReflectorStyle.cpp
+++ b/Engine/Source/Developer/SlateReflector/Private/Styling/WidgetReflectorStyle.cpp
@@ -4,6 +4,10 @@
 #include "Styling/SlateStyleRegistry.h"
 #include "Styling/SlateTypes.h"
 #include "Styling/CoreStyle.h"
+#include "Styling/SlateStyleMacros.h"
+
+// This is to fix the issue that SlateStyleMacros like IMAGE_BRUSH look for RootToContentDir but StyleSet->RootToContentDir is how this style is set up
+#define RootToContentDir StyleSet->RootToContentDir
 
 TSharedPtr< FSlateStyleSet > FWidgetReflectorStyle::StyleInstance = nullptr;
 
@@ -32,6 +36,7 @@ FName FWidgetReflectorStyle::GetStyleSetName()
 TSharedRef< FSlateStyleSet > FWidgetReflectorStyle::Create()
 {
 	const FVector2D Icon8x8(8.0f, 8.0f);
+	const FVector2D Icon16x16(16.0f, 16.0f);
 	const FVector2D Icon24x24(24.0f, 24.0f);
 
 	TSharedRef<FSlateStyleSet> StyleSet = MakeShareable(new FSlateStyleSet(FWidgetReflectorStyle::GetStyleSetName()));
@@ -41,9 +46,9 @@ TSharedRef< FSlateStyleSet > FWidgetReflectorStyle::Create()
 #if WITH_EDITOR
 	{
 		FButtonStyle Button = FButtonStyle()
-			.SetNormal(FSlateBoxBrush(StyleSet->RootToContentDir("Common/ButtonHoverHint.png"), FMargin(4 / 16.0f), FLinearColor(1, 1, 1, 0.15f)))
-			.SetHovered(FSlateBoxBrush(StyleSet->RootToContentDir("Common/ButtonHoverHint.png"), FMargin(4 / 16.0f), FLinearColor(1, 1, 1, 0.25f)))
-			.SetPressed(FSlateBoxBrush(StyleSet->RootToContentDir("Common/ButtonHoverHint.png"), FMargin(4 / 16.0f), FLinearColor(1, 1, 1, 0.30f)))
+			.SetNormal(FSlateBoxBrush(RootToContentDir("Common/ButtonHoverHint.png"), FMargin(4 / 16.0f), FLinearColor(1, 1, 1, 0.15f)))
+			.SetHovered(FSlateBoxBrush(RootToContentDir("Common/ButtonHoverHint.png"), FMargin(4 / 16.0f), FLinearColor(1, 1, 1, 0.25f)))
+			.SetPressed(FSlateBoxBrush(RootToContentDir("Common/ButtonHoverHint.png"), FMargin(4 / 16.0f), FLinearColor(1, 1, 1, 0.30f)))
 			.SetNormalPadding(FMargin(0, 0, 0, 1))
 			.SetPressedPadding(FMargin(0, 1, 0, 0));
 		StyleSet->Set("Button", Button);
@@ -70,14 +75,16 @@ TSharedRef< FSlateStyleSet > FWidgetReflectorStyle::Create()
 	}
 	
 	{
-		StyleSet->Set("Icon.FocusPicking", new FSlateImageBrush(StyleSet->RootToContentDir("Icons/SlateReflector/FocusPicking_24x.png"), Icon24x24));
-		StyleSet->Set("Icon.HitTestPicking", new FSlateImageBrush(StyleSet->RootToContentDir("Icons/SlateReflector/HitTestPicking_24x.png"), Icon24x24));
-		StyleSet->Set("Icon.VisualPicking", new FSlateImageBrush(StyleSet->RootToContentDir("Icons/SlateReflector/VisualPicking_24x.png"), Icon24x24));
-
-		StyleSet->Set("Symbols.LeftArrow", new FSlateImageBrush(StyleSet->RootToContentDir("Common/LeftArrow.png"), Icon24x24));
-		StyleSet->Set("Symbols.RightArrow", new FSlateImageBrush(StyleSet->RootToContentDir("Common/SubmenuArrow.png"), Icon24x24));
-		StyleSet->Set("Symbols.UpArrow", new FSlateImageBrush(StyleSet->RootToContentDir("Common/UpArrow.png"), Icon24x24));
-		StyleSet->Set("Symbols.DownArrow", new FSlateImageBrush(StyleSet->RootToContentDir("Common/DownArrow.png"), Icon24x24));
+		StyleSet->Set("Icon.FocusPicking", new FSlateImageBrush(RootToContentDir("Icons/SlateReflector/FocusPicking_24x.png"), Icon24x24));
+		StyleSet->Set("Icon.HitTestPicking", new FSlateImageBrush(RootToContentDir("Icons/SlateReflector/HitTestPicking_24x.png"), Icon24x24));
+		StyleSet->Set("Icon.VisualPicking", new FSlateImageBrush(RootToContentDir("Icons/SlateReflector/VisualPicking_24x.png"), Icon24x24));
+
+		StyleSet->Set("Symbols.LeftArrow", new FSlateImageBrush(RootToContentDir("Common/LeftArrow.png"), Icon24x24));
+		StyleSet->Set("Symbols.RightArrow", new FSlateImageBrush(RootToContentDir("Common/SubmenuArrow.png"), Icon24x24));
+		StyleSet->Set("Symbols.UpArrow", new FSlateImageBrush(RootToContentDir("Common/UpArrow.png"), Icon24x24));
+		StyleSet->Set("Symbols.DownArrow", new FSlateImageBrush(RootToContentDir("Common/DownArrow.png"), Icon24x24));
+
+		StyleSet->Set("WidgetReflector.TabIcon", new IMAGE_BRUSH_SVG("Starship/Common/Widget", Icon16x16));
 	}
 #endif
 
diff --git a/Engine/Source/Developer/TaskGraph/Private/STaskGraph.cpp b/Engine/Source/Developer/TaskGraph/Private/STaskGraph.cpp
--- a/Engine/Source/Developer/TaskGraph/Private/STaskGraph.cpp
+++ b/Engine/Source/Developer/TaskGraph/Private/STaskGraph.cpp
@@ -173,7 +173,7 @@ public:
 		FGlobalTabmanager::Get()->RegisterNomadTabSpawner(TaskGraphTabId, FOnSpawnTab::CreateRaw(this, &FProfileVisualizerModule::SpawnProfileVizualizerTab))
 			.SetDisplayName(NSLOCTEXT("ProfileVisualizerModule", "TabTitle", "Profile Data Visualizer"))
 			.SetTooltipText(NSLOCTEXT("ProfileVisualizerModule", "TooltipText", "Open the Profile Data Visualizer tab."))
-			.SetGroup(WorkspaceMenu::GetMenuStructure().GetDeveloperToolsDebugCategory())
+			.SetGroup(WorkspaceMenu::GetMenuStructure().GetDeveloperToolsProfilingCategory())
 			.SetIcon(FSlateIcon(FEditorStyle::Get().GetStyleSetName(), "Profiler.Tab"));
 #endif
 			
diff --git a/Engine/Source/Editor/EditorStyle/Private/StarshipStyle.cpp b/Engine/Source/Editor/EditorStyle/Private/StarshipStyle.cpp
--- a/Engine/Source/Editor/EditorStyle/Private/StarshipStyle.cpp
+++ b/Engine/Source/Editor/EditorStyle/Private/StarshipStyle.cpp
@@ -919,7 +919,7 @@ void FStarshipEditorStyle::FStyle::SetupGeneralStyles()
 		Set( "MessageLog.Tutorial", new IMAGE_BRUSH( "Icons/icon_Blueprint_Enum_16x", Icon16x16 ) );
 		Set( "MessageLog.Url", new IMAGE_BRUSH( "Icons/icon_world_16x", Icon16x16 ) );
 
-		Set( "MessageLog.TabIcon", new IMAGE_BRUSH( "Icons/icon_tab_MessageLog_16x", Icon16x16 ) );
+		Set( "MessageLog.TabIcon", new IMAGE_BRUSH_SVG( "Starship/Common/MessageLog", Icon16x16 ) );
 		Set( "MessageLog.ListBorder", new BOX_BRUSH( "/Docking/AppTabContentArea", FMargin(4/16.0f) ) );
 	}
 
@@ -973,7 +973,7 @@ void FStarshipEditorStyle::FStyle::SetupGeneralStyles()
 			.SetColorAndOpacity( LogColor_Error )
 			);
 
-		Set("Log.TabIcon", new IMAGE_BRUSH( "Icons/icon_tab_OutputLog_16x", Icon16x16 ) );
+		Set("Log.TabIcon", new IMAGE_BRUSH_SVG( "Starship/Common/OutputLog", Icon16x16 ) );
 
 		Set("Log.TextBox", FEditableTextBoxStyle(NormalEditableTextBoxStyle)
 			.SetBackgroundImageNormal( BOX_BRUSH( "Common/WhiteGroupBorder", FMargin(4.0f/16.0f) ) )
@@ -1002,27 +1002,18 @@ void FStarshipEditorStyle::FStyle::SetupGeneralStyles()
 		Set("OutputLog.OpenInExternalEditor", new IMAGE_BRUSH("Icons/icon_Asset_Open_In_External_Editor_16x", Icon16x16));
 
 	}
-	// Performance Analysis Tools Window
+	// Debugging tools 
 	{
 		Set("PerfTools.TabIcon", new IMAGE_BRUSH( "Icons/icon_tab_PerfTools_16x", Icon16x16 ) );
+		Set("ClassViewer.TabIcon", new IMAGE_BRUSH_SVG( "Starship/Common/Class", Icon16x16 ) );
+		Set("StructViewer.TabIcon", new IMAGE_BRUSH_SVG( "Starship/Common/UserDefinedStruct", Icon16x16 ) );
+		Set("BlueprintDebugger.TabIcon", new IMAGE_BRUSH_SVG( "Starship/Common/BlueprintDebugger", Icon16x16 ) );
+		Set("CollisionAnalyzer.TabIcon", new IMAGE_BRUSH_SVG("Starship/Common/Collision", Icon16x16));
+		Set("ObjectBrowser.TabIcon", new IMAGE_BRUSH_SVG( "Starship/Common/ObjectsBrowser", Icon16x16 ) );
+		Set("PixelInspector.TabIcon", new IMAGE_BRUSH_SVG( "Starship/Common/PixelInspector", Icon16x16 ) );
+		Set("StylusInputDebug.TabIcon", new IMAGE_BRUSH_SVG( "Starship/Common/StylusInputBug", Icon16x16 ) );
 	}
 
-	// Class Viewer Window
-	{
-		Set("ClassViewer.TabIcon", new IMAGE_BRUSH( "Icons/icon_tab_ClassViewer_16x", Icon16x16 ) );
-	}
-
-	// Blueprint Debugger Window
-	{
-		Set("BlueprintDebugger.TabIcon", new IMAGE_BRUSH_SVG( "Starship/Common/Debug", Icon16x16 ) );
-	}
-
-	// Collision Analyzer Window
-	{
-		Set("CollisionAnalyzer.TabIcon", new IMAGE_BRUSH("Icons/icon_ShowCollision_16x", Icon16x16));
-	}
-
-	// Developer Tools Menu
 	{
 		Set("DeveloperTools.MenuIcon", new IMAGE_BRUSH_SVG( "Starship/Common/DeveloperTools", Icon16x16 ) );
 	}
@@ -1060,7 +1051,7 @@ void FStarshipEditorStyle::FStyle::SetupGeneralStyles()
 
 	// Session Frontend Window
 	{
-		Set("SessionFrontEnd.TabIcon", new IMAGE_BRUSH( "Icons/icon_tab_SessionFrontend_16x", Icon16x16 ) );
+		Set("SessionFrontEnd.TabIcon", new IMAGE_BRUSH_SVG( "Starship/Common/SessionFrontend", Icon16x16 ) );
 		Set("SessionFrontEnd.Tabs.Tools", new IMAGE_BRUSH( "/Icons/icon_tab_Tools_16x", Icon16x16 ) );
 	}
 
@@ -1308,6 +1299,8 @@ void FStarshipEditorStyle::FStyle::SetupGeneralStyles()
 		Set( "MainFrame.RecentProjects",       new IMAGE_BRUSH_SVG( "Starship/Common/ProjectsRecent", Icon16x16 ) );
 		Set( "MainFrame.RecentLevels",         new IMAGE_BRUSH_SVG( "Starship/Common/LevelRecent", Icon16x16 ) );
 		Set( "MainFrame.FavoriteLevels",       new IMAGE_BRUSH_SVG( "Starship/Common/LevelFavorite", Icon16x16 ) );
+		Set( "MainFrame.ZipUpProject", 		   new IMAGE_BRUSH_SVG( "Starship/Common/ZipProject", Icon16x16 ) );
+
 
 		Set( "MainFrame.ChooseFilesToSave",       new IMAGE_BRUSH_SVG( "Starship/Common/SaveChoose", Icon16x16 ) );
 		Set( "MainFrame.ConnectToSourceControl",  new IMAGE_BRUSH_SVG( "Starship/Common/SourceControl", Icon16x16 ) );
@@ -2281,7 +2274,7 @@ void FStarshipEditorStyle::FStyle::SetupSequencerStyles()
 	// Sequence recorder standalone UI
 	if (IncludeEditorSpecificStyles())
 	{
-		Set("SequenceRecorder.TabIcon", new IMAGE_BRUSH("SequenceRecorder/icon_tab_SequenceRecorder_16x", Icon16x16));
+		Set("SequenceRecorder.TabIcon", new IMAGE_BRUSH_SVG("Starship/Sequencer/SequenceRecorder", Icon16x16));
 		Set("SequenceRecorder.Common.RecordAll.Small", new IMAGE_BRUSH("SequenceRecorder/icon_RecordAll_40x", Icon20x20));
 		Set("SequenceRecorder.Common.RecordAll", new IMAGE_BRUSH("SequenceRecorder/icon_RecordAll_40x", Icon40x40));
 		Set("SequenceRecorder.Common.StopAll.Small", new IMAGE_BRUSH("SequenceRecorder/icon_StopAll_40x", Icon20x20));
@@ -2586,6 +2579,9 @@ void FStarshipEditorStyle::FStyle::SetupGeneralIcons()
 	Set("AssetEditor.CompileStatus.Overlay.Warning", new IMAGE_BRUSH_SVG("Starship/Blueprints/CompileStatus_Warning_Badge", Icon20x20, FStyleColors::Warning));
 	Set("AssetEditor.CompileStatus.Overlay.Good", new IMAGE_BRUSH_SVG("Starship/Blueprints/CompileStatus_Good_Badge", Icon20x20, FStyleColors::AccentGreen));
 	Set("AssetEditor.CompileStatus.Overlay.Error", new IMAGE_BRUSH_SVG("Starship/Blueprints/CompileStatus_Fail_Badge", Icon20x20, FStyleColors::Error));
+
+	Set("Debug", new IMAGE_BRUSH_SVG( "Starship/Common/Debug", Icon16x16 ) );
+	Set("Modules", new IMAGE_BRUSH_SVG( "Starship/Common/Modules", Icon16x16 ) );
 }
 
 void FStarshipEditorStyle::FStyle::SetupWindowStyles()
@@ -3280,7 +3276,7 @@ void FStarshipEditorStyle::FStyle::SetupProfilerStyle()
 		Set( "Profiler.Group.16", new BOX_BRUSH( "Icons/Profiler/GroupBorder-16Gray", FMargin(4.0f/16.0f) ) );
 
 		// Profiler toolbar icons
-		Set( "Profiler.Tab", new IMAGE_BRUSH( "Icons/Profiler/Profiler_Tab_16x", Icon16x16 ) );
+		Set( "Profiler.Tab", new IMAGE_BRUSH_SVG( "Starship/Common/Visualizer", Icon16x16 ) );
 		Set( "Profiler.Tab.GraphView", new IMAGE_BRUSH( "Icons/Profiler/Profiler_Graph_View_Tab_16x", Icon16x16 ) );
 		Set( "Profiler.Tab.EventGraph", new IMAGE_BRUSH( "Icons/Profiler/profiler_OpenEventGraph_32x", Icon16x16 ) );
 		Set( "Profiler.Tab.FiltersAndPresets", new IMAGE_BRUSH( "Icons/Profiler/Profiler_Filter_Presets_Tab_16x", Icon16x16 ) );
@@ -4415,6 +4411,10 @@ void FStarshipEditorStyle::FStyle::SetupLevelEditorStyle()
 			Set( "LevelEditor.Tabs.Viewports",              new IMAGE_BRUSH_SVG("Starship/Common/Viewports", Icon16x16) );
 			Set( "LevelEditor.Tabs.HLOD",                   new IMAGE_BRUSH_SVG("Starship/Common/HierarchicalLODOutliner", Icon16x16) );
 			Set( "LevelEditor.Tabs.Debug",                  new IMAGE_BRUSH_SVG("Starship/Common/Bug", Icon16x16) );
+
+			Set( "LevelEditor.Audit",                  		new IMAGE_BRUSH_SVG("Starship/Common/AssetAudit", Icon16x16) );
+			Set( "LevelEditor.Profile",                  	new IMAGE_BRUSH_SVG("Starship/Common/Profile", Icon16x16) );
+			Set( "LevelEditor.Platforms",                  	new IMAGE_BRUSH_SVG("Starship/Common/DeviceManager", Icon16x16) );
 		}
 #endif // WITH_EDITOR || (IS_PROGRAM && WITH_UNREAL_DEVELOPER_TOOLS)
 
@@ -4525,7 +4525,7 @@ void FStarshipEditorStyle::FStyle::SetupLevelEditorStyle()
 		Set( "MergeActors.MeshMergingTool", new IMAGE_BRUSH_SVG( "Starship/MergeActors/MeshMerging_16", Icon16x16 ) );
 		Set( "MergeActors.MeshProxyTool", new IMAGE_BRUSH_SVG( "Starship/MergeActors/MeshProxy_16", Icon16x16 ) );
 		Set( "MergeActors.MeshInstancingTool", new IMAGE_BRUSH_SVG("Starship/AssetIcons/StaticMeshActor_16", Icon16x16 ) );
-		Set( "MergeActors.TabIcon", new IMAGE_BRUSH_SVG("Starship/MergeActors/MeshMerging_16", Icon16x16));
+		Set( "MergeActors.TabIcon", new IMAGE_BRUSH_SVG("Starship/Common/Merge", Icon16x16));
 		
 		Set( "PlacementBrowser.OptionsMenu", new IMAGE_BRUSH( "Icons/icon_Blueprint_Macro_16x", Icon16x16 ) );
 
@@ -6640,8 +6640,8 @@ void FStarshipEditorStyle::FStyle::SetupToolkitStyles()
 		Set( "MaterialEditor.SetCubePreview", new IMAGE_BRUSH( "Icons/AssetIcons/Cube_16x", Icon16x16) );
 		Set( "MaterialEditor.SetPreviewMeshFromSelection", new IMAGE_BRUSH_SVG( "Starship/AssetIcons/StaticMesh_64", Icon16x16 ) );
 		Set( "MaterialEditor.TogglePreviewGrid", new IMAGE_BRUSH_SVG( "Starship/Common/Grid", Icon16x16 ) );
-		Set( "MaterialEditor.ToggleMaterialStats", new IMAGE_BRUSH_SVG( "Starship/Common/Statistics", Icon20x20 ) );
-		Set( "MaterialEditor.ToggleMaterialStats.Tab", new IMAGE_BRUSH_SVG("Starship/Common/Statistics", Icon16x16));
+		Set( "MaterialEditor.ToggleMaterialStats", new IMAGE_BRUSH_SVG( "Starship/Common/MaterialAnalyzer", Icon20x20 ) );
+		Set( "MaterialEditor.ToggleMaterialStats.Tab", new IMAGE_BRUSH_SVG("Starship/Common/MaterialAnalyzer", Icon16x16));
 		Set( "MaterialEditor.TogglePlatformStats", new IMAGE_BRUSH_SVG( "Starship/MaterialEditor/PlatformStats", Icon20x20 ) );
 		Set("MaterialEditor.TogglePlatformStats.Tab", new IMAGE_BRUSH_SVG("Starship/MaterialEditor/PlatformStats", Icon16x16));
 		Set( "MaterialEditor.CameraHome", new IMAGE_BRUSH_SVG( "Starship/Common/Home", Icon20x20 ) );
@@ -7526,9 +7526,9 @@ void FStarshipEditorStyle::FStyle::SetupAutomationStyles()
 		Set( "DeviceDetails.Reboot", new IMAGE_BRUSH( "Icons/icon_DeviceReboot_40x", Icon40x40 ) );
 		Set( "DeviceDetails.Reboot.Small", new IMAGE_BRUSH( "Icons/icon_DeviceReboot_40x", Icon20x20 ) );
 
-		Set( "DeviceDetails.TabIcon", new IMAGE_BRUSH( "Icons/icon_tab_DeviceManager_16x", Icon16x16 ) );
+		Set( "DeviceDetails.TabIcon", new IMAGE_BRUSH_SVG( "Starship/Common/DeviceManager", Icon16x16 ) );
 		Set( "DeviceDetails.Tabs.Tools", new IMAGE_BRUSH( "/Icons/icon_tab_Tools_16x", Icon16x16 ) );
-		Set( "DeviceDetails.Tabs.ProfileEditor", new IMAGE_BRUSH( "/Icons/icon_tab_DeviceProfileEditor_16x", Icon16x16 ) );
+		Set( "DeviceDetails.Tabs.ProfileEditor", new IMAGE_BRUSH_SVG( "Starship/Common/DeviceProfiles", Icon16x16 ) );
 		Set( "DeviceDetails.Tabs.ProfileEditorSingleProfile", new IMAGE_BRUSH( "/Icons/icon_tab_DeviceProfileEditor_16x", Icon16x16 ) );
 	}
 
diff --git a/Engine/Source/Editor/LevelEditor/Private/SLevelEditor.cpp b/Engine/Source/Editor/LevelEditor/Private/SLevelEditor.cpp
--- a/Engine/Source/Editor/LevelEditor/Private/SLevelEditor.cpp
+++ b/Engine/Source/Editor/LevelEditor/Private/SLevelEditor.cpp
@@ -1378,11 +1378,10 @@ TSharedRef<SWidget> SLevelEditor::RestoreContentArea( const TSharedRef<SDockTab>
 
 		{
 			// @todo remove when world-centric mode is added
-			const FSlateIcon SequencerIcon("LevelSequenceEditorStyle", "LevelSequenceEditor.Tabs.Sequencer" );
 			LevelEditorTabManager->RegisterTabSpawner(LevelEditorTabIds::Sequencer, FOnSpawnTab::CreateSP<SLevelEditor, FName, FString>(this, &SLevelEditor::SpawnLevelEditorTab, LevelEditorTabIds::Sequencer, FString()) )
 				.SetDisplayName(NSLOCTEXT("LevelEditorTabs", "Sequencer", "Sequencer"))
 				.SetGroup( MenuStructure.GetLevelEditorCinematicsCategory() )
-				.SetIcon( SequencerIcon );
+				.SetIcon( FSlateIcon(FAppStyle::GetAppStyleSetName(), "LevelEditor.Tabs.Cinematics") );
 		}
 
 		{
diff --git a/Engine/Source/Editor/MainFrame/Private/Menus/MainMenu.cpp b/Engine/Source/Editor/MainFrame/Private/Menus/MainMenu.cpp
--- a/Engine/Source/Editor/MainFrame/Private/Menus/MainMenu.cpp
+++ b/Engine/Source/Editor/MainFrame/Private/Menus/MainMenu.cpp
@@ -158,24 +158,7 @@ void FMainMenu::RegisterEditMenu()
 			}
 		}));
 
-		Section.AddDynamicEntry("SourceControlDyn", FNewToolMenuSectionDelegate::CreateLambda([](FToolMenuSection& InSection)
-		{
-			if (ISourceControlModule::Get().IsEnabled() && ISourceControlModule::Get().GetProvider().IsAvailable())
-			{
-				// Choose specific files to submit
-				InSection.AddMenuEntry(FMainFrameCommands::Get().ViewChangelists, 
-					TAttribute<FText>(),
-					TAttribute<FText>(),
-					FSlateIcon(FEditorStyle::GetStyleSetName(), "SourceControl.ChangelistsTab"));
-			}
-			else
-			{
-				InSection.AddMenuEntry(FMainFrameCommands::Get().ConnectToSourceControl, 
-					TAttribute<FText>(),
-					TAttribute<FText>(),
-					FSlateIcon(FEditorStyle::GetStyleSetName(), "SourceControl.Actions.Connect"));
-			}
-		}));
+		Section.AddMenuEntry(FMainFrameCommands::Get().ConnectToSourceControl);
 	}
 }
 
@@ -218,7 +201,7 @@ void FMainMenu::RegisterWindowMenu()
 		// Make sure at least one is enabled before creating the section
 		if (bLocalizationDashboard || bTranslationPicker)
 		{
-			FToolMenuSection& Section = Menu->AddSection("ExperimentalTabSpawners", LOCTEXT("ExperimentalTabSpawnersHeading", "Experimental"), FToolMenuInsert("WindowGlobalTabSpawners", EToolMenuInsertType::After));
+			FToolMenuSection& Section = Menu->AddSection("ExperimentalTabSpawners", LOCTEXT("ExperimentalTabSpawnersHeading", "Experimental"), FToolMenuInsert("WindowLayout", EToolMenuInsertType::Before));
 			{
 				// Localization Dashboard
 				if (bLocalizationDashboard)
diff --git a/Engine/Source/Editor/PixelInspector/Private/PixelInspectorModule.cpp b/Engine/Source/Editor/PixelInspector/Private/PixelInspectorModule.cpp
--- a/Engine/Source/Editor/PixelInspector/Private/PixelInspectorModule.cpp
+++ b/Engine/Source/Editor/PixelInspector/Private/PixelInspectorModule.cpp
@@ -143,7 +143,7 @@ void FPixelInspectorModule::RegisterTabSpawner(const TSharedPtr<FWorkspaceItem>&
 		FTabSpawnerEntry& SpawnerEntry = FGlobalTabmanager::Get()->RegisterNomadTabSpawner("LevelEditorPixelInspector", FOnSpawnTab::CreateRaw(this, &FPixelInspectorModule::MakePixelInspectorTab))
 			.SetDisplayName(NSLOCTEXT("LevelEditorTabs", "LevelEditorPixelInspector", "Pixel Inspector"))
 			.SetTooltipText(NSLOCTEXT("LevelEditorTabs", "LevelEditorPixelInspectorTooltipText", "Open the viewport pixel inspector tool."))
-			.SetIcon(FSlateIcon(FPixelInspectorStyle::Get()->GetStyleSetName(), "PixelInspector.TabIcon"));
+			.SetIcon(FSlateIcon(FAppStyle::GetAppStyleSetName(), "PixelInspector.TabIcon"));
 
 		if (WorkspaceGroup.IsValid())
 		{
@@ -166,7 +166,7 @@ void FPixelInspectorModule::UnregisterTabSpawner()
 TSharedRef<SDockTab> FPixelInspectorModule::MakePixelInspectorTab(const FSpawnTabArgs&)
 {
 	TSharedRef<SDockTab> PixelInspectorTab = SNew(SDockTab)
-	.Icon(FPixelInspectorStyle::Get()->GetBrush("PixelInspector.TabIcon"))
+	.Icon(FAppStyle::Get().GetBrush("PixelInspector.TabIcon"))
 	.TabRole(ETabRole::NomadTab);
 	
 	PixelInspectorTab->SetContent(CreatePixelInspectorWidget());
diff --git a/Engine/Source/Editor/StructViewer/Private/StructViewerModule.cpp b/Engine/Source/Editor/StructViewer/Private/StructViewerModule.cpp
--- a/Engine/Source/Editor/StructViewer/Private/StructViewerModule.cpp
+++ b/Engine/Source/Editor/StructViewer/Private/StructViewerModule.cpp
@@ -47,7 +47,7 @@ void FStructViewerModule::StartupModule()
 		.SetDisplayName(NSLOCTEXT("StructViewerApp", "TabTitle", "Struct Viewer"))
 		.SetTooltipText(NSLOCTEXT("StructViewerApp", "TooltipText", "Displays all structs that exist within this project."))
 		.SetGroup(WorkspaceMenu::GetMenuStructure().GetToolsCategory())
-		.SetIcon(FSlateIcon(FEditorStyle::GetStyleSetName(), "ClassViewer.TabIcon"));
+		.SetIcon(FSlateIcon(FEditorStyle::GetStyleSetName(), "ClassIcon.UserDefinedStruct"));
 
 	ISettingsModule* SettingsModule = FModuleManager::GetModulePtr<ISettingsModule>("Settings");
 	if (SettingsModule)
diff --git a/Engine/Source/Editor/WorkspaceMenuStructure/Private/WorkspaceMenuStructureModule.cpp b/Engine/Source/Editor/WorkspaceMenuStructure/Private/WorkspaceMenuStructureModule.cpp
--- a/Engine/Source/Editor/WorkspaceMenuStructure/Private/WorkspaceMenuStructureModule.cpp
+++ b/Engine/Source/Editor/WorkspaceMenuStructure/Private/WorkspaceMenuStructureModule.cpp
@@ -120,9 +120,9 @@ public:
 		DeveloperToolsDebugCategory = DeveloperToolsSubMenuDebugCategory->AddGroup(LOCTEXT("WorkspaceMenu_DeveloperToolsDebugCategory", "Debug"), FSlateIcon(), true);
 		AutomationToolsCategory = DeveloperToolsSubMenuDebugCategory->AddGroup(LOCTEXT("WorkspaceMenu_AutomationToolsCategory", "Testing"), FSlateIcon(), true);
 
-		DeveloperToolsProfilingCategory = DeveloperToolsCategory->AddGroup(LOCTEXT("WorkspaceMenu_DeveloperToolsProfilingCategory", "Profile"),  FSlateIcon(FEditorStyle::GetStyleSetName(), "LevelEditor.Tabs.Debug"), true);
-		DeveloperToolsAuditCategory =     DeveloperToolsCategory->AddGroup(LOCTEXT("WorkspaceMenu_DeveloperToolsAuditCategory", "Audit"), FSlateIcon(FEditorStyle::GetStyleSetName(), "LevelEditor.Tabs.Debug"), true);
-		DeveloperToolsPlatformsCategory = DeveloperToolsCategory->AddGroup(LOCTEXT("WorkspaceMenu_DeveloperToolsPlatforms:WCategory", "Platforms"), FSlateIcon(FEditorStyle::GetStyleSetName(), "LevelEditor.Tabs.Debug"), true);
+		DeveloperToolsProfilingCategory = DeveloperToolsCategory->AddGroup(LOCTEXT("WorkspaceMenu_DeveloperToolsProfilingCategory", "Profile"),  FSlateIcon(FEditorStyle::GetStyleSetName(), "LevelEditor.Profile"), true);
+		DeveloperToolsAuditCategory =     DeveloperToolsCategory->AddGroup(LOCTEXT("WorkspaceMenu_DeveloperToolsAuditCategory", "Audit"), FSlateIcon(FEditorStyle::GetStyleSetName(), "LevelEditor.Audit"), true);
+		DeveloperToolsPlatformsCategory = DeveloperToolsCategory->AddGroup(LOCTEXT("WorkspaceMenu_DeveloperToolsPlatforms:WCategory", "Platforms"), FSlateIcon(FEditorStyle::GetStyleSetName(), "LevelEditor.Platforms"), true);
 		DeveloperToolsMiscCategory =      DeveloperToolsCategory->AddGroup(LOCTEXT("WorkspaceMenu_DeveloperToolsMiscCategory", "Miscellaneous"), FSlateIcon(FEditorStyle::GetStyleSetName(), "LevelEditor.Tabs.Debug"), true);
 	}
 
