Updating icons on major asset editor toolbars: Animation/Skeleton/Physics Asset, Texture editors, Play dropdown menu

diff --git a/Engine/Source/Editor/ClothPainter/Private/ClothPaintingModule.cpp b/Engine/Source/Editor/ClothPainter/Private/ClothPaintingModule.cpp
--- a/Engine/Source/Editor/ClothPainter/Private/ClothPaintingModule.cpp
+++ b/Engine/Source/Editor/ClothPainter/Private/ClothPaintingModule.cpp
@@ -137,7 +137,7 @@ void FClothPaintingModule::RegisterMenus()
 					FClothPainterCommands::Get().TogglePaintMode,
 					TAttribute<FText>::Create(TAttribute<FText>::FGetter::CreateRaw(this, &FClothPaintingModule::GetPaintToolsButtonText, Context->SkeletalMeshEditor)),
 					FText(),
-					FSlateIcon(FEditorStyle::GetStyleSetName(), "LevelEditor.MeshPaintMode.TexturePaint")
+					FSlateIcon(FAppStyle::Get().GetStyleSetName(), "MeshPaint.Brush")
 				));	
 			}
 		}));
diff --git a/Engine/Source/Editor/EditorStyle/Private/StarshipStyle.cpp b/Engine/Source/Editor/EditorStyle/Private/StarshipStyle.cpp
--- a/Engine/Source/Editor/EditorStyle/Private/StarshipStyle.cpp
+++ b/Engine/Source/Editor/EditorStyle/Private/StarshipStyle.cpp
@@ -4891,6 +4891,7 @@ void FStarshipEditorStyle::FStyle::SetupLevelEditorStyle()
 		Set("MeshPaint.CycleToNextLOD.Small",     new IMAGE_BRUSH(TEXT("Icons/GeneralTools/Next_40x"), Icon20x20));
 		Set("MeshPaint.SwitchForeAndBackgroundColor",       new IMAGE_BRUSH("Icons/Paint/Paint_SwapColors_40x", Icon20x20));
 		Set("MeshPaint.SwitchForeAndBackgroundColor.Small", new IMAGE_BRUSH("Icons/Paint/Paint_SwapColors_40x", Icon20x20));
+		Set("MeshPaint.Brush",			new IMAGE_BRUSH_SVG("Starship/Common/Paintbrush", Icon20x20));
 	}
 
 	// Scalability (Performance Warning)
@@ -4956,63 +4957,38 @@ void FStarshipEditorStyle::FStyle::SetupPersonaStyle()
 			.SetHovered(IMAGE_BRUSH("/Docking/CloseApp_Hovered", Icon16x16)));
 
 		// persona commands
-		Set("Persona.AnimNotifyWindow", new IMAGE_BRUSH(TEXT("Icons/icon_Persona_AnimNotift_40x"), Icon40x40));
-		Set("Persona.AnimNotifyWindow.Small", new IMAGE_BRUSH(TEXT("Icons/icon_Persona_AnimNotift_40x"), Icon20x20));
-		Set("Persona.RetargetManager", new IMAGE_BRUSH(TEXT("Icons/icon_Persona_Retarget_40x"), Icon40x40));
-		Set("Persona.RetargetManager.Small", new IMAGE_BRUSH(TEXT("Icons/icon_Persona_Retarget_40x"), Icon20x20));
-		Set("Persona.ImportMesh", new IMAGE_BRUSH(TEXT("Icons/icon_Persona_ImportMesh_40x"), Icon40x40));
-		Set("Persona.ImportMesh.Small", new IMAGE_BRUSH(TEXT("Icons/icon_Persona_ImportMesh_40x"), Icon20x20));
-		Set("Persona.ReimportMesh", new IMAGE_BRUSH(TEXT("Icons/icon_Persona_ReimportMesh_40x"), Icon40x40));
-		Set("Persona.ReimportMesh.Small", new IMAGE_BRUSH(TEXT("Icons/icon_Persona_ReimportMesh_40x"), Icon20x20));
-		Set("Persona.ImportLODs", new IMAGE_BRUSH(TEXT("Icons/icon_Persona_ImportLODs_40x"), Icon40x40));
-		Set("Persona.ImportLODs.Small", new IMAGE_BRUSH(TEXT("Icons/icon_Persona_ImportLODs_40x"), Icon20x20));
-//		Set("Persona.AddBodyPart", new IMAGE_BRUSH(TEXT("Icons/icon_Placement_AllContent_40x"), Icon40x40));
-//		Set("Persona.AddBodyPart.Small", new IMAGE_BRUSH(TEXT("Icons/icon_Placement_AllContent_40x"), Icon20x20));
-		Set("Persona.ImportAnimation", new IMAGE_BRUSH(TEXT("Icons/icon_Persona_ImportAnim_40x"), Icon40x40));
-		Set("Persona.ImportAnimation.Small", new IMAGE_BRUSH(TEXT("Icons/icon_Persona_ImportAnim_40x"), Icon20x20));
-		Set("Persona.ReimportAnimation", new IMAGE_BRUSH(TEXT("Icons/icon_Persona_ReimportAnim_40x"), Icon40x40));
-		Set("Persona.ReimportAnimation.Small", new IMAGE_BRUSH(TEXT("Icons/icon_Persona_ReimportAnim_40x"), Icon20x20));
-		Set("Persona.ApplyCompression", new IMAGE_BRUSH(TEXT("Icons/icon_Persona_Compression_40x"), Icon40x40));
-		Set("Persona.ApplyCompression.Small", new IMAGE_BRUSH(TEXT("Icons/icon_Persona_Compression_40x"), Icon20x20));
-		Set("Persona.ExportToFBX", new IMAGE_BRUSH(TEXT("Icons/icon_Persona_ExportFBX_40x"), Icon40x40));
-		Set("Persona.ExportToFBX.Small", new IMAGE_BRUSH(TEXT("Icons/icon_Persona_ExportFBX_40x"), Icon20x20));
-		Set("Persona.CreateAsset", new IMAGE_BRUSH(TEXT("Icons/icon_Persona_CreateAsset_40x"), Icon40x40));
-		Set("Persona.CreateAsset.Small", new IMAGE_BRUSH(TEXT("Icons/icon_Persona_CreateAsset_40x"), Icon20x20));
+		Set("Persona.AnimNotifyWindow", new IMAGE_BRUSH_SVG("Starship/Persona/AnimationNotifies", Icon20x20));
+		Set("Persona.RetargetManager", new IMAGE_BRUSH_SVG("Starship/Persona/RetargetManager", Icon20x20));
+		Set("Persona.ImportMesh", new CORE_IMAGE_BRUSH_SVG("Starship/Common/import", Icon20x20));
+		Set("Persona.ReimportMesh", new CORE_IMAGE_BRUSH_SVG("Starship/Common/import", Icon20x20));
+		Set("Persona.ImportLODs", new CORE_IMAGE_BRUSH_SVG("Starship/Common/import", Icon20x20));
+		Set("Persona.ImportAnimation", new CORE_IMAGE_BRUSH_SVG("Starship/Common/import", Icon20x20));
+		Set("Persona.ReimportAnimation", new CORE_IMAGE_BRUSH_SVG("Starship/Common/import", Icon20x20));
+		Set("Persona.ApplyCompression", new IMAGE_BRUSH_SVG("Starship/Common/Compress", Icon20x20));
+		Set("Persona.ExportToFBX", new IMAGE_BRUSH_SVG("Starship/Common/Export", Icon20x20));
+		Set("Persona.CreateAsset", new IMAGE_BRUSH_SVG("Starship/Persona/PersonaCreateAsset", Icon20x20));
 		Set("Persona.StartRecordAnimation", new IMAGE_BRUSH(TEXT("Icons/icon_Persona_StartRecord_40x"), Icon40x40));
-		Set("Persona.StartRecordAnimation.Small", new IMAGE_BRUSH(TEXT("Icons/icon_Persona_StartRecord_40x"), Icon20x20));
 		Set("Persona.StopRecordAnimation", new IMAGE_BRUSH(TEXT("Icons/icon_Persona_StopRecord_40x"), Icon40x40));
-		Set("Persona.StopRecordAnimation.Small", new IMAGE_BRUSH(TEXT("Icons/icon_Persona_StopRecord_40x"), Icon20x20));
 		Set("Persona.StopRecordAnimation_Alt", new IMAGE_BRUSH(TEXT("Icons/icon_Persona_StopRecord_Alt_40x"), Icon40x40));
-		Set("Persona.StopRecordAnimation_Alt.Small", new IMAGE_BRUSH(TEXT("Icons/icon_Persona_StopRecord_Alt_40x"), Icon20x20));
 		Set("Persona.SetKey", new IMAGE_BRUSH(TEXT("Icons/icon_Persona_SetKey_40x"), Icon40x40));
-		Set("Persona.SetKey.Small", new IMAGE_BRUSH(TEXT("Icons/icon_Persona_SetKey_40x"), Icon20x20));
 		Set("Persona.ApplyAnimation", new IMAGE_BRUSH(TEXT("Icons/icon_Persona_BakeAnim_40x"), Icon40x40));
-		Set("Persona.ApplyAnimation.Small", new IMAGE_BRUSH(TEXT("Icons/icon_Persona_BakeAnim_40x"), Icon20x20));
 
 		// preview set up
-		Set("Persona.TogglePreviewAsset", new IMAGE_BRUSH(TEXT("Icons/icon_Persona_PreviewAsset_40x"), Icon40x40));
-		Set("Persona.TogglePreviewAsset.Small", new IMAGE_BRUSH(TEXT("Icons/icon_Persona_PreviewAsset_40x"), Icon20x20));
-		Set("Persona.TogglePreviewAnimation", new IMAGE_BRUSH(TEXT("Icons/icon_Persona_PreviewAnim_40x"), Icon40x40));
-		Set("Persona.TogglePreviewAnimation.Small", new IMAGE_BRUSH(TEXT("Icons/icon_Persona_PreviewAnim_40x"), Icon20x20));
-		Set("Persona.ToggleReferencePose", new IMAGE_BRUSH(TEXT("Icons/icon_Persona_ReferencePose_40x"), Icon40x40));
-		Set("Persona.ToggleReferencePose.Small", new IMAGE_BRUSH(TEXT("Icons/icon_Persona_ReferencePose_40x"), Icon20x20));
+		Set("Persona.TogglePreviewAsset", new IMAGE_BRUSH_SVG("Starship/Persona/AnimationPreviewMesh", Icon20x20));
+		Set("Persona.TogglePreviewAnimation", new IMAGE_BRUSH_SVG("Starship/Persona/PersonaPreviewAnimation", Icon20x20));
+		Set("Persona.ToggleReferencePose", new IMAGE_BRUSH_SVG("Starship/Persona/PersonaTPose", Icon20x20));
 		Set("Persona.SavePreviewMeshCollection", new IMAGE_BRUSH(TEXT("Icons/Save_16x"), Icon16x16));
 
 		// persona extras
 		Set("Persona.ConvertAnimationGraph", new IMAGE_BRUSH("Old/Graph/ConvertIcon", Icon40x40));
-		Set("Persona.ReimportAsset", new IMAGE_BRUSH("Icons/Reimport_12x", Icon12x12));
-		Set("Persona.ConvertToStaticMesh", new IMAGE_BRUSH("Icons/icon_ShowStaticMeshes_40x", Icon40x40));
-		Set("Persona.ConvertToStaticMesh.Small", new IMAGE_BRUSH("Icons/icon_ShowStaticMeshes_40x", Icon20x20));
+		Set("Persona.ReimportAsset", new CORE_IMAGE_BRUSH_SVG("Starship/Common/import", Icon20x20));
+		Set("Persona.ConvertToStaticMesh", new IMAGE_BRUSH_SVG("Starship/Persona/MakeStaticMesh", Icon20x20));
 		Set("Persona.BakeMaterials", new IMAGE_BRUSH("Icons/icon_tab_Layers_40x", Icon40x40));
-		Set("Persona.BakeMaterials.Small", new IMAGE_BRUSH("Icons/icon_tab_Layers_40x", Icon20x20));
 
 		// Anim Slot Manager
 		Set("AnimSlotManager.SaveSkeleton", new IMAGE_BRUSH("Persona/AnimSlotManager/icon_SaveSkeleton_40x", Icon40x40));
-		Set("AnimSlotManager.SaveSkeleton.Small", new IMAGE_BRUSH("Persona/AnimSlotManager/icon_SaveSkeleton_40x", Icon20x20));
 		Set("AnimSlotManager.AddGroup", new IMAGE_BRUSH("Persona/AnimSlotManager/icon_AddGroup_40x", Icon40x40));
-		Set("AnimSlotManager.AddGroup.Small", new IMAGE_BRUSH("Persona/AnimSlotManager/icon_AddGroup_40x", Icon20x20));
 		Set("AnimSlotManager.AddSlot", new IMAGE_BRUSH("Persona/AnimSlotManager/icon_AddSlot_40x", Icon40x40));
-		Set("AnimSlotManager.AddSlot.Small", new IMAGE_BRUSH("Persona/AnimSlotManager/icon_AddSlot_40x", Icon20x20));
 		Set("AnimSlotManager.Warning", new IMAGE_BRUSH("Persona/AnimSlotManager/icon_Warning_14x", Icon16x16));
 
 		// Anim Notify Editor
@@ -5027,7 +5003,7 @@ void FStarshipEditorStyle::FStyle::SetupPersonaStyle()
 		Set("Persona.Tabs.MorphTargetPreviewer", new IMAGE_BRUSH(TEXT("Icons/icon_Persona_Morph_Target_Previewer_16x"), Icon16x16));
 		Set("Persona.Tabs.AnimCurvePreviewer", new IMAGE_BRUSH(TEXT("Icons/icon_Persona_AnimCurve_Previewer_16x"), Icon16x16));
 		Set("Persona.Tabs.AnimationNotifies", new IMAGE_BRUSH(TEXT("Icons/icon_Persona_Animation_Notifies_16x"), Icon16x16));
-		Set("Persona.Tabs.RetargetManager", new IMAGE_BRUSH(TEXT("Icons/icon_Persona_Retarget_Manager_16x"), Icon16x16));
+		Set("Persona.Tabs.RetargetManager", new IMAGE_BRUSH_SVG("Starship/Persona/RetargetManager", Icon16x16));
 		Set("Persona.Tabs.AnimSlotManager", new IMAGE_BRUSH(TEXT("Icons/icon_Persona_Anim_Slot_Manager_16x"), Icon16x16));
 		Set("Persona.Tabs.SkeletonCurves", new IMAGE_BRUSH(TEXT("Icons/icon_Persona_Skeleton_Curves_16x"), Icon16x16));
 		Set("Persona.Tabs.AnimAssetDetails", new IMAGE_BRUSH(TEXT("Icons/icon_Persona_Anim_Asset_Details_16x"), Icon16x16));
@@ -5043,12 +5019,9 @@ void FStarshipEditorStyle::FStyle::SetupPersonaStyle()
 
 	// Skeleton editor
 	{
-		Set("SkeletonEditor.AnimNotifyWindow", new IMAGE_BRUSH(TEXT("Icons/icon_Persona_AnimNotift_40x"), Icon40x40));
-		Set("SkeletonEditor.AnimNotifyWindow.Small", new IMAGE_BRUSH(TEXT("Icons/icon_Persona_AnimNotift_40x"), Icon20x20));
-		Set("SkeletonEditor.RetargetManager", new IMAGE_BRUSH(TEXT("Icons/icon_Persona_Retarget_40x"), Icon40x40));
-		Set("SkeletonEditor.RetargetManager.Small", new IMAGE_BRUSH(TEXT("Icons/icon_Persona_Retarget_40x"), Icon20x20));
-		Set("SkeletonEditor.ImportMesh", new IMAGE_BRUSH(TEXT("Icons/icon_Persona_ImportMesh_40x"), Icon40x40));
-		Set("SkeletonEditor.ImportMesh.Small", new IMAGE_BRUSH(TEXT("Icons/icon_Persona_ImportMesh_40x"), Icon20x20));
+		Set("SkeletonEditor.AnimNotifyWindow", new IMAGE_BRUSH_SVG("Starship/Persona/AnimationNotifies", Icon20x20));
+		Set("SkeletonEditor.RetargetManager", new IMAGE_BRUSH_SVG("Starship/Persona/RetargetManager", Icon20x20));
+		Set("SkeletonEditor.ImportMesh", new CORE_IMAGE_BRUSH_SVG("Starship/Common/import", Icon20x20));
 
 		// Skeleton Tree
 		Set("SkeletonTree.SkeletonSocket", new IMAGE_BRUSH("Persona/SkeletonTree/icon_SocketG_16px", Icon16x16));
@@ -5088,18 +5061,12 @@ void FStarshipEditorStyle::FStyle::SetupPersonaStyle()
 
 	// Animation editor
 	{
-		Set("AnimationEditor.ApplyCompression", new IMAGE_BRUSH(TEXT("Icons/icon_Persona_Compression_40x"), Icon40x40));
-		Set("AnimationEditor.ApplyCompression.Small", new IMAGE_BRUSH(TEXT("Icons/icon_Persona_Compression_40x"), Icon20x20));
-		Set("AnimationEditor.ExportToFBX", new IMAGE_BRUSH(TEXT("Icons/icon_Persona_ExportFBX_40x"), Icon40x40));
-		Set("AnimationEditor.ExportToFBX.Small", new IMAGE_BRUSH(TEXT("Icons/icon_Persona_ExportFBX_40x"), Icon20x20));
-		Set("AnimationEditor.ReimportAnimation", new IMAGE_BRUSH(TEXT("Icons/icon_Persona_ReimportAnim_40x"), Icon40x40));
-		Set("AnimationEditor.ReimportAnimation.Small", new IMAGE_BRUSH(TEXT("Icons/icon_Persona_ReimportAnim_40x"), Icon20x20));
-		Set("AnimationEditor.CreateAsset", new IMAGE_BRUSH(TEXT("Icons/icon_Persona_CreateAsset_40x"), Icon40x40));
-		Set("AnimationEditor.CreateAsset.Small", new IMAGE_BRUSH(TEXT("Icons/icon_Persona_CreateAsset_40x"), Icon20x20));
-		Set("AnimationEditor.SetKey", new IMAGE_BRUSH(TEXT("Icons/icon_Persona_SetKey_40x"), Icon40x40));
-		Set("AnimationEditor.SetKey.Small", new IMAGE_BRUSH(TEXT("Icons/icon_Persona_SetKey_40x"), Icon20x20));
-		Set("AnimationEditor.ApplyAnimation", new IMAGE_BRUSH(TEXT("Icons/icon_Persona_BakeAnim_40x"), Icon40x40));
-		Set("AnimationEditor.ApplyAnimation.Small", new IMAGE_BRUSH(TEXT("Icons/icon_Persona_BakeAnim_40x"), Icon20x20));
+		Set("AnimationEditor.ApplyCompression", new IMAGE_BRUSH_SVG("Starship/Common/Compress", Icon20x20));
+		Set("AnimationEditor.ExportToFBX", new IMAGE_BRUSH_SVG("Starship/Common/Export", Icon20x20));
+		Set("AnimationEditor.ReimportAnimation", new CORE_IMAGE_BRUSH_SVG("Starship/Common/import", Icon20x20));
+		Set("AnimationEditor.CreateAsset", new IMAGE_BRUSH_SVG("Starship/Persona/PersonaCreateAsset", Icon20x20));
+		Set("AnimationEditor.SetKey", new CORE_IMAGE_BRUSH_SVG("Starship/Common/plus", Icon20x20));
+		Set("AnimationEditor.ApplyAnimation", new IMAGE_BRUSH_SVG("Starship/Common/Apply", Icon20x20));
 
 		Set("AnimTimeline.Outliner.DefaultBorder", new FSlateColorBrush(FLinearColor::White));
 
@@ -5121,13 +5088,10 @@ void FStarshipEditorStyle::FStyle::SetupPersonaStyle()
 
 	// Skeletal mesh editor
 	{
-		Set("SkeletalMeshEditor.ReimportMesh", new IMAGE_BRUSH(TEXT("Icons/icon_Persona_ReimportMesh_40x"), Icon40x40));
-		Set("SkeletalMeshEditor.ReimportMesh.Small", new IMAGE_BRUSH(TEXT("Icons/icon_Persona_ReimportMesh_40x"), Icon20x20));
-		Set("SkeletalMeshEditor.ImportLODs", new IMAGE_BRUSH(TEXT("Icons/icon_Persona_ImportLODs_40x"), Icon40x40));
-		Set("SkeletalMeshEditor.ImportLODs.Small", new IMAGE_BRUSH(TEXT("Icons/icon_Persona_ImportLODs_40x"), Icon20x20));
+		Set("SkeletalMeshEditor.ReimportMesh", new CORE_IMAGE_BRUSH_SVG("Starship/Common/import", Icon20x20));
+		Set("SkeletalMeshEditor.ImportLODs", new CORE_IMAGE_BRUSH_SVG("Starship/Common/import", Icon20x20));
 
-		Set("SkeletalMeshEditor.MeshSectionSelection", new IMAGE_BRUSH(TEXT("Icons/icon_Persona_MeshSectionSelection_40x"), Icon40x40));
-		Set("SkeletalMeshEditor.MeshSectionSelection.Small", new IMAGE_BRUSH(TEXT("Icons/icon_Persona_MeshSectionSelection_40x"), Icon20x20));
+		Set("SkeletalMeshEditor.MeshSectionSelection", new IMAGE_BRUSH_SVG("Starship/Persona/SectionSelection", Icon20x20));
 	}
 
 	// Kismet 2
@@ -5397,29 +5361,20 @@ void FStarshipEditorStyle::FStyle::SetupPersonaStyle()
 
 		Set("PlayToolBar", PlayToolbar);
 
-		Set("PlayWorld.Simulate", new IMAGE_BRUSH_SVG("Starship/MainToolbar/simulate", Icon40x40));
-		Set("PlayWorld.Simulate.Small", new IMAGE_BRUSH_SVG("Starship/MainToolbar/simulate", Icon20x20));
-
-		Set("PlayWorld.RepeatLastPlay", new IMAGE_BRUSH_SVG("Starship/MainToolbar/simulate", Icon40x40));
-		Set("PlayWorld.RepeatLastPlay.Small", new IMAGE_BRUSH_SVG("Starship/MainToolbar/simulate", Icon20x20));
+		Set("PlayWorld.Simulate", new IMAGE_BRUSH_SVG("Starship/MainToolbar/simulate", Icon20x20));
+		Set( "PlayWorld.RepeatLastPlay", new IMAGE_BRUSH_SVG("Starship/MainToolbar/play", Icon20x20));
 		Set( "PlayWorld.PlayInViewport", new IMAGE_BRUSH_SVG("Starship/MainToolbar/play", Icon20x20));
 
-		Set( "PlayWorld.PlayInEditorFloating", new IMAGE_BRUSH( "Icons/icon_playInWindow_40x", Icon40x40 ) );
-		Set( "PlayWorld.PlayInEditorFloating.Small", new IMAGE_BRUSH( "Icons/icon_playInWindow_40x", Icon20x20 ) );
-		Set( "PlayWorld.PlayInVR", new IMAGE_BRUSH( "Icons/icon_playInVR_40x", Icon40x40 ) );
-		Set( "PlayWorld.PlayInVR.Small", new IMAGE_BRUSH( "Icons/icon_playInVR_16x", Icon20x20 ) );
-		Set( "PlayWorld.PlayInMobilePreview", new IMAGE_BRUSH( "Icons/icon_PlayMobilePreview_40x", Icon40x40 ) );
-		Set( "PlayWorld.PlayInMobilePreview.Small", new IMAGE_BRUSH( "Icons/icon_PlayMobilePreview_16x", Icon20x20 ) );
-		Set( "PlayWorld.PlayInVulkanPreview", new IMAGE_BRUSH( "Icons/icon_PlayMobilePreview_40x", Icon40x40 ) );
-		Set( "PlayWorld.PlayInVulkanPreview.Small", new IMAGE_BRUSH( "Icons/icon_PlayMobilePreview_16x", Icon20x20 ) );
-		Set( "PlayWorld.PlayInNewProcess", new IMAGE_BRUSH( "Icons/icon_PlayStandalone_40x", Icon40x40 ) );
-		Set( "PlayWorld.PlayInNewProcess.Small", new IMAGE_BRUSH( "Icons/icon_PlayStandalone_40x", Icon20x20 ) );
+		Set( "PlayWorld.PlayInEditorFloating", new IMAGE_BRUSH_SVG("Starship/MainToolbar/PlayNewEditorWindow", Icon20x20));
+		Set( "PlayWorld.PlayInVR", new IMAGE_BRUSH_SVG("Starship/MainToolbar/PlayVRPreview", Icon20x20));
+		Set( "PlayWorld.PlayInMobilePreview", new IMAGE_BRUSH_SVG("Starship/MainToolbar/PlayMobilePreview", Icon20x20));
+		Set( "PlayWorld.PlayInVulkanPreview", new IMAGE_BRUSH_SVG("Starship/MainToolbar/PlayMobilePreview", Icon20x20));
+		Set( "PlayWorld.PlayInNewProcess", new IMAGE_BRUSH_SVG("Starship/MainToolbar/PlayStandaloneGame", Icon20x20));
 
-		Set( "PlayWorld.RepeatLastLaunch", new IMAGE_BRUSH_SVG( "Starship/MainToolbar/launch", Icon40x40 ) );
-		Set( "PlayWorld.RepeatLastLaunch.Small", new IMAGE_BRUSH_SVG( "Starship/MainToolbar/launch", Icon20x20 ) );
+		Set( "PlayWorld.RepeatLastLaunch", new IMAGE_BRUSH_SVG( "Starship/MainToolbar/launch", Icon20x20 ) );
 
-		Set( "PlayWorld.PlayInCameraLocation", new IMAGE_BRUSH( "Icons/icon_PlayCameraLocation_40x", Icon40x40 ) );
-		Set( "PlayWorld.PlayInDefaultPlayerStart", new IMAGE_BRUSH( "Icons/icon_PlayDefaultPlayerStart_40x", Icon40x40 ) );
+		Set( "PlayWorld.PlayInCameraLocation", new IMAGE_BRUSH_SVG( "Starship/AssetIcons/CameraActor_16", Icon20x20 ) );
+		Set( "PlayWorld.PlayInDefaultPlayerStart", new IMAGE_BRUSH_SVG( "Starship/AssetIcons/PlayerStart_16", Icon20x20 ) );
 
 		Set("PlayWorld.ResumePlaySession", new IMAGE_BRUSH_SVG("Starship/MainToolbar/simulate", Icon40x40));
 		Set("PlayWorld.ResumePlaySession.Small", new IMAGE_BRUSH_SVG("Starship/MainToolbar/simulate", Icon20x20));
@@ -5434,10 +5389,8 @@ void FStarshipEditorStyle::FStyle::SetupPersonaStyle()
 		Set("PlayWorld.LateJoinSession", new IMAGE_BRUSH_SVG("Starship/MainToolbar/simulate", Icon40x40));
 		Set("PlayWorld.LateJoinSession.Small", new IMAGE_BRUSH_SVG("Starship/MainToolbar/simulate", Icon20x20));
 
-		Set( "PlayWorld.PossessPlayer", new IMAGE_BRUSH( "Icons/icon_possess_40x", Icon40x40 ) );
-		Set( "PlayWorld.PossessPlayer.Small", new IMAGE_BRUSH( "Icons/icon_possess_40x", Icon20x20 ) );
+		Set( "PlayWorld.PossessPlayer", new IMAGE_BRUSH_SVG("Starship/AssetIcons/PlayerController_16", Icon20x20));
 		Set( "PlayWorld.EjectFromPlayer", new IMAGE_BRUSH_SVG("Starship/MainToolbar/eject", Icon40x40));
-		Set( "PlayWorld.EjectFromPlayer.Small", new IMAGE_BRUSH_SVG("Starship/MainToolbar/eject", Icon20x20));
 
 		Set( "PlayWorld.ShowCurrentStatement", new IMAGE_BRUSH( "Icons/icon_findnode_40x", Icon40x40 ) );
 		Set( "PlayWorld.ShowCurrentStatement.Small", new IMAGE_BRUSH( "Icons/icon_findnode_40x", Icon20x20 ) );
@@ -5477,8 +5430,8 @@ void FStarshipEditorStyle::FStyle::SetupPersonaStyle()
 
 	// Asset context menu
 	{
-		Set("Persona.AssetActions.CreateAnimAsset", new IMAGE_BRUSH(TEXT("Icons/icon_Persona_CreateAsset_16x"), Icon16x16));
-		Set("Persona.AssetActions.ReimportAnim", new IMAGE_BRUSH(TEXT("Icons/icon_Persona_ReimportAnim_16x"), Icon16x16));
+		Set("Persona.AssetActions.CreateAnimAsset", new IMAGE_BRUSH_SVG("Starship/Persona/PersonaCreateAsset", Icon20x20));
+		Set("Persona.AssetActions.ReimportAnim", new CORE_IMAGE_BRUSH_SVG("Starship/Common/import", Icon20x20));
 		Set("Persona.AssetActions.Retarget", new IMAGE_BRUSH(TEXT("Icons/icon_Persona_Retarget_16x"), Icon16x16));
 		Set("Persona.AssetActions.RetargetSkeleton", new IMAGE_BRUSH(TEXT("Icons/icon_Animation_Retarget_Skeleton_16x"), Icon16x16));
 		Set("Persona.AssetActions.FindSkeleton", new IMAGE_BRUSH(TEXT("Icons/icon_Genericfinder_16x"), Icon16x16));
@@ -6746,10 +6699,8 @@ void FStarshipEditorStyle::FStyle::SetupToolkitStyles()
 		Set("TextureEditor.Saturation", new IMAGE_BRUSH( "Icons/icon_TextureEd_Saturation_40x", Icon40x40));
 		Set("TextureEditor.Saturation.Small", new IMAGE_BRUSH( "Icons/icon_TextureEd_Saturation_40x", Icon20x20));
 
-		Set("TextureEditor.CompressNow", new IMAGE_BRUSH( "Icons/icon_TextureEd_CompressNow_40x", Icon40x40));
-		Set("TextureEditor.CompressNow.Small", new IMAGE_BRUSH( "Icons/icon_TextureEd_CompressNow_40x", Icon20x20));
-		Set("TextureEditor.Reimport", new IMAGE_BRUSH( "Icons/icon_TextureEd_Reimport_40x", Icon40x40));
-		Set("TextureEditor.Reimport.Small", new IMAGE_BRUSH( "Icons/icon_TextureEd_Reimport_40x", Icon20x20));
+		Set("TextureEditor.CompressNow", new IMAGE_BRUSH_SVG( "Starship/Common/Compress", Icon20x20));
+		Set("TextureEditor.Reimport", new CORE_IMAGE_BRUSH_SVG("Starship/Common/import", Icon20x20));
 
 		FButtonStyle MipmapButtonStyle = 
 			FButtonStyle(FStarshipCoreStyle::GetCoreStyle().GetWidgetStyle<FButtonStyle>("Button"))
@@ -7020,22 +6971,22 @@ void FStarshipEditorStyle::FStyle::SetupToolkitStyles()
 
 		Set( "PhysicsAssetEditor.SimulationNoGravity", new IMAGE_BRUSH( "PhysicsAssetEditor/icon_PhAT_PlaySimNoGravity_40x", Icon40x40 ) );
 		Set( "PhysicsAssetEditor.SelectedSimulation", new IMAGE_BRUSH( "PhysicsAssetEditor/icon_PhAT_PlaySimSelected_40x", Icon40x40 ) );
-		Set( "PhysicsAssetEditor.SimulationAll", new IMAGE_BRUSH( "PhysicsAssetEditor/icon_PhAT_PlaySim_40x", Icon40x40 ) );
-		Set( "PhysicsAssetEditor.Undo", new IMAGE_BRUSH( "Icons/icon_Generic_Undo_40x", Icon40x40 ) );
-		Set( "PhysicsAssetEditor.Redo", new IMAGE_BRUSH( "Icons/icon_Generic_Redo_40x", Icon40x40 ) );
+		Set( "PhysicsAssetEditor.SimulationAll", new IMAGE_BRUSH_SVG("Starship/MainToolbar/simulate", Icon20x20));
+		Set( "PhysicsAssetEditor.Undo", new CORE_IMAGE_BRUSH_SVG("Starship/Common/Undo", Icon20x20));
+		Set( "PhysicsAssetEditor.Redo", new CORE_IMAGE_BRUSH_SVG("Starship/Common/Redo", Icon20x20));
 		Set( "PhysicsAssetEditor.ChangeDefaultMesh", new IMAGE_BRUSH( "PhysicsAssetEditor/icon_PhAT_Mesh_40x", Icon40x40 ) );
-		Set( "PhysicsAssetEditor.ApplyPhysicalMaterial", new IMAGE_BRUSH( "PhysicsAssetEditor/icon_PhAT_PhysMat_40x", Icon40x40 ) );
+		Set( "PhysicsAssetEditor.ApplyPhysicalMaterial", new IMAGE_BRUSH_SVG("Starship/Persona/AnimationPhysicalMaterial", Icon20x20));
 		Set( "PhysicsAssetEditor.CopyJointSettings", new IMAGE_BRUSH( "PhysicsAssetEditor/icon_PhAT_CopyJoints_40x", Icon40x40 ) );
-		Set( "PhysicsAssetEditor.PlayAnimation", new IMAGE_BRUSH( "PhysicsAssetEditor/icon_PhAT_Play_40x", Icon40x40 ) );
+		Set( "PhysicsAssetEditor.PlayAnimation", new IMAGE_BRUSH_SVG( "Starship/MainToolbar/play", Icon20x20 ) );
 		Set( "PhysicsAssetEditor.PhATTranslationMode", new IMAGE_BRUSH( "Icons/icon_translate_40x", Icon40x40 ) );
 		Set( "PhysicsAssetEditor.PhATRotationMode", new IMAGE_BRUSH( "Icons/icon_rotate_40x", Icon40x40 ) );
 		Set( "PhysicsAssetEditor.PhATScaleMode", new IMAGE_BRUSH( "Icons/icon_scale_40x", Icon40x40 ) );
 		Set( "PhysicsAssetEditor.Snap", new IMAGE_BRUSH( "PhysicsAssetEditor/icon_PhAT_Snap_40x", Icon40x40 ) );
 		Set( "PhysicsAssetEditor.CopyProperties", new IMAGE_BRUSH( "PhysicsAssetEditor/icon_PhAT_CopyProperties_40x", Icon40x40 ) );
-		Set( "PhysicsAssetEditor.DisableCollision", new IMAGE_BRUSH( "PhysicsAssetEditor/icon_PhAT_DisableCollision_40x", Icon40x40 ) );
-		Set( "PhysicsAssetEditor.EnableCollision", new IMAGE_BRUSH( "PhysicsAssetEditor/icon_PhAT_EnableCollision_40x", Icon40x40 ) );
-		Set( "PhysicsAssetEditor.DisableCollisionAll", new IMAGE_BRUSH( "PhysicsAssetEditor/icon_PhAT_DisableCollision_40x", Icon40x40 ) );
-		Set( "PhysicsAssetEditor.EnableCollisionAll", new IMAGE_BRUSH( "PhysicsAssetEditor/icon_PhAT_EnableCollision_40x", Icon40x40 ) );
+		Set( "PhysicsAssetEditor.DisableCollision", new IMAGE_BRUSH_SVG( "Starship/Common/DisableCollision", Icon20x20 ) );
+		Set( "PhysicsAssetEditor.EnableCollision", new IMAGE_BRUSH_SVG("Starship/Common/EnableCollision", Icon20x20));
+		Set( "PhysicsAssetEditor.DisableCollisionAll", new IMAGE_BRUSH_SVG("Starship/Common/DisableCollision", Icon20x20));
+		Set( "PhysicsAssetEditor.EnableCollisionAll", new IMAGE_BRUSH_SVG("Starship/Common/EnableCollision", Icon20x20));
 		Set( "PhysicsAssetEditor.WeldToBody", new IMAGE_BRUSH( "PhysicsAssetEditor/icon_PhAT_Weld_40x", Icon40x40 ) );
 		Set( "PhysicsAssetEditor.AddNewBody", new IMAGE_BRUSH( "PhysicsAssetEditor/icon_PhAT_NewBody_40x", Icon40x40 ) );
 		Set( "PhysicsAssetEditor.AddSphere", new IMAGE_BRUSH( "PhysicsAssetEditor/icon_PhAT_Sphere_40x", Icon40x40 ) );
@@ -7047,49 +6998,12 @@ void FStarshipEditorStyle::FStyle::SetupToolkitStyles()
 		Set( "PhysicsAssetEditor.ResetConstraint", new IMAGE_BRUSH( "PhysicsAssetEditor/icon_PhAT_ResetConstraint_40x", Icon40x40 ) );
 		Set( "PhysicsAssetEditor.SnapConstraint", new IMAGE_BRUSH( "PhysicsAssetEditor/icon_PhAT_SnapConstraint_40x", Icon40x40 ) );
 		Set( "PhysicsAssetEditor.SnapAllConstraints", new IMAGE_BRUSH( "PhysicsAssetEditor/icon_PhAT_SnapAll_40x", Icon40x40 ) );
-		Set( "PhysicsAssetEditor.ConvertToBallAndSocket", new IMAGE_BRUSH( "PhysicsAssetEditor/icon_PhAT_Ball_40x", Icon40x40 ) );
-		Set( "PhysicsAssetEditor.ConvertToHinge", new IMAGE_BRUSH( "PhysicsAssetEditor/icon_PhAT_Hinge_40x", Icon40x40 ) );
-		Set( "PhysicsAssetEditor.ConvertToPrismatic", new IMAGE_BRUSH( "PhysicsAssetEditor/icon_PhAT_Prismatic_40x", Icon40x40 ) );
-		Set( "PhysicsAssetEditor.ConvertToSkeletal", new IMAGE_BRUSH( "PhysicsAssetEditor/icon_PhAT_Skeletal_40x", Icon40x40 ) );
+		Set( "PhysicsAssetEditor.ConvertToBallAndSocket", new IMAGE_BRUSH_SVG("Starship/Persona/AnimationToBallAndSocket", Icon20x20));
+		Set( "PhysicsAssetEditor.ConvertToHinge", new IMAGE_BRUSH_SVG("Starship/Persona/AnimationToHinge", Icon20x20));
+		Set( "PhysicsAssetEditor.ConvertToPrismatic", new IMAGE_BRUSH_SVG("Starship/Persona/AnimationToPrismatic", Icon20x20));
+		Set( "PhysicsAssetEditor.ConvertToSkeletal", new IMAGE_BRUSH_SVG("Starship/Persona/AnimationToSkeletal", Icon20x20));
 		Set( "PhysicsAssetEditor.DeleteConstraint", new IMAGE_BRUSH( "PhysicsAssetEditor/icon_PhAT_DeleteConstraint_40x", Icon40x40 ) );
 
-		Set("PhysicsAssetEditor.SimulationNoGravity.Small", new IMAGE_BRUSH("PhysicsAssetEditor/icon_PhAT_PlaySimNoGravity_40x", Icon20x20));
-		Set("PhysicsAssetEditor.SelectedSimulation.Small", new IMAGE_BRUSH("PhysicsAssetEditor/icon_PhAT_PlaySimSelected_40x", Icon20x20));
-		Set( "PhysicsAssetEditor.SimulationAll.Small", new IMAGE_BRUSH( "PhysicsAssetEditor/icon_PhAT_PlaySim_40x", Icon20x20 ) );
-		Set( "PhysicsAssetEditor.Undo.Small", new IMAGE_BRUSH( "Icons/icon_Generic_Undo_40x", Icon20x20 ) );
-		Set( "PhysicsAssetEditor.Redo.Small", new IMAGE_BRUSH( "Icons/icon_Generic_Redo_40x", Icon20x20 ) );
-		Set( "PhysicsAssetEditor.ChangeDefaultMesh.Small", new IMAGE_BRUSH( "PhysicsAssetEditor/icon_PhAT_Mesh_40x", Icon20x20 ) );
-		Set( "PhysicsAssetEditor.ResetEntireAsset.Small", new IMAGE_BRUSH( "PhysicsAssetEditor/icon_PhAT_ResetAsset_40x", Icon20x20 ) );
-		Set( "PhysicsAssetEditor.ResetBoneCollision.Small", new IMAGE_BRUSH( "PhysicsAssetEditor/icon_PhAT_ResetCollision_40x", Icon20x20 ) );
-		Set( "PhysicsAssetEditor.ApplyPhysicalMaterial.Small", new IMAGE_BRUSH( "PhysicsAssetEditor/icon_PhAT_PhysMat_40x", Icon20x20 ) );
-		Set( "PhysicsAssetEditor.CopyJointSettings.Small", new IMAGE_BRUSH( "PhysicsAssetEditor/icon_PhAT_CopyJoints_40x", Icon20x20 ) );
-		Set( "PhysicsAssetEditor.PlayAnimation.Small", new IMAGE_BRUSH( "PhysicsAssetEditor/icon_PhAT_Play_40x", Icon20x20 ) );
-		Set( "PhysicsAssetEditor.PhATTranslationMode.Small", new IMAGE_BRUSH( "Icons/icon_translate_40x", Icon20x20 ) );
-		Set( "PhysicsAssetEditor.PhATRotationMode.Small", new IMAGE_BRUSH( "Icons/icon_rotate_40x", Icon20x20 ) );
-		Set( "PhysicsAssetEditor.PhATScaleMode.Small", new IMAGE_BRUSH( "Icons/icon_scale_40x", Icon20x20 ) );
-		Set( "PhysicsAssetEditor.Snap.Small", new IMAGE_BRUSH( "PhysicsAssetEditor/icon_PhAT_Snap_40x", Icon20x20 ) );
-		Set( "PhysicsAssetEditor.CopyProperties.Small", new IMAGE_BRUSH( "PhysicsAssetEditor/icon_PhAT_CopyProperties_40x", Icon20x20 ) );
-		Set( "PhysicsAssetEditor.DisableCollision.Small", new IMAGE_BRUSH( "PhysicsAssetEditor/icon_PhAT_DisableCollision_40x", Icon20x20 ) );
-		Set( "PhysicsAssetEditor.EnableCollision.Small", new IMAGE_BRUSH( "PhysicsAssetEditor/icon_PhAT_EnableCollision_40x", Icon20x20 ) );
-		Set( "PhysicsAssetEditor.DisableCollisionAll.Small", new IMAGE_BRUSH( "PhysicsAssetEditor/icon_PhAT_DisableCollision_40x", Icon20x20 ) );
-		Set( "PhysicsAssetEditor.EnableCollisionAll.Small", new IMAGE_BRUSH( "PhysicsAssetEditor/icon_PhAT_EnableCollision_40x", Icon20x20 ) );
-		Set( "PhysicsAssetEditor.WeldToBody.Small", new IMAGE_BRUSH( "PhysicsAssetEditor/icon_PhAT_Weld_40x", Icon20x20 ) );
-		Set( "PhysicsAssetEditor.AddNewBody.Small", new IMAGE_BRUSH( "PhysicsAssetEditor/icon_PhAT_NewBody_40x", Icon20x20 ) );
-		Set( "PhysicsAssetEditor.AddSphere.Small", new IMAGE_BRUSH( "PhysicsAssetEditor/icon_PhAT_Sphere_40x", Icon20x20 ) );
-		Set( "PhysicsAssetEditor.AddSphyl.Small", new IMAGE_BRUSH( "PhysicsAssetEditor/icon_PhAT_Sphyl_40x", Icon20x20 ) );
-		Set( "PhysicsAssetEditor.AddBox.Small", new IMAGE_BRUSH( "PhysicsAssetEditor/icon_PhAT_Box_40x", Icon20x20 ) );
-		Set( "PhysicsAssetEditor.AddTaperedCapsule.Small", new IMAGE_BRUSH( "PhysicsAssetEditor/icon_PhAT_TaperedCapsule_40x", Icon20x20 ) );
-		Set( "PhysicsAssetEditor.DeletePrimitive.Small", new IMAGE_BRUSH( "PhysicsAssetEditor/icon_PhAT_DeletePrimitive_40x", Icon20x20 ) );
-		Set( "PhysicsAssetEditor.DuplicatePrimitive.Small", new IMAGE_BRUSH( "PhysicsAssetEditor/icon_PhAT_DupePrim_40x", Icon20x20 ) );
-		Set( "PhysicsAssetEditor.ResetConstraint.Small", new IMAGE_BRUSH( "PhysicsAssetEditor/icon_PhAT_ResetConstraint_40x", Icon20x20 ) );
-		Set( "PhysicsAssetEditor.SnapConstraint.Small", new IMAGE_BRUSH( "PhysicsAssetEditor/icon_PhAT_SnapConstraint_40x", Icon20x20 ) );
-		Set( "PhysicsAssetEditor.SnapAllConstraints.Small", new IMAGE_BRUSH( "PhysicsAssetEditor/icon_PhAT_SnapAll_40x", Icon20x20 ) );
-		Set( "PhysicsAssetEditor.ConvertToBallAndSocket.Small", new IMAGE_BRUSH( "PhysicsAssetEditor/icon_PhAT_Ball_40x", Icon20x20 ) );
-		Set( "PhysicsAssetEditor.ConvertToHinge.Small", new IMAGE_BRUSH( "PhysicsAssetEditor/icon_PhAT_Hinge_40x", Icon20x20 ) );
-		Set( "PhysicsAssetEditor.ConvertToPrismatic.Small", new IMAGE_BRUSH( "PhysicsAssetEditor/icon_PhAT_Prismatic_40x", Icon20x20 ) );
-		Set( "PhysicsAssetEditor.ConvertToSkeletal.Small", new IMAGE_BRUSH( "PhysicsAssetEditor/icon_PhAT_Skeletal_40x", Icon20x20 ) );
-		Set( "PhysicsAssetEditor.DeleteConstraint.Small", new IMAGE_BRUSH( "PhysicsAssetEditor/icon_PhAT_DeleteConstraint_40x", Icon20x20 ) );
-
 		Set("PhysicsAssetEditor.NewPhysicalAnimationProfile", new IMAGE_BRUSH("PhysicsAssetEditor/icon_PhAT_NewBody_40x", Icon20x20));
 		Set("PhysicsAssetEditor.DeleteCurrentPhysicalAnimationProfile", new IMAGE_BRUSH("PhysicsAssetEditor/icon_PhAT_DeletePrimitive_40x", Icon20x20));
 		Set("PhysicsAssetEditor.AddBodyToPhysicalAnimationProfile", new IMAGE_BRUSH("PhysicsAssetEditor/icon_PhAT_NewBody_40x", Icon20x20));
