Fixing comment bubble styling and changing the slate style class used

diff --git a/Engine/Source/Editor/EditorStyle/Private/StarshipStyle.cpp b/Engine/Source/Editor/EditorStyle/Private/StarshipStyle.cpp
--- a/Engine/Source/Editor/EditorStyle/Private/StarshipStyle.cpp
+++ b/Engine/Source/Editor/EditorStyle/Private/StarshipStyle.cpp
@@ -3748,6 +3748,10 @@ void FStarshipEditorStyle::FStyle::SetupGraphEditorStyles()
 				.SetBackgroundImageHovered(FSlateRoundedBoxBrush(FStyleColors::Foreground, FStyleColors::Hover, 1.0f))
 				.SetBackgroundImageFocused(FSlateRoundedBoxBrush(FStyleColors::Foreground, FStyleColors::Primary, 1.0f))
 				.SetBackgroundImageReadOnly(FSlateRoundedBoxBrush(FStyleColors::Header, FStyleColors::InputOutline, 1.0f))
+				.SetForegroundColor(FStyleColors::Background)
+				.SetBackgroundColor(FStyleColors::White)
+				.SetReadOnlyForegroundColor(FStyleColors::Foreground)
+				.SetFocusedForegroundColor(FStyleColors::Background)
 				.SetScrollBarStyle( ScrollBar );
 			Set( "Graph.Node.NodeTitleEditableText", GraphNodeTitleEditableText );
 
@@ -3762,6 +3766,20 @@ void FStarshipEditorStyle::FStyle::SetupGraphEditorStyles()
 				.SetShadowOffset( FVector2D::ZeroVector )
 				.SetShadowColorAndOpacity( FLinearColor(0.f,0.f,0.f, 0.7f) )
 			);
+
+			FEditableTextBoxStyle CommentEditableText = FEditableTextBoxStyle(NormalEditableTextBoxStyle)
+				.SetFont(NormalText.Font)
+				.SetForegroundColor(FStyleColors::Foreground)
+				.SetBackgroundImageNormal(FSlateRoundedBoxBrush(FStyleColors::Foreground, FStyleColors::Secondary, 1.0f))
+				.SetBackgroundImageHovered(FSlateRoundedBoxBrush(FStyleColors::Foreground, FStyleColors::Hover, 1.0f))
+				.SetBackgroundImageFocused(FSlateRoundedBoxBrush(FStyleColors::Foreground, FStyleColors::Primary, 1.0f))
+				.SetBackgroundImageReadOnly(FSlateRoundedBoxBrush(FStyleColors::Header, FStyleColors::InputOutline, 1.0f))
+				.SetForegroundColor(FStyleColors::Background)
+				.SetBackgroundColor(FStyleColors::White)
+				.SetReadOnlyForegroundColor(FSlateColor::UseForeground())
+				.SetFocusedForegroundColor(FStyleColors::White)
+				.SetScrollBarStyle(ScrollBar);
+			Set("Graph.CommentBubble.EditableText", CommentEditableText);
 		
 			FTextBlockStyle GraphCommentBlockTitle = FTextBlockStyle(NormalText)
 				.SetFont( DEFAULT_FONT( "Bold", 18 ) )
diff --git a/Engine/Source/Editor/GraphEditor/Private/SCommentBubble.cpp b/Engine/Source/Editor/GraphEditor/Private/SCommentBubble.cpp
--- a/Engine/Source/Editor/GraphEditor/Private/SCommentBubble.cpp
+++ b/Engine/Source/Editor/GraphEditor/Private/SCommentBubble.cpp
@@ -251,7 +251,7 @@ void SCommentBubble::UpdateBubble()
 						[
 							SAssignNew(TextBlock, SMultiLineEditableTextBox)
 							.Text(MakeAttributeLambda([this] { return CachedCommentText; }))
-							.Style(FAppStyle::Get(), "Graph.Node.NodeTitleInlineEditableText")
+							.Style(FAppStyle::Get(), "Graph.CommentBubble.EditableText")
 							.HintText( NSLOCTEXT( "CommentBubble", "EditCommentHint", "Click to edit" ))
 							.IsReadOnly(this, &SCommentBubble::IsReadOnly)
 							.SelectAllTextWhenFocused( true )
@@ -409,17 +409,17 @@ void SCommentBubble::OnCommentTextCommitted( const FText& NewText, ETextCommit::
 
 FSlateColor SCommentBubble::GetTextBackgroundColor() const
 {
-	return TextBlock->HasKeyboardFocus() ? FStyleColors::Foreground : SCommentBubbleDefs::TextClearBackground;
+	return TextBlock->HasKeyboardFocus() ? FSlateColor::UseStyle() : SCommentBubbleDefs::TextClearBackground;
 }
 
 FSlateColor SCommentBubble::GetTextForegroundColor() const
 {
-	return TextBlock->HasKeyboardFocus() ? FSlateColor::UseStyle() : ForegroundColor;
+	return FSlateColor::UseStyle();
 }
 
 FSlateColor SCommentBubble::GetReadOnlyTextForegroundColor() const
 {
-	return ForegroundColor;
+	return TextBlock->HasKeyboardFocus() ? FStyleColors::Foreground : FSlateColor::UseStyle();
 }
 
 EVisibility SCommentBubble::GetToggleButtonVisibility() const
