Added missing component icons

diff --git a/Engine/Source/Editor/EditorStyle/Private/StarshipStyle.cpp b/Engine/Source/Editor/EditorStyle/Private/StarshipStyle.cpp
--- a/Engine/Source/Editor/EditorStyle/Private/StarshipStyle.cpp
+++ b/Engine/Source/Editor/EditorStyle/Private/StarshipStyle.cpp
@@ -5744,20 +5744,22 @@ void FStarshipEditorStyle::FStyle::SetupClassIconsAndThumbnails()
 		Set("ClassIcon.Deleted", new IMAGE_BRUSH("Icons/ActorIcons/DeletedActor_16px", Icon16x16));
 
 		// Component classes
-		Set("ClassIcon.AudioComponent", new IMAGE_BRUSH("Icons/ActorIcons/SoundActor_16x", Icon16x16));
+	
 		Set("ClassIcon.BlueprintCore", new IMAGE_BRUSH("Icons/AssetIcons/Blueprint_16x", Icon16x16));
 		Set("ClassIcon.LightComponent", new IMAGE_BRUSH("Icons/ActorIcons/LightActor_16x", Icon16x16));
 		Set("ClassIcon.ArrowComponent", new IMAGE_BRUSH("Icons/ActorIcons/Arrow_16px", Icon16x16));
-		Set("ClassIcon.BoxComponent", new IMAGE_BRUSH("Icons/ActorIcons/Box_16px", Icon16x16));
-		Set("ClassIcon.CapsuleComponent", new IMAGE_BRUSH("Icons/ActorIcons/Capsule_16px", Icon16x16));
-		Set("ClassIcon.InstancedStaticMeshComponent", new IMAGE_BRUSH("Icons/ActorIcons/InstancedStaticMesh_16px", Icon16x16));
 		Set("ClassIcon.MaterialBillboardComponent", new IMAGE_BRUSH("Icons/ActorIcons/MaterialSprite_16px", Icon16x16));
-		Set("ClassIcon.SphereComponent", new IMAGE_BRUSH("Icons/ActorIcons/Sphere_16px", Icon16x16));
-		Set("ClassIcon.SplineComponent", new IMAGE_BRUSH("Icons/ActorIcons/Spline_16px", Icon16x16));
 		Set("ClassIcon.BillboardComponent", new IMAGE_BRUSH("Icons/ActorIcons/SpriteComponent_16px", Icon16x16));
 		Set("ClassIcon.TimelineComponent", new IMAGE_BRUSH("Icons/ActorIcons/TimelineComponent_16px", Icon16x16));
 		Set("ClassIcon.ChildActorComponent", new IMAGE_BRUSH("Icons/ActorIcons/ChildActorComponent_16px", Icon16x16));
 
+		Set("ClassIcon.AudioComponent", new IMAGE_BRUSH_SVG("Starship/AssetIcons/Audio_16", Icon16x16));
+		Set("ClassIcon.BoxComponent", new IMAGE_BRUSH_SVG("Starship/AssetIcons/BoxCollision_16", Icon16x16));
+		Set("ClassIcon.CapsuleComponent", new IMAGE_BRUSH_SVG("Starship/AssetIcons/CapsuleCollision_16", Icon16x16));
+		Set("ClassIcon.InstancedStaticMeshComponent", new IMAGE_BRUSH_SVG("Icons/ActorIcons/InstancedStaticMesh_16px", Icon16x16));
+		Set("ClassIcon.SphereComponent", new IMAGE_BRUSH_SVG("Starship/AssetIcons/SphereCollision_16", Icon16x16));
+		Set("ClassIcon.SplineComponent", new IMAGE_BRUSH_SVG("Starship/AssetIcons/Spline_16", Icon16x16));
+
 		Set("ClassIcon.AtmosphericFogComponent", new IMAGE_BRUSH_SVG("Starship/AssetIcons/AtmosphericFog_16", Icon16x16));
 		Set("ClassIcon.BrushComponent", new IMAGE_BRUSH_SVG("Starship/AssetIcons/Brush_16", Icon16x16));
 		Set("ClassIcon.CableComponent", new IMAGE_BRUSH_SVG("Starship/AssetIcons/CableActor_16", Icon16x16));
