Remove a few output log custom colors that can be styled with the new style system

diff --git a/Engine/Source/Editor/EditorStyle/Private/StarshipStyle.cpp b/Engine/Source/Editor/EditorStyle/Private/StarshipStyle.cpp
--- a/Engine/Source/Editor/EditorStyle/Private/StarshipStyle.cpp
+++ b/Engine/Source/Editor/EditorStyle/Private/StarshipStyle.cpp
@@ -47,9 +47,6 @@ void FStarshipEditorStyle::Initialize()
 	PropertyEditorModule.RegisterCustomClassLayout("EditorStyleSettings", FOnGetDetailCustomizationInstance::CreateStatic(&FEditorStyleSettingsCustomization::MakeInstance));
 	PropertyEditorModule.RegisterCustomPropertyTypeLayout("StyleColorList", FOnGetPropertyTypeCustomizationInstance::CreateStatic(&FStyleColorListCustomization::MakeInstance));
 
-	// Background
-	USlateThemeManager::Get().SetDefaultColor(EStyleColor::User1, USlateThemeManager::Get().GetColor(EStyleColor::Recessed));
-	USlateThemeManager::Get().SetColorDisplayName(EStyleColor::User1, LOCTEXT("UserColor_OutputLogBackground", "Log Background"));
 	// Selection highlight
 	USlateThemeManager::Get().SetDefaultColor(EStyleColor::User2, USlateThemeManager::Get().GetColor(EStyleColor::Highlight));
 	USlateThemeManager::Get().SetColorDisplayName(EStyleColor::User2, LOCTEXT("UserColor_OutputLogHighlight", "Log Highlight"));
@@ -59,12 +56,7 @@ void FStarshipEditorStyle::Initialize()
 	// Command
 	USlateThemeManager::Get().SetDefaultColor(EStyleColor::User4, USlateThemeManager::Get().GetColor(EStyleColor::AccentGreen));
 	USlateThemeManager::Get().SetColorDisplayName(EStyleColor::User4, LOCTEXT("UserColor_OutputLogCommand", "Log Command"));
-	// Warning
-	USlateThemeManager::Get().SetDefaultColor(EStyleColor::User5, USlateThemeManager::Get().GetColor(EStyleColor::AccentYellow));
-	USlateThemeManager::Get().SetColorDisplayName(EStyleColor::User5, LOCTEXT("UserColor_OutputLogWarning", "Log Warning"));
-	// Error
-	USlateThemeManager::Get().SetDefaultColor(EStyleColor::User6, USlateThemeManager::Get().GetColor(EStyleColor::AccentRed));
-	USlateThemeManager::Get().SetColorDisplayName(EStyleColor::User6, LOCTEXT("UserColor_OutputLogError", "Log Error"));
+
 #endif
 
 	const FString ThemesSubDir = TEXT("Slate/Themes");
@@ -162,13 +154,9 @@ FStarshipEditorStyle::FStyle::FStyle( const TWeakObjectPtr< UEditorStyleSettings
 	, SelectionColor_Subdued( SelectionColor_Subdued_LinearRef )
 	, HighlightColor( HighlightColor_LinearRef )
 	, WindowHighlightColor(WindowHighlightColor_LinearRef)
-
-	, LogColor_Background(EStyleColor::User1)
 	, LogColor_SelectionBackground(EStyleColor::User2)
 	, LogColor_Normal(EStyleColor::User3)
 	, LogColor_Command(EStyleColor::User4)
-	, LogColor_Warning(EStyleColor::User5)
-	, LogColor_Error(EStyleColor::User6)
 
 	, InheritedFromBlueprintTextColor(FLinearColor(0.25f, 0.5f, 1.0f))
 
@@ -967,11 +955,11 @@ void FStarshipEditorStyle::FStyle::SetupGeneralStyles()
 			);
 
 		Set("Log.Warning", FTextBlockStyle(NormalLogText)
-			.SetColorAndOpacity( LogColor_Warning )
+			.SetColorAndOpacity(FStyleColors::Warning)
 			);
 
 		Set("Log.Error", FTextBlockStyle(NormalLogText)
-			.SetColorAndOpacity( LogColor_Error )
+			.SetColorAndOpacity(FStyleColors::Error)
 			);
 
 		Set("Log.TabIcon", new IMAGE_BRUSH_SVG( "Starship/Common/OutputLog", Icon16x16 ) );
@@ -5470,7 +5458,7 @@ void FStarshipEditorStyle::FStyle::SetupPersonaStyle()
 		Set("BlendSpaceKey.Pressed", SelectionColor_Pressed);
 		Set("BlendSpaceKey.Drag", SelectionColor_Subdued);
 		Set("BlendSpaceKey.Drop", SelectionColor_Inactive);
-		Set("BlendSpaceKey.Invalid", LogColor_Error);
+		Set("BlendSpaceKey.Invalid", FStyleColors::Warning);
 		Set("BlendSpaceKey.Preview", LogColor_Command);
 	}
 
diff --git a/Engine/Source/Editor/EditorStyle/Private/StarshipStyle.h b/Engine/Source/Editor/EditorStyle/Private/StarshipStyle.h
--- a/Engine/Source/Editor/EditorStyle/Private/StarshipStyle.h
+++ b/Engine/Source/Editor/EditorStyle/Private/StarshipStyle.h
@@ -119,12 +119,9 @@ public:
 		const FSlateColor HighlightColor;
 		const FSlateColor WindowHighlightColor;
 
-		const FSlateColor LogColor_Background;
 		const FSlateColor LogColor_SelectionBackground;
 		const FSlateColor LogColor_Normal;
 		const FSlateColor LogColor_Command;
-		const FSlateColor LogColor_Warning;
-		const FSlateColor LogColor_Error;
 
 		// These are common colors used thruout the editor in mutliple style elements
 		const FSlateColor InheritedFromBlueprintTextColor;
