Sequencer: Add Sequencer cursor decorators

diff --git a/Engine/Source/Editor/EditorStyle/Private/StarshipStyle.cpp b/Engine/Source/Editor/EditorStyle/Private/StarshipStyle.cpp
--- a/Engine/Source/Editor/EditorStyle/Private/StarshipStyle.cpp
+++ b/Engine/Source/Editor/EditorStyle/Private/StarshipStyle.cpp
@@ -2173,6 +2173,8 @@ void FStarshipEditorStyle::FStyle::SetupSequencerStyles()
 
 		Set("Sequencer.CursorDecorator_MarqueeAdd", new IMAGE_BRUSH("Sequencer/CursorDecorator_MarqueeAdd", Icon16x16));
 		Set("Sequencer.CursorDecorator_MarqueeSubtract", new IMAGE_BRUSH("Sequencer/CursorDecorator_MarqueeSubtract", Icon16x16));
+		Set("Sequencer.CursorDecorator_Retime", new IMAGE_BRUSH("Sequencer/CursorDecorator_Retime", Icon16x16));
+		Set("Sequencer.CursorDecorator_EasingHandle", new IMAGE_BRUSH("Sequencer/CursorDecorator_EasingHandle", Icon16x16));
 
 		Set("Sequencer.BreadcrumbText", FTextBlockStyle(NormalText).SetFont(FStyleFonts::Get().NormalBold));
 		Set("Sequencer.BreadcrumbIcon", new IMAGE_BRUSH("Common/SmallArrowRight", Icon10x10));
