New anim graph icons

diff --git a/Engine/Source/Editor/AnimationBlueprintEditor/Private/AnimationPins/SGraphPinPose.cpp b/Engine/Source/Editor/AnimationBlueprintEditor/Private/AnimationPins/SGraphPinPose.cpp
--- a/Engine/Source/Editor/AnimationBlueprintEditor/Private/AnimationPins/SGraphPinPose.cpp
+++ b/Engine/Source/Editor/AnimationBlueprintEditor/Private/AnimationPins/SGraphPinPose.cpp
@@ -76,34 +76,14 @@ public:
 		AnimGraphNode = InAnimGraphNode;
 		ActiveColor = AttributeDesc->Color.GetSpecifiedColor();
 		Color = ActiveColor;
-		BackgroundColor = FLinearColor::Transparent;
 		Value = 0.0f;
 
 		ChildSlot
 		[
-			SNew(SOverlay)
-			.ToolTipText(AttributeDesc->ToolTipText)
-			+SOverlay::Slot()
-			[
-				SNew(SBorder)
-				.Padding(FMargin(2.0f, 1.0f, 5.0f, 1.0f))
-				.BorderImage(FEditorStyle::GetBrush("AnimGraph.Attribute.Border.Solid"))
-				.BorderBackgroundColor_Lambda([this](){ return BackgroundColor; })
-			]
-			+SOverlay::Slot()
-			.Padding(FMargin(2.0f, 1.0f, 2.0f, 1.0f))
-			[
-				SNew(SHorizontalBox)
-				+SHorizontalBox::Slot()
-				.AutoWidth()
-				.VAlign(VAlign_Center)
-				.HAlign(HAlign_Center)
-				[
-					SNew(SImage)
-					.Image(&AttributeDesc->Icon)
-					.ColorAndOpacity_Lambda([this](){ return Color; })
-				]
-			]
+			SNew(SImage)
+			.Visibility_Lambda([this](){ return Value == 0.0f ? EVisibility::Hidden : EVisibility::Visible; })
+			.Image(&AttributeDesc->Icon)
+			.ColorAndOpacity_Lambda([this](){ return Color; })
 		];
 	}
 
@@ -134,15 +114,13 @@ public:
 		if(bActive)
 		{
 			Value = 1.0f;
-			Color = FLinearColor::Black;
-			BackgroundColor = ActiveColor;
 		}
 		else
 		{
 			Value = FMath::FInterpTo(Value, 0.0f, InDeltaTime, 4.0f);
-			Color = FMath::Lerp(ActiveColor, FLinearColor::Black, Value);
-			BackgroundColor = FMath::Lerp(FLinearColor::Transparent, ActiveColor, Value);
 		}
+
+		Color = FMath::Lerp(FLinearColor::Transparent, ActiveColor, Value);
 	}
 
 	UAnimBlueprint* AnimBlueprint;
@@ -152,7 +130,6 @@ public:
 	const UAnimGraphNode_Base* AnimGraphNode;
 	FLinearColor ActiveColor;
 	FLinearColor Color;
-	FLinearColor BackgroundColor;
 	float Value;
 };
 
diff --git a/Engine/Source/Editor/EditorStyle/Private/StarshipStyle.cpp b/Engine/Source/Editor/EditorStyle/Private/StarshipStyle.cpp
--- a/Engine/Source/Editor/EditorStyle/Private/StarshipStyle.cpp
+++ b/Engine/Source/Editor/EditorStyle/Private/StarshipStyle.cpp
@@ -3929,10 +3929,10 @@ void FStarshipEditorStyle::FStyle::SetupGraphEditorStyles()
 		Set( "Graph.ExecPin.DisconnectedHovered", new IMAGE_BRUSH( "Old/Graph/ExecPin_Disconnected", Icon12x16, FLinearColor(0.8f,0.8f,0.8f) ) );
 
 		const FVector2D Icon15x28(15.0f, 28.0f);
-		Set("Graph.PosePin.Connected", new IMAGE_BRUSH("Graph/Animation/PosePin_Connected_15x28", Icon15x28));
-		Set("Graph.PosePin.Disconnected", new IMAGE_BRUSH("Graph/Animation/PosePin_Disconnected_15x28", Icon15x28));
-		Set("Graph.PosePin.ConnectedHovered", new IMAGE_BRUSH("Graph/Animation/PosePin_Connected_15x28", Icon15x28, FLinearColor(0.8f, 0.8f, 0.8f)));
-		Set("Graph.PosePin.DisconnectedHovered", new IMAGE_BRUSH("Graph/Animation/PosePin_Disconnected_15x28", Icon15x28, FLinearColor(0.8f, 0.8f, 0.8f)));
+		Set("Graph.PosePin.Connected", new IMAGE_BRUSH_SVG("Starship/AnimationBlueprintEditor/AnimationGraphPose", Icon16x16));
+		Set("Graph.PosePin.Disconnected", new IMAGE_BRUSH_SVG("Starship/AnimationBlueprintEditor/AnimationGraphPoseDisconnected", Icon16x16));
+		Set("Graph.PosePin.ConnectedHovered", new IMAGE_BRUSH_SVG("Starship/AnimationBlueprintEditor/AnimationGraphPose", Icon16x16, FLinearColor(0.8f, 0.8f, 0.8f)));
+		Set("Graph.PosePin.DisconnectedHovered", new IMAGE_BRUSH_SVG("Starship/AnimationBlueprintEditor/AnimationGraphPoseDisconnected", Icon16x16, FLinearColor(0.8f, 0.8f, 0.8f)));
 
 		// Events Exec Pins
 		Set( "Graph.ExecEventPin.Connected", new IMAGE_BRUSH( "Graph/EventPin_Connected", Icon16x16 ) );
@@ -5722,15 +5722,15 @@ void FStarshipEditorStyle::FStyle::SetupPersonaStyle()
 			.SetColorAndOpacity(AttributeTextColor));
 
 		Set("AnimGraph.Attribute.Pose.Color", FStyleColors::White);
-		Set("AnimGraph.Attribute.Curves.Icon", new IMAGE_BRUSH("GenericCurveEditor/Icons/GenericCurveEditor_48x", Icon10x10));
+		Set("AnimGraph.Attribute.Curves.Icon", new IMAGE_BRUSH_SVG("Starship/AnimationBlueprintEditor/AnimationGraphCurves", Icon16x16));
 		Set("AnimGraph.Attribute.Curves.Color", FStyleColors::AccentPurple);
-		Set("AnimGraph.Attribute.Attributes.Icon", new IMAGE_BRUSH("Icons/pillarray_16x", Icon10x10));
+		Set("AnimGraph.Attribute.Attributes.Icon", new IMAGE_BRUSH_SVG("Starship/AnimationBlueprintEditor/AnimationGraphAttributes", Icon16x16));
 		Set("AnimGraph.Attribute.Attributes.Color", FStyleColors::AccentYellow);
-		Set("AnimGraph.Attribute.Sync.Icon", new IMAGE_BRUSH("Icons/icon_tab_TimecodeProvider_16x", Icon10x10));
+		Set("AnimGraph.Attribute.Sync.Icon", new IMAGE_BRUSH_SVG("Starship/AnimationBlueprintEditor/AnimationGraphSync", Icon16x16));
 		Set("AnimGraph.Attribute.Sync.Color", FStyleColors::AccentBlue);
-		Set("AnimGraph.Attribute.RootMotionDelta.Icon", new IMAGE_BRUSH("Icons/icon_axis_local_16px", Icon12x12));
+		Set("AnimGraph.Attribute.RootMotionDelta.Icon", new IMAGE_BRUSH("Icons/icon_axis_local_16px", Icon16x16));
 		Set("AnimGraph.Attribute.RootMotionDelta.Color", FStyleColors::AccentGreen);
-		Set("AnimGraph.Attribute.InertialBlending.Icon", new IMAGE_BRUSH("Icons/icon_tab_TimecodeProvider_16x", Icon10x10));
+		Set("AnimGraph.Attribute.InertialBlending.Icon", new IMAGE_BRUSH_SVG("Starship/AnimationBlueprintEditor/AnimationGraphInertialBlending", Icon16x16));
 		Set("AnimGraph.Attribute.InertialBlending.Color", FStyleColors::AccentOrange);
 	}
 	
