Adding missing ISM and HISM icons

diff --git a/Engine/Source/Editor/EditorStyle/Private/StarshipStyle.cpp b/Engine/Source/Editor/EditorStyle/Private/StarshipStyle.cpp
--- a/Engine/Source/Editor/EditorStyle/Private/StarshipStyle.cpp
+++ b/Engine/Source/Editor/EditorStyle/Private/StarshipStyle.cpp
@@ -5788,7 +5788,6 @@ void FStarshipEditorStyle::FStyle::SetupClassIconsAndThumbnails()
 		Set("ClassIcon.AudioComponent", new IMAGE_BRUSH_SVG("Starship/AssetIcons/Audio_16", Icon16x16));
 		Set("ClassIcon.BoxComponent", new IMAGE_BRUSH_SVG("Starship/AssetIcons/BoxCollision_16", Icon16x16));
 		Set("ClassIcon.CapsuleComponent", new IMAGE_BRUSH_SVG("Starship/AssetIcons/CapsuleCollision_16", Icon16x16));
-		Set("ClassIcon.InstancedStaticMeshComponent", new IMAGE_BRUSH_SVG("Icons/ActorIcons/InstancedStaticMesh_16px", Icon16x16));
 		Set("ClassIcon.SphereComponent", new IMAGE_BRUSH_SVG("Starship/AssetIcons/SphereCollision_16", Icon16x16));
 		Set("ClassIcon.SplineComponent", new IMAGE_BRUSH_SVG("Starship/AssetIcons/Spline_16", Icon16x16));
 
@@ -5929,8 +5928,10 @@ void FStarshipEditorStyle::FStyle::SetupClassIconsAndThumbnails()
 			{TEXT("GameStateBase")},
 			{TEXT("GeometryCollection")},
 			{TEXT("GroupActor")},
+			{TEXT("HierarchicalInstancedStaticMeshComponent")},
 			{TEXT("HUD")},
 			{TEXT("ImagePlate")},
+			{TEXT("InstancedStaticMeshComponent")},
 			{TEXT("Interface")},
 			{TEXT("KillZVolume")},
 			{TEXT("Landscape")},
