Matinee/Styles: Deleted Matinee style entries from SlateEditorStyle and StarshipStyle.

diff --git a/Engine/Source/Editor/EditorStyle/Private/SlateEditorStyle.cpp b/Engine/Source/Editor/EditorStyle/Private/SlateEditorStyle.cpp
--- a/Engine/Source/Editor/EditorStyle/Private/SlateEditorStyle.cpp
+++ b/Engine/Source/Editor/EditorStyle/Private/SlateEditorStyle.cpp
@@ -196,7 +196,6 @@ void FSlateEditorStyle::FStyle::Initialize()
 	SetupToolkitStyles();
 	SetupTranslationEditorStyles();
 	SetupLocalizationDashboardStyles();
-	SetupMatineeStyle();
 	SetupSourceControlStyles();
 	SetupAutomationStyles();
 	SetupUMGEditorStyles();
@@ -1464,12 +1463,6 @@ void FSlateEditorStyle::FStyle::SetupGeneralStyles()
 		Set( "SocketIcon.Socket", new IMAGE_BRUSH( "Old/socket", Icon16x16 ) );
 		Set( "SocketIcon.None", new IMAGE_BRUSH( "Old/Favorites_Disabled", Icon16x16 ) );
 	}
-	
-	// Matinee Recorder
-	{
-		Set( "MatineeRecorder.Record", new IMAGE_BRUSH( "Icons/Record_16x", Icon16x16 ) );
-		Set( "MatineeRecorder.Stop", new IMAGE_BRUSH( "Icons/Stop_16x", Icon16x16 ) );
-	}
 
 	// Graph breadcrumb button
 	{
@@ -4965,8 +4958,6 @@ void FSlateEditorStyle::FStyle::SetupLevelEditorStyle()
 		Set( "LevelEditor.OpenLevelBlueprint.Small", new IMAGE_BRUSH( "Icons/icon_kismet2_40x", Icon20x20 ) );
 		Set( "LevelEditor.CreateClassBlueprint", new IMAGE_BRUSH("Icons/icon_class_Blueprint_New_16x", Icon16x16));
 		Set( "LevelEditor.OpenClassBlueprint", new IMAGE_BRUSH("Icons/icon_class_Blueprint_Open_16x", Icon16x16));
-		Set( "LevelEditor.EditMatinee", new IMAGE_BRUSH( "Icons/icon_matinee_40x", Icon40x40 ) );
-		Set( "LevelEditor.EditMatinee.Small", new IMAGE_BRUSH( "Icons/icon_matinee_40x", Icon20x20 ) );
 
 		Set( "LevelEditor.ToggleVR", new IMAGE_BRUSH( "Icons/VREditor/VR_Editor_Toolbar_Icon", Icon40x40 ) );
 		Set( "LevelEditor.ToggleVR.Small", new IMAGE_BRUSH( "Icons/VREditor/VR_Editor_Toolbar_Icon_Small", Icon20x20 ) );
@@ -6497,7 +6488,6 @@ void FSlateEditorStyle::FStyle::SetupClassIconsAndThumbnails()
 			TEXT("MaterialInstanceActor"),
 			TEXT("MaterialInstanceConstant"),
 			TEXT("MaterialParameterCollection"),
-			TEXT("MatineeActor"),
 			TEXT("MediaPlayer"),
 			TEXT("MediaTexture"),
 			TEXT("LevelSequenceActor"),
@@ -7768,64 +7758,6 @@ void FSlateEditorStyle::FStyle::SetupToolkitStyles()
 #endif // WITH_EDITOR
 }
 
-void FSlateEditorStyle::FStyle::SetupMatineeStyle()
-{
-	//Matinee
-#if WITH_EDITOR
-	{
-		Set( "Matinee.Tabs.RecordingViewport", new IMAGE_BRUSH( "/Icons/icon_Matinee_RecordingViewport_16x", Icon16x16 ) );
-		Set( "Matinee.Tabs.CurveEditor", new IMAGE_BRUSH( "/Icons/icon_Matinee_Curve_Editor_16x", Icon16x16 ) );
-		Set( "Matinee.Tabs.Tracks", new IMAGE_BRUSH( "/Icons/icon_Matinee_Tracks_16x", Icon16x16 ) );
-
-		Set("Matinee.Filters.Text", FTextBlockStyle(NormalText)
-			.SetFont(DEFAULT_FONT("Bold", 9))
-			.SetColorAndOpacity(FLinearColor(1.0f, 1.0f, 1.0f, 0.9f))
-			.SetShadowOffset(FVector2D(1, 1))
-			.SetShadowColorAndOpacity(FLinearColor(0, 0, 0, 0.9f)));
-
-		Set( "Matinee.AddKey", new IMAGE_BRUSH( "Icons/icon_Matinee_AddKey_40x", Icon40x40 ) );
-		Set( "Matinee.CreateMovie", new IMAGE_BRUSH( "Icons/icon_Matinee_CreateMovie_40x", Icon40x40 ) );
-		Set( "Matinee.Play", new IMAGE_BRUSH( "Icons/icon_Matinee_Play_40x", Icon40x40 ) );
-		Set( "Matinee.PlayLoop", new IMAGE_BRUSH( "Icons/icon_Matinee_PlayLoopSection_40x", Icon40x40 ) );
-		Set( "Matinee.Stop", new IMAGE_BRUSH( "Icons/icon_Matinee_Stop_40x", Icon40x40 ) );
-		Set( "Matinee.PlayReverse", new IMAGE_BRUSH( "Icons/icon_Matinee_PlayReverse_40x", Icon40x40 ) );
-		Set( "Matinee.ToggleSnap", new IMAGE_BRUSH( "Icons/icon_Matinee_ToggleSnap_40x", Icon40x40 ) );
-		Set( "Matinee.FitSequence", new IMAGE_BRUSH( "Icons/icon_Matinee_FitSequence_40x", Icon40x40 ) );
-		Set( "Matinee.FitViewToSelected", new IMAGE_BRUSH( "Icons/icon_Matinee_FitSelected_40x", Icon40x40 ) );
-		Set( "Matinee.FitLoop", new IMAGE_BRUSH( "Icons/icon_Matinee_FitLoop_40x", Icon40x40 ) );
-		Set( "Matinee.FitLoopSequence", new IMAGE_BRUSH( "Icons/icon_Matinee_FitLoopSequnce_40x", Icon40x40 ) );
-		Set( "Matinee.ViewEndofTrack", new IMAGE_BRUSH( "Icons/icon_Matinee_EndOfTrack_40x", Icon40x40 ) );
-		Set( "Matinee.ToggleSnapTimeToFrames", new IMAGE_BRUSH( "Icons/icon_Matinee_SnapTimeToFrames_40x", Icon40x40 ) );
-		Set( "Matinee.FixedTimeStepPlayback", new IMAGE_BRUSH( "Icons/icon_Matinee_FixedTimeStepPlayback_40x", Icon40x40 ) );
-		Set( "Matinee.ToggleGorePreview", new IMAGE_BRUSH( "Icons/icon_Matinee_GorePreview_40x", Icon40x40 ) );
-		Set( "Matinee.CreateCameraActor", new IMAGE_BRUSH( "Icons/icon_Matinee_CreateCameraActor_40x", Icon40x40 ) );
-		Set( "Matinee.LaunchRecordWindow", new IMAGE_BRUSH( "Icons/icon_Matinee_LaunchRecorderWindow_40x", Icon40x40 ) );
-		Set( "Matinee.ToggleCurveEditor", new IMAGE_BRUSH("Icons/icon_MatineeCurveView_40px", Icon40x40) );
-		Set( "Matinee.ToggleDirectorTimeline", new IMAGE_BRUSH("Icons/icon_MatineeDirectorView_40px", Icon40x40) );
-
-		Set( "Matinee.AddKey.Small", new IMAGE_BRUSH( "Icons/icon_Matinee_AddKey_40x", Icon20x20 ) );
-		Set( "Matinee.CreateMovie.Small", new IMAGE_BRUSH( "Icons/icon_Matinee_CreateMovie_40x", Icon20x20 ) );
-		Set( "Matinee.Play.Small", new IMAGE_BRUSH( "Icons/icon_Matinee_Play_40x", Icon20x20 ) );
-		Set( "Matinee.PlayLoop.Small", new IMAGE_BRUSH( "Icons/icon_Matinee_PlayLoopSection_40x", Icon20x20 ) );
-		Set( "Matinee.Stop.Small", new IMAGE_BRUSH( "Icons/icon_Matinee_Stop_40x", Icon20x20 ) );
-		Set( "Matinee.PlayReverse.Small", new IMAGE_BRUSH( "Icons/icon_Matinee_PlayReverse_40x", Icon20x20 ) );
-		Set( "Matinee.ToggleSnap.Small", new IMAGE_BRUSH( "Icons/icon_Matinee_ToggleSnap_40x", Icon20x20 ) );
-		Set( "Matinee.FitSequence.Small", new IMAGE_BRUSH( "Icons/icon_Matinee_FitSequence_40x", Icon20x20 ) );
-		Set( "Matinee.FitViewToSelected.Small", new IMAGE_BRUSH( "Icons/icon_Matinee_FitSelected_40x", Icon20x20 ) );
-		Set( "Matinee.FitLoop.Small", new IMAGE_BRUSH( "Icons/icon_Matinee_FitLoop_40x", Icon20x20 ) );
-		Set( "Matinee.FitLoopSequence.Small", new IMAGE_BRUSH( "Icons/icon_Matinee_FitLoopSequnce_40x", Icon20x20 ) );
-		Set( "Matinee.ViewEndofTrack.Small", new IMAGE_BRUSH( "Icons/icon_Matinee_EndOfTrack_40x", Icon20x20 ) );
-		Set( "Matinee.ToggleSnapTimeToFrames.Small", new IMAGE_BRUSH( "Icons/icon_Matinee_SnapTimeToFrames_40x", Icon20x20 ) );
-		Set( "Matinee.FixedTimeStepPlayback.Small", new IMAGE_BRUSH( "Icons/icon_Matinee_FixedTimeStepPlayback_40x", Icon20x20 ) );
-		Set( "Matinee.ToggleGorePreview.Small", new IMAGE_BRUSH( "Icons/icon_Matinee_GorePreview_40x", Icon20x20 ) );
-		Set( "Matinee.CreateCameraActor.Small", new IMAGE_BRUSH( "Icons/icon_Matinee_CreateCameraActor_40x", Icon20x20 ) );
-		Set( "Matinee.LaunchRecordWindow.Small", new IMAGE_BRUSH( "Icons/icon_Matinee_LaunchRecorderWindow_40x", Icon20x20 ) );
-		Set( "Matinee.ToggleCurveEditor.Small", new IMAGE_BRUSH("Icons/icon_MatineeCurveView_40px", Icon20x20) );
-		Set( "Matinee.ToggleDirectorTimeline.Small", new IMAGE_BRUSH("Icons/icon_MatineeDirectorView_40px", Icon20x20) );
-	}
-#endif // WITH_EDITOR
-	}
-
 void FSlateEditorStyle::FStyle::SetupSourceControlStyles()
 {
 	//Source Control
diff --git a/Engine/Source/Editor/EditorStyle/Private/SlateEditorStyle.h b/Engine/Source/Editor/EditorStyle/Private/SlateEditorStyle.h
--- a/Engine/Source/Editor/EditorStyle/Private/SlateEditorStyle.h
+++ b/Engine/Source/Editor/EditorStyle/Private/SlateEditorStyle.h
@@ -94,7 +94,6 @@ public:
 		void SetupContentBrowserStyle();
 		void SetupLandscapeEditorStyle();
 		void SetupToolkitStyles();
-		void SetupMatineeStyle();
 		void SetupSourceControlStyles();
 		void SetupAutomationStyles();
 		void SetupUMGEditorStyles();
diff --git a/Engine/Source/Editor/EditorStyle/Private/StarshipStyle.cpp b/Engine/Source/Editor/EditorStyle/Private/StarshipStyle.cpp
--- a/Engine/Source/Editor/EditorStyle/Private/StarshipStyle.cpp
+++ b/Engine/Source/Editor/EditorStyle/Private/StarshipStyle.cpp
@@ -282,7 +282,6 @@ void FStarshipEditorStyle::FStyle::Initialize()
 	SetupToolkitStyles();
 	SetupTranslationEditorStyles();
 	SetupLocalizationDashboardStyles();
-	SetupMatineeStyle();
 	SetupSourceControlStyles();
 	SetupAutomationStyles();
 	SetupUMGEditorStyles();
@@ -1178,12 +1177,6 @@ void FStarshipEditorStyle::FStyle::SetupGeneralStyles()
 		Set( "SocketIcon.Socket", new IMAGE_BRUSH( "Old/socket", Icon16x16 ) );
 		Set( "SocketIcon.None", new IMAGE_BRUSH( "Old/Favorites_Disabled", Icon16x16 ) );
 	}
-	
-	// Matinee Recorder
-	{
-		Set( "MatineeRecorder.Record", new IMAGE_BRUSH( "Icons/Record_16x", Icon16x16 ) );
-		Set( "MatineeRecorder.Stop", new IMAGE_BRUSH( "Icons/Stop_16x", Icon16x16 ) );
-	}
 
 	// Graph breadcrumb button
 	{
@@ -4526,10 +4519,10 @@ void FStarshipEditorStyle::FStyle::SetupLevelEditorStyle()
 		Set( "LevelEditor.OpenContentBrowser", new IMAGE_BRUSH_SVG( "Starship/MainToolbar/content", Icon20x20 ) );
 		Set( "LevelEditor.OpenMarketplace", new IMAGE_BRUSH_SVG( "Starship/MainToolbar/marketplace", Icon20x20) );
 		Set( "LevelEditor.OpenLevelBlueprint", new IMAGE_BRUSH_SVG( "Starship/MainToolbar/blueprints", Icon20x20) );
+		Set( "LevelEditor.OpenCinematic", new IMAGE_BRUSH_SVG("Starship/MainToolbar/cinematics", Icon20x20));
 
 		Set( "LevelEditor.CreateClassBlueprint", new IMAGE_BRUSH("Icons/icon_class_Blueprint_New_16x", Icon16x16));
 		Set( "LevelEditor.OpenClassBlueprint", new IMAGE_BRUSH("Icons/icon_class_Blueprint_Open_16x", Icon16x16));
-		Set( "LevelEditor.EditMatinee", new IMAGE_BRUSH_SVG( "Starship/MainToolbar/cinematics", Icon20x20 ) );
 
 		Set( "LevelEditor.ToggleVR", new IMAGE_BRUSH( "Icons/VREditor/VR_Editor_Toolbar_Icon", Icon40x40 ) );
 		Set( "LevelEditor.ToggleVR.Small", new IMAGE_BRUSH( "Icons/VREditor/VR_Editor_Toolbar_Icon_Small", Icon20x20 ) );
@@ -5842,7 +5835,6 @@ void FStarshipEditorStyle::FStyle::SetupClassIconsAndThumbnails()
 			TEXT("LightmassCharacterIndirectDetailVolume"),
 			TEXT("MassiveLODOverrideVolume"),
 			TEXT("MaterialParameterCollection"),
-			TEXT("MatineeActor"),
 			TEXT("MultiFont"),
 			TEXT("ParticleSystem"),
 			TEXT("PhysicsConstraintComponent"),
@@ -7187,64 +7179,6 @@ void FStarshipEditorStyle::FStyle::SetupToolkitStyles()
 #endif // WITH_EDITOR
 }
 
-void FStarshipEditorStyle::FStyle::SetupMatineeStyle()
-{
-	//Matinee
-#if WITH_EDITOR
-	{
-		Set( "Matinee.Tabs.RecordingViewport", new IMAGE_BRUSH( "/Icons/icon_Matinee_RecordingViewport_16x", Icon16x16 ) );
-		Set( "Matinee.Tabs.CurveEditor", new IMAGE_BRUSH( "/Icons/icon_Matinee_Curve_Editor_16x", Icon16x16 ) );
-		Set( "Matinee.Tabs.Tracks", new IMAGE_BRUSH( "/Icons/icon_Matinee_Tracks_16x", Icon16x16 ) );
-
-		Set("Matinee.Filters.Text", FTextBlockStyle(NormalText)
-			.SetFont(DEFAULT_FONT("Bold", 9))
-			.SetColorAndOpacity(FLinearColor(1.0f, 1.0f, 1.0f, 0.9f))
-			.SetShadowOffset(FVector2D(1, 1))
-			.SetShadowColorAndOpacity(FLinearColor(0, 0, 0, 0.9f)));
-
-		Set( "Matinee.AddKey", new IMAGE_BRUSH( "Icons/icon_Matinee_AddKey_40x", Icon40x40 ) );
-		Set( "Matinee.CreateMovie", new IMAGE_BRUSH( "Icons/icon_Matinee_CreateMovie_40x", Icon40x40 ) );
-		Set( "Matinee.Play", new IMAGE_BRUSH( "Icons/icon_Matinee_Play_40x", Icon40x40 ) );
-		Set( "Matinee.PlayLoop", new IMAGE_BRUSH( "Icons/icon_Matinee_PlayLoopSection_40x", Icon40x40 ) );
-		Set( "Matinee.Stop", new IMAGE_BRUSH( "Icons/icon_Matinee_Stop_40x", Icon40x40 ) );
-		Set( "Matinee.PlayReverse", new IMAGE_BRUSH( "Icons/icon_Matinee_PlayReverse_40x", Icon40x40 ) );
-		Set( "Matinee.ToggleSnap", new IMAGE_BRUSH( "Icons/icon_Matinee_ToggleSnap_40x", Icon40x40 ) );
-		Set( "Matinee.FitSequence", new IMAGE_BRUSH( "Icons/icon_Matinee_FitSequence_40x", Icon40x40 ) );
-		Set( "Matinee.FitViewToSelected", new IMAGE_BRUSH( "Icons/icon_Matinee_FitSelected_40x", Icon40x40 ) );
-		Set( "Matinee.FitLoop", new IMAGE_BRUSH( "Icons/icon_Matinee_FitLoop_40x", Icon40x40 ) );
-		Set( "Matinee.FitLoopSequence", new IMAGE_BRUSH( "Icons/icon_Matinee_FitLoopSequnce_40x", Icon40x40 ) );
-		Set( "Matinee.ViewEndofTrack", new IMAGE_BRUSH( "Icons/icon_Matinee_EndOfTrack_40x", Icon40x40 ) );
-		Set( "Matinee.ToggleSnapTimeToFrames", new IMAGE_BRUSH( "Icons/icon_Matinee_SnapTimeToFrames_40x", Icon40x40 ) );
-		Set( "Matinee.FixedTimeStepPlayback", new IMAGE_BRUSH( "Icons/icon_Matinee_FixedTimeStepPlayback_40x", Icon40x40 ) );
-		Set( "Matinee.ToggleGorePreview", new IMAGE_BRUSH( "Icons/icon_Matinee_GorePreview_40x", Icon40x40 ) );
-		Set( "Matinee.CreateCameraActor", new IMAGE_BRUSH( "Icons/icon_Matinee_CreateCameraActor_40x", Icon40x40 ) );
-		Set( "Matinee.LaunchRecordWindow", new IMAGE_BRUSH( "Icons/icon_Matinee_LaunchRecorderWindow_40x", Icon40x40 ) );
-		Set( "Matinee.ToggleCurveEditor", new IMAGE_BRUSH("Icons/icon_MatineeCurveView_40px", Icon40x40) );
-		Set( "Matinee.ToggleDirectorTimeline", new IMAGE_BRUSH("Icons/icon_MatineeDirectorView_40px", Icon40x40) );
-
-		Set( "Matinee.AddKey.Small", new IMAGE_BRUSH( "Icons/icon_Matinee_AddKey_40x", Icon20x20 ) );
-		Set( "Matinee.CreateMovie.Small", new IMAGE_BRUSH( "Icons/icon_Matinee_CreateMovie_40x", Icon20x20 ) );
-		Set( "Matinee.Play.Small", new IMAGE_BRUSH( "Icons/icon_Matinee_Play_40x", Icon20x20 ) );
-		Set( "Matinee.PlayLoop.Small", new IMAGE_BRUSH( "Icons/icon_Matinee_PlayLoopSection_40x", Icon20x20 ) );
-		Set( "Matinee.Stop.Small", new IMAGE_BRUSH( "Icons/icon_Matinee_Stop_40x", Icon20x20 ) );
-		Set( "Matinee.PlayReverse.Small", new IMAGE_BRUSH( "Icons/icon_Matinee_PlayReverse_40x", Icon20x20 ) );
-		Set( "Matinee.ToggleSnap.Small", new IMAGE_BRUSH( "Icons/icon_Matinee_ToggleSnap_40x", Icon20x20 ) );
-		Set( "Matinee.FitSequence.Small", new IMAGE_BRUSH( "Icons/icon_Matinee_FitSequence_40x", Icon20x20 ) );
-		Set( "Matinee.FitViewToSelected.Small", new IMAGE_BRUSH( "Icons/icon_Matinee_FitSelected_40x", Icon20x20 ) );
-		Set( "Matinee.FitLoop.Small", new IMAGE_BRUSH( "Icons/icon_Matinee_FitLoop_40x", Icon20x20 ) );
-		Set( "Matinee.FitLoopSequence.Small", new IMAGE_BRUSH( "Icons/icon_Matinee_FitLoopSequnce_40x", Icon20x20 ) );
-		Set( "Matinee.ViewEndofTrack.Small", new IMAGE_BRUSH( "Icons/icon_Matinee_EndOfTrack_40x", Icon20x20 ) );
-		Set( "Matinee.ToggleSnapTimeToFrames.Small", new IMAGE_BRUSH( "Icons/icon_Matinee_SnapTimeToFrames_40x", Icon20x20 ) );
-		Set( "Matinee.FixedTimeStepPlayback.Small", new IMAGE_BRUSH( "Icons/icon_Matinee_FixedTimeStepPlayback_40x", Icon20x20 ) );
-		Set( "Matinee.ToggleGorePreview.Small", new IMAGE_BRUSH( "Icons/icon_Matinee_GorePreview_40x", Icon20x20 ) );
-		Set( "Matinee.CreateCameraActor.Small", new IMAGE_BRUSH( "Icons/icon_Matinee_CreateCameraActor_40x", Icon20x20 ) );
-		Set( "Matinee.LaunchRecordWindow.Small", new IMAGE_BRUSH( "Icons/icon_Matinee_LaunchRecorderWindow_40x", Icon20x20 ) );
-		Set( "Matinee.ToggleCurveEditor.Small", new IMAGE_BRUSH("Icons/icon_MatineeCurveView_40px", Icon20x20) );
-		Set( "Matinee.ToggleDirectorTimeline.Small", new IMAGE_BRUSH("Icons/icon_MatineeDirectorView_40px", Icon20x20) );
-	}
-#endif // WITH_EDITOR
-	}
-
 void FStarshipEditorStyle::FStyle::SetupSourceControlStyles()
 {
 	//Source Control
diff --git a/Engine/Source/Editor/EditorStyle/Private/StarshipStyle.h b/Engine/Source/Editor/EditorStyle/Private/StarshipStyle.h
--- a/Engine/Source/Editor/EditorStyle/Private/StarshipStyle.h
+++ b/Engine/Source/Editor/EditorStyle/Private/StarshipStyle.h
@@ -66,7 +66,6 @@ public:
 		void SetupContentBrowserStyle();
 		void SetupLandscapeEditorStyle();
 		void SetupToolkitStyles();
-		void SetupMatineeStyle();
 		void SetupSourceControlStyles();
 		void SetupAutomationStyles();
 		void SetupUMGEditorStyles();
diff --git a/Engine/Source/Editor/LevelEditor/Private/LevelEditorToolBar.cpp b/Engine/Source/Editor/LevelEditor/Private/LevelEditorToolBar.cpp
--- a/Engine/Source/Editor/LevelEditor/Private/LevelEditorToolBar.cpp
+++ b/Engine/Source/Editor/LevelEditor/Private/LevelEditorToolBar.cpp
@@ -1235,7 +1235,7 @@ void FLevelEditorToolBar::RegisterLevelEditorToolBar( const TSharedRef<FUIComman
 				FOnGetContent::CreateStatic(&FLevelEditorToolBar::GenerateCinematicsMenuContent, InCommandList, TWeakPtr<SLevelEditor>(InLevelEditor)),
 				LOCTEXT("EditCinematics_Label", "Cinematics"),
 				LOCTEXT("EditCinematics_Tooltip", "Displays a list of Level Sequence objects to open in their respective editors"),
-				FSlateIcon(FEditorStyle::GetStyleSetName(), "LevelEditor.EditMatinee")
+				FSlateIcon(FAppStyle::Get().GetStyleSetName(), "LevelEditor.OpenCinematic")
 			);
 			CinematicsEntry.StyleNameOverride = "CalloutToolbar";
 			Section.AddEntry(CinematicsEntry);
