Updated Graph Editor icons and removed unused Graph Editor brushes from StarshipStyle

diff --git a/Engine/Source/Editor/EditorStyle/Private/StarshipStyle.cpp b/Engine/Source/Editor/EditorStyle/Private/StarshipStyle.cpp
--- a/Engine/Source/Editor/EditorStyle/Private/StarshipStyle.cpp
+++ b/Engine/Source/Editor/EditorStyle/Private/StarshipStyle.cpp
@@ -3991,13 +3991,13 @@ void FStarshipEditorStyle::FStyle::SetupGraphEditorStyles()
 
 		Set("ClassIcon.K2Node_CallFunction",	new IMAGE_BRUSH_SVG("Starship/Blueprints/icon_Blueprint_Function", Icon16x16));
 		Set("ClassIcon.K2Node_FunctionEntry",	new IMAGE_BRUSH_SVG("Starship/Blueprints/icon_Blueprint_Function", Icon16x16));
-		Set("ClassIcon.K2Node_CustomEvent",		new IMAGE_BRUSH("Icons/icon_Blueprint_Event_16x", Icon16x16));
-		Set("ClassIcon.K2Node_Event",			new IMAGE_BRUSH("Icons/icon_Blueprint_Event_16x", Icon16x16));
-		Set("ClassIcon.K2Node_Variable",		new IMAGE_BRUSH("Graph/Icons/Node", Icon16x16, FLinearColor::White));
-		Set("ClassIcon.K2Node_VariableGet",		new IMAGE_BRUSH("Graph/Icons/FIB_VarGet", Icon16x16, FLinearColor::White));
-		Set("ClassIcon.K2Node_VariableSet",		new IMAGE_BRUSH("Graph/Icons/FIB_VarSet", Icon16x16, FLinearColor::White));
-		Set("ClassIcon.K2Node_DynamicCast",		new IMAGE_BRUSH("Icons/icon_Blueprint_Cast_16x", Icon16x16));
-		Set("ClassIcon.EdGraphNode_Comment",	new IMAGE_BRUSH("Icons/icon_Blueprint_Comment_16x", Icon16x16));
+		Set("ClassIcon.K2Node_CustomEvent",		new IMAGE_BRUSH_SVG("Starship/GraphEditors/Event", Icon16x16));
+		Set("ClassIcon.K2Node_Event",			new IMAGE_BRUSH_SVG("Starship/GraphEditors/Event", Icon16x16));
+		Set("ClassIcon.K2Node_Variable",		new IMAGE_BRUSH_SVG("Starship/GraphEditors/Node", Icon16x16, FLinearColor::White));
+		Set("ClassIcon.K2Node_VariableGet",		new IMAGE_BRUSH_SVG("Starship/GraphEditors/VarGet", Icon16x16, FLinearColor::White));
+		Set("ClassIcon.K2Node_VariableSet",		new IMAGE_BRUSH_SVG("Starship/GraphEditors/VarSet", Icon16x16, FLinearColor::White));
+		Set("ClassIcon.K2Node_DynamicCast",		new IMAGE_BRUSH_SVG("Starship/GraphEditors/Cast", Icon16x16));
+		Set("ClassIcon.EdGraphNode_Comment",	new IMAGE_BRUSH_SVG("Starship/GraphEditors/Comment", Icon16x16));
 
 		Set("GraphEditor.Clean",				new IMAGE_BRUSH_SVG("Starship/GraphEditors/CleanUp", Icon20x20));
 		Set("GraphEditor.OverrideFunction_16x", new IMAGE_BRUSH_SVG("Starship/Blueprints/icon_Blueprint_OverrideFunction", Icon16x16));
@@ -4013,106 +4013,82 @@ void FStarshipEditorStyle::FStyle::SetupGraphEditorStyles()
 
 
 
-		Set( "GraphEditor.Default_16x", new IMAGE_BRUSH("Icons/icon_Blueprint_Node_16x", Icon16x16));
-		Set( "GraphEditor.InterfaceFunction_16x", new IMAGE_BRUSH( "Icons/icon_Blueprint_Interfacefunction_16x", Icon16x16));
-		Set( "GraphEditor.PureFunction_16x", new IMAGE_BRUSH( "Icons/icon_Blueprint_NewPureFunction_16x", Icon16x16 ) );
-		Set( "GraphEditor.PotentialOverrideFunction_16x", new IMAGE_BRUSH( "Icons/icon_Blueprint_OverrideableFunction_16x", Icon16x16 ) );
-		Set( "GraphEditor.OverridePureFunction_16x", new IMAGE_BRUSH( "Icons/icon_Blueprint_OverridePureFunction_16x", Icon16x16 ) );
-		Set( "GraphEditor.SubGraph_16x", new IMAGE_BRUSH( "Icons/icon_Blueprint_SubgraphComposite_16x", Icon16x16 ) );
-		Set( "GraphEditor.Animation_16x", new IMAGE_BRUSH( "Icons/icon_Blueprint_Anim_16x", Icon16x16 ) );
-		Set( "GraphEditor.Conduit_16x", new IMAGE_BRUSH( "Icons/icon_Blueprint_Conduit_16x", Icon16x16 ) );
-		Set( "GraphEditor.Rule_16x", new IMAGE_BRUSH( "Icons/icon_Blueprint_Rule_16x", Icon16x16 ) );
-		Set( "GraphEditor.State_16x", new IMAGE_BRUSH( "Icons/icon_Blueprint_State_16x", Icon16x16 ) );
-		Set( "GraphEditor.StateMachine_16x", new IMAGE_BRUSH( "Icons/icon_Blueprint_StateMachine_16x", Icon16x16 ) );
-		Set( "GraphEditor.Event_16x", new IMAGE_BRUSH( "Icons/icon_Blueprint_Event_16x", Icon16x16 ) );
-		Set( "GraphEditor.CustomEvent_16x", new IMAGE_BRUSH( "Icons/icon_Blueprint_CustomEvent_16x", Icon16x16 ) );
-		Set( "GraphEditor.CallInEditorEvent_16x", new IMAGE_BRUSH( "Icons/icon_Blueprint_CallInEditor_16x", Icon16x16 ) );
-		Set( "GraphEditor.Timeline_16x", new IMAGE_BRUSH("Icons/icon_Blueprint_Timeline_16x", Icon16x16));
-		Set( "GraphEditor.Comment_16x", new IMAGE_BRUSH("Icons/icon_Blueprint_Comment_16x", Icon16x16));
-		Set( "GraphEditor.Documentation_16x", new IMAGE_BRUSH("Icons/icon_Blueprint_Documentation_16x", Icon16x16));
-		Set( "GraphEditor.Switch_16x", new IMAGE_BRUSH("Icons/icon_Blueprint_Switch_16x", Icon16x16));
-		Set( "GraphEditor.BreakStruct_16x", new IMAGE_BRUSH("Icons/icon_Blueprint_BreakStruct_16x", Icon16x16));
-		Set( "GraphEditor.MakeStruct_16x", new IMAGE_BRUSH("Icons/icon_Blueprint_MakeStruct_16x", Icon16x16));
-		Set( "GraphEditor.Sequence_16x", new IMAGE_BRUSH("Icons/icon_Blueprint_Sequence_16x", Icon16x16));
-		Set( "GraphEditor.Branch_16x", new IMAGE_BRUSH("Icons/icon_Blueprint_Branch_16x", Icon16x16));
-		Set( "GraphEditor.SpawnActor_16x", new IMAGE_BRUSH("Icons/icon_Blueprint_SpawnActor_16x", Icon16x16));
-		Set( "GraphEditor.PadEvent_16x", new IMAGE_BRUSH("Icons/icon_Blueprint_PadEvent_16x", Icon16x16));
-		Set( "GraphEditor.MouseEvent_16x", new IMAGE_BRUSH("Icons/icon_Blueprint_MouseEvent_16x", Icon16x16));
-		Set( "GraphEditor.KeyEvent_16x", new IMAGE_BRUSH("Icons/icon_Blueprint_KeyboardEvent_16x", Icon16x16));
-		Set( "GraphEditor.TouchEvent_16x", new IMAGE_BRUSH("Icons/icon_Blueprint_TouchEvent_16x", Icon16x16));
-		Set( "GraphEditor.MakeArray_16x", new IMAGE_BRUSH("Icons/icon_Blueprint_MakeArray_16x", Icon16x16));
-		Set( "GraphEditor.MakeSet_16x", new IMAGE_BRUSH("Icons/icon_Blueprint_MakeSet_16x", Icon16x16));
-		Set( "GraphEditor.MakeMap_16x", new IMAGE_BRUSH("Icons/icon_Blueprint_MakeMap_16x", Icon16x16));
-		Set( "GraphEditor.Enum_16x", new IMAGE_BRUSH("Icons/icon_Blueprint_Enum_16x", Icon16x16));
-		Set( "GraphEditor.Select_16x", new IMAGE_BRUSH("Icons/icon_Blueprint_Select_16x", Icon16x16));
-		Set( "GraphEditor.Cast_16x", new IMAGE_BRUSH("Icons/icon_Blueprint_Cast_16x", Icon16x16));
-
-		Set( "GraphEditor.Macro.Loop_16x", new IMAGE_BRUSH("Icons/icon_Blueprint_Loop_16x", Icon16x16));
-		Set( "GraphEditor.Macro.Gate_16x", new IMAGE_BRUSH("Icons/icon_Blueprint_Gate_16x", Icon16x16));
-		Set( "GraphEditor.Macro.DoN_16x", new IMAGE_BRUSH("Icons/icon_Blueprint_DoN_16x", Icon16x16));
-		Set( "GraphEditor.Macro.DoOnce_16x", new IMAGE_BRUSH("Icons/icon_Blueprint_DoOnce_16x", Icon16x16));
-		Set( "GraphEditor.Macro.IsValid_16x", new IMAGE_BRUSH("Icons/icon_Blueprint_IsValid_16x", Icon16x16));
-		Set( "GraphEditor.Macro.FlipFlop_16x", new IMAGE_BRUSH("Icons/icon_Blueprint_FlipFlop_16x", Icon16x16));
-		Set( "GraphEditor.Macro.ForEach_16x", new IMAGE_BRUSH("Icons/icon_Blueprint_ForEach_16x", Icon16x16));
-
-
-
-		Set( "GraphEditor.InterfaceFunction_24x", new IMAGE_BRUSH( "Icons/icon_Blueprint_InterfaceFunction_24x", Icon24x24 ) );
-
-		Set( "GraphEditor.PotentialOverrideFunction_24x", new IMAGE_BRUSH( "Icons/icon_Blueprint_OverrideableFunction_24x", Icon24x24 ) );
-
-		Set( "GraphEditor.SubGraph_24x", new IMAGE_BRUSH( "Icons/icon_Blueprint_SubgraphComposite_24x", Icon24x24 ) );
-		Set( "GraphEditor.Animation_24x", new IMAGE_BRUSH( "Icons/icon_Blueprint_Anim_24x", Icon24x24 ) );
-		Set( "GraphEditor.Conduit_24x", new IMAGE_BRUSH( "Icons/icon_Blueprint_Conduit_24x", Icon24x24 ) );
-		Set( "GraphEditor.Rule_24x", new IMAGE_BRUSH( "Icons/icon_Blueprint_Rule_24x", Icon24x24 ) );
-		Set( "GraphEditor.State_24x", new IMAGE_BRUSH( "Icons/icon_Blueprint_State_24x", Icon24x24 ) );
-		Set( "GraphEditor.StateMachine_24x", new IMAGE_BRUSH( "Icons/icon_Blueprint_StateMachine_24x", Icon24x24 ) );
-
-		Set( "GraphEditor.FunctionGlyph", new IMAGE_BRUSH( "Graph/Icons/Function", Icon22x22, FLinearColor::White) );
-		Set( "GraphEditor.NodeGlyph", new IMAGE_BRUSH( "Graph/Icons/Node", Icon22x22, FLinearColor::White) );
-		Set( "GraphEditor.PinIcon", new IMAGE_BRUSH( "Graph/Icons/Pin", Icon22x22, FLinearColor::White) );
-		Set( "GraphEditor.ArrayPinIcon", new IMAGE_BRUSH( "Graph/Icons/ArrayPin", Icon22x22, FLinearColor::White ) );
-		Set( "GraphEditor.RefPinIcon", new IMAGE_BRUSH( "Graph/Icons/RefPin", Icon22x22, FLinearColor::White ) );
-		Set( "GraphEditor.UbergraphGlyph", new IMAGE_BRUSH( "Graph/Icons/EventGraph", Icon22x22, FLinearColor::White) );		
-		Set( "GraphEditor.SubgraphGlyph", new IMAGE_BRUSH( "Graph/Icons/Subgraph", Icon22x22, FLinearColor::White) );		
-		Set( "GraphEditor.AnimationGlyph", new IMAGE_BRUSH( "Graph/Icons/Robot", Icon22x22, FLinearColor::White) );
-		Set( "GraphEditor.MacroGlyph", new IMAGE_BRUSH( "Graph/Icons/Macro", Icon22x22, FLinearColor::White) );
-		Set( "GraphEditor.EnumGlyph", new IMAGE_BRUSH( "Graph/Icons/Enum", Icon22x22, FLinearColor::White) );
-		Set( "GraphEditor.TimelineGlyph", new IMAGE_BRUSH( "Graph/Icons/Timeline", Icon22x22, FLinearColor::White) );
-		Set( "GraphEditor.EventGlyph", new IMAGE_BRUSH( "Graph/Icons/Event", Icon22x22, FLinearColor::White) );
-		Set( "GraphEditor.EventCustomGlyph", new IMAGE_BRUSH( "Graph/Icons/Event_Custom", Icon22x22, FLinearColor::White) );
-		Set( "GraphEditor.SCSGlyph", new IMAGE_BRUSH( "Graph/Icons/Hammer", Icon22x22, FLinearColor::White) );
-		Set( "GraphEditor.StructGlyph", new IMAGE_BRUSH( "Graph/Icons/Struct", Icon22x22, FLinearColor::White) );
-		// Find In Blueprints
-		Set( "GraphEditor.FIB_CallFunction", new IMAGE_BRUSH( "Graph/Icons/FIB_CallFunction", Icon22x22, FLinearColor::White) );
-		Set( "GraphEditor.FIB_MacroInstance", new IMAGE_BRUSH( "Graph/Icons/FIB_MacroInstance", Icon22x22, FLinearColor::White) );
-		Set( "GraphEditor.FIB_Event", new IMAGE_BRUSH( "Graph/Icons/FIB_Event", Icon22x22, FLinearColor::White) );
-		Set( "GraphEditor.FIB_VariableGet", new IMAGE_BRUSH( "Graph/Icons/FIB_VarGet", Icon22x22, FLinearColor::White) );
-		Set( "GraphEditor.FIB_VariableSet", new IMAGE_BRUSH( "Graph/Icons/FIB_VarSet", Icon22x22, FLinearColor::White) );
-
-		Set( "GraphEditor.FunctionOL.Interface", new IMAGE_BRUSH( "Graph/Icons/Overlay_Interface", Icon22x22 ) );
-		Set( "GraphEditor.FunctionOL.New", new IMAGE_BRUSH( "Graph/Icons/Overlay_New", Icon22x22 ) );
-		Set( "GraphEditor.FunctionOL.Override", new IMAGE_BRUSH( "Graph/Icons/Overlay_Override", Icon22x22 ) );
-		Set( "GraphEditor.FunctionOL.PotentialOverride", new IMAGE_BRUSH( "Graph/Icons/Overlay_PotentialOverride", Icon22x22 ) );
+		Set( "GraphEditor.Default_16x", new IMAGE_BRUSH_SVG("Starship/GraphEditors/Node", Icon16x16));
+		Set( "GraphEditor.InterfaceFunction_16x", new IMAGE_BRUSH_SVG( "Starship/GraphEditors/InterfaceFunction", Icon16x16));
+		Set( "GraphEditor.PureFunction_16x", new IMAGE_BRUSH_SVG( "Starship/Blueprints/icon_Blueprint_Function", Icon16x16 ) );
+		Set( "GraphEditor.PotentialOverrideFunction_16x", new IMAGE_BRUSH_SVG( "Starship/Blueprints/icon_Blueprint_OverrideFunction", Icon16x16 ) );
+		Set( "GraphEditor.OverridePureFunction_16x", new IMAGE_BRUSH_SVG( "Starship/Blueprints/icon_Blueprint_OverrideFunction", Icon16x16 ) );
+		Set( "GraphEditor.SubGraph_16x", new IMAGE_BRUSH_SVG( "Starship/GraphEditors/SubGraph", Icon16x16 ) );
+		Set( "GraphEditor.Animation_16x", new IMAGE_BRUSH_SVG( "Starship/Common/Animation", Icon16x16 ) );
+		Set( "GraphEditor.Conduit_16x", new IMAGE_BRUSH_SVG( "Starship/GraphEditors/Conduit", Icon16x16 ) );
+		Set( "GraphEditor.Rule_16x", new IMAGE_BRUSH_SVG( "Starship/GraphEditors/Rule", Icon16x16 ) );
+		Set( "GraphEditor.State_16x", new IMAGE_BRUSH_SVG( "Starship/GraphEditors/State", Icon16x16 ) );
+		Set( "GraphEditor.StateMachine_16x", new IMAGE_BRUSH_SVG( "Starship/GraphEditors/StateMachine", Icon16x16 ) );
+		Set( "GraphEditor.Event_16x", new IMAGE_BRUSH_SVG( "Starship/GraphEditors/Event", Icon16x16 ) );
+		Set( "GraphEditor.CustomEvent_16x", new IMAGE_BRUSH_SVG( "Starship/GraphEditors/Event", Icon16x16 ) );
+		Set( "GraphEditor.CallInEditorEvent_16x", new IMAGE_BRUSH_SVG( "Starship/GraphEditors/CallInEditorEvent", Icon16x16 ) );
+		Set( "GraphEditor.Timeline_16x", new IMAGE_BRUSH_SVG("Starship/Common/Timecode", Icon16x16));
+		Set( "GraphEditor.Comment_16x", new IMAGE_BRUSH_SVG("Starship/GraphEditors/Comment", Icon16x16));
+		Set( "GraphEditor.Documentation_16x", new IMAGE_BRUSH_SVG("Starship/Common/Documentation", Icon16x16));
+		Set( "GraphEditor.Switch_16x", new IMAGE_BRUSH_SVG("Starship/GraphEditors/Switch", Icon16x16));
+		Set( "GraphEditor.BreakStruct_16x", new IMAGE_BRUSH_SVG("Starship/GraphEditors/BreakStruct", Icon16x16));
+		Set( "GraphEditor.MakeStruct_16x", new IMAGE_BRUSH_SVG("Starship/GraphEditors/MakeStruct", Icon16x16));
+		Set( "GraphEditor.Sequence_16x", new IMAGE_BRUSH_SVG("Starship/GraphEditors/Sequence", Icon16x16));
+		Set( "GraphEditor.Branch_16x", new IMAGE_BRUSH_SVG("Starship/GraphEditors/Branch", Icon16x16));
+		Set( "GraphEditor.SpawnActor_16x", new IMAGE_BRUSH_SVG("Starship/GraphEditors/SpawnActor", Icon16x16));
+		Set( "GraphEditor.PadEvent_16x", new IMAGE_BRUSH_SVG("Starship/Common/PlayerController", Icon16x16));
+		Set( "GraphEditor.MouseEvent_16x", new IMAGE_BRUSH_SVG("Starship/GraphEditors/MouseEvent", Icon16x16));
+		Set( "GraphEditor.KeyEvent_16x", new IMAGE_BRUSH_SVG("Starship/Common/ViewportControls", Icon16x16));
+		Set( "GraphEditor.TouchEvent_16x", new IMAGE_BRUSH_SVG("Starship/Common/TouchInterface", Icon16x16));
+		Set( "GraphEditor.MakeArray_16x", new IMAGE_BRUSH_SVG("Starship/GraphEditors/MakeArray", Icon16x16));
+		Set( "GraphEditor.MakeSet_16x", new IMAGE_BRUSH_SVG("Starship/GraphEditors/MakeSet", Icon16x16));
+		Set( "GraphEditor.MakeMap_16x", new IMAGE_BRUSH_SVG("Starship/GraphEditors/MakeMap", Icon16x16));
+		Set( "GraphEditor.Enum_16x", new IMAGE_BRUSH_SVG("Starship/GraphEditors/Enum", Icon16x16));
+		Set( "GraphEditor.Select_16x", new IMAGE_BRUSH_SVG("Starship/GraphEditors/Select", Icon16x16));
+		Set( "GraphEditor.Cast_16x", new IMAGE_BRUSH_SVG("Starship/GraphEditors/Cast", Icon16x16));
+
+		Set( "GraphEditor.Macro.Loop_16x", new IMAGE_BRUSH_SVG("Starship/GraphEditors/Loop", Icon16x16));
+		Set( "GraphEditor.Macro.Gate_16x", new IMAGE_BRUSH_SVG("Starship/GraphEditors/Gate", Icon16x16));
+		Set( "GraphEditor.Macro.DoN_16x", new IMAGE_BRUSH_SVG("Starship/GraphEditors/DoN", Icon16x16));
+		Set( "GraphEditor.Macro.DoOnce_16x", new IMAGE_BRUSH_SVG("Starship/GraphEditors/DoOnce", Icon16x16));
+		Set( "GraphEditor.Macro.IsValid_16x", new IMAGE_BRUSH_SVG("Starship/GraphEditors/IsValid", Icon16x16));
+		Set( "GraphEditor.Macro.FlipFlop_16x", new IMAGE_BRUSH_SVG("Starship/GraphEditors/FlipFlop", Icon16x16));
+		Set( "GraphEditor.Macro.ForEach_16x", new IMAGE_BRUSH_SVG("Starship/GraphEditors/ForEach", Icon16x16));
+
+		// The 24x icons are used for the SGraphTitleBar that shows the breadcrumbs for subgraphs
+		Set( "GraphEditor.SubGraph_24x", new IMAGE_BRUSH_SVG( "Starship/GraphEditors/SubGraph", Icon24x24 ) );
+		Set( "GraphEditor.Animation_24x", new IMAGE_BRUSH_SVG( "Starship/Common/Animation", Icon24x24 ) );
+		Set( "GraphEditor.Conduit_24x", new IMAGE_BRUSH_SVG( "Starship/GraphEditors/Conduit", Icon24x24 ) );
+		Set( "GraphEditor.Rule_24x", new IMAGE_BRUSH_SVG( "Starship/GraphEditors/Rule", Icon24x24 ) );
+		Set( "GraphEditor.State_24x", new IMAGE_BRUSH_SVG("Starship/GraphEditors/State", Icon24x24));
+		Set( "GraphEditor.StateMachine_24x", new IMAGE_BRUSH_SVG( "Starship/GraphEditors/StateMachine", Icon24x24 ) );
+
+		Set( "GraphEditor.NodeGlyph", new IMAGE_BRUSH_SVG( "Starship/GraphEditors/Node", Icon22x22, FLinearColor::White) );
+		Set( "GraphEditor.PinIcon", new IMAGE_BRUSH_SVG( "Starship/GraphEditors/PinIcon", Icon22x22, FLinearColor::White) );
+		Set( "GraphEditor.ArrayPinIcon", new IMAGE_BRUSH_SVG( "Starship/Blueprints/pillarray", Icon22x22, FLinearColor::White ) );
+		Set( "GraphEditor.RefPinIcon", new IMAGE_BRUSH_SVG( "Starship/GraphEditors/RefPin", Icon22x22, FLinearColor::White ) );
+		Set( "GraphEditor.EnumGlyph", new IMAGE_BRUSH_SVG( "Starship/GraphEditors/Enum", Icon22x22, FLinearColor::White) );
+		Set( "GraphEditor.TimelineGlyph", new IMAGE_BRUSH_SVG( "Starship/Common/Timecode", Icon22x22, FLinearColor::White) );
+		Set( "GraphEditor.StructGlyph", new IMAGE_BRUSH_SVG( "Starship/Common/Struct", Icon22x22, FLinearColor::White) );
 
-		Set( "GraphEditor.GetSequenceBinding", new IMAGE_BRUSH("Icons/icon_Blueprint_GetSequenceBinding_16x", Icon16x16));
+		// Find In Blueprints
+		Set( "GraphEditor.FIB_Event", new IMAGE_BRUSH_SVG( "Starship/GraphEditors/Event", Icon22x22, FLinearColor::White) );
 
-		Set( "GraphEditor.HideUnusedPins", new IMAGE_BRUSH( "Icons/hide_unusedpins", Icon40x40 ) );
-		Set( "GraphEditor.HideUnusedPins.Small", new IMAGE_BRUSH( "Icons/hide_unusedpins", Icon20x20 ) );
+		Set( "GraphEditor.GetSequenceBinding", new IMAGE_BRUSH_SVG("Starship/GraphEditors/GetSequenceBinding", Icon16x16));
 
-		Set( "GraphEditor.GoToDocumentation", new IMAGE_BRUSH( "Common/icon_Help_Hover_16x", Icon16x16 ) );
+		Set( "GraphEditor.GoToDocumentation", new IMAGE_BRUSH_SVG( "Starship/Common/Documentation", Icon16x16 ) );
 
-		Set( "GraphEditor.AlignNodesTop", new IMAGE_BRUSH( "Icons/GraphEditor/icon_AlignNodesTop_20px", Icon20x20 ) );
-		Set( "GraphEditor.AlignNodesMiddle", new IMAGE_BRUSH( "Icons/GraphEditor/icon_AlignNodesMiddle_20px", Icon20x20 ) );
-		Set( "GraphEditor.AlignNodesBottom", new IMAGE_BRUSH( "Icons/GraphEditor/icon_AlignNodesBottom_20px", Icon20x20 ) );
-		Set( "GraphEditor.AlignNodesLeft", new IMAGE_BRUSH( "Icons/GraphEditor/icon_AlignNodesLeft_20px", Icon20x20 ) );
-		Set( "GraphEditor.AlignNodesCenter", new IMAGE_BRUSH( "Icons/GraphEditor/icon_AlignNodesCenter_20px", Icon20x20 ) );
-		Set( "GraphEditor.AlignNodesRight", new IMAGE_BRUSH( "Icons/GraphEditor/icon_AlignNodesRight_20px", Icon20x20 ) );
+		Set( "GraphEditor.AlignNodesTop", new IMAGE_BRUSH_SVG( "Starship/GraphEditors/AlignTop", Icon20x20 ) );
+		Set( "GraphEditor.AlignNodesMiddle", new IMAGE_BRUSH_SVG( "Starship/GraphEditors/AlignMiddle", Icon20x20 ) );
+		Set( "GraphEditor.AlignNodesBottom", new IMAGE_BRUSH_SVG( "Starship/GraphEditors/AlignBottom", Icon20x20 ) );
+		Set( "GraphEditor.AlignNodesLeft", new IMAGE_BRUSH_SVG( "Starship/GraphEditors/AlignLeft", Icon20x20 ) );
+		Set( "GraphEditor.AlignNodesCenter", new IMAGE_BRUSH_SVG( "Starship/Common/Align", Icon20x20 ) );
+		Set( "GraphEditor.AlignNodesRight", new IMAGE_BRUSH_SVG( "Starship/GraphEditors/AlignRight", Icon20x20 ) );
 
-		Set( "GraphEditor.StraightenConnections", new IMAGE_BRUSH( "Icons/GraphEditor/icon_StraightenConnections_20px", Icon20x20 ) );
+		Set( "GraphEditor.StraightenConnections", new IMAGE_BRUSH_SVG( "Starship/GraphEditors/StraightenConnections", Icon20x20 ) );
 
-		Set( "GraphEditor.DistributeNodesHorizontally", new IMAGE_BRUSH( "Icons/GraphEditor/icon_DistributeNodesHorizontally_20px", Icon20x20 ) );
-		Set( "GraphEditor.DistributeNodesVertically", new IMAGE_BRUSH( "Icons/GraphEditor/icon_DistributeNodesVertically_20px", Icon20x20 ) );
+		Set( "GraphEditor.DistributeNodesHorizontally", new IMAGE_BRUSH_SVG( "Starship/GraphEditors/DistributeHorizontally", Icon20x20 ) );
+		Set( "GraphEditor.DistributeNodesVertically", new IMAGE_BRUSH_SVG( "Starship/GraphEditors/DistributeVertically", Icon20x20 ) );
 		
 		Set( "GraphEditor.ToggleHideUnrelatedNodes", new IMAGE_BRUSH_SVG( "Starship/GraphEditors/HideUnrelated", Icon20x20 ) );
 		Set("GraphEditor.Bookmark", new IMAGE_BRUSH_SVG("Starship/Common/Bookmarks", Icon16x16));
@@ -5046,7 +5022,7 @@ void FStarshipEditorStyle::FStyle::SetupPersonaStyle()
 
 		// AssetFamilyIcons
 		Set("Persona.AssetClass.Skeleton", new IMAGE_BRUSH_SVG("Starship/Persona/Skeleton", Icon20x20));
-		Set("Persona.AssetClass.Animation", new IMAGE_BRUSH_SVG("Starship/Persona/Animation", Icon20x20));
+		Set("Persona.AssetClass.Animation", new IMAGE_BRUSH_SVG("Starship/Common/Animation", Icon20x20));
 		Set("Persona.AssetClass.SkeletalMesh", new IMAGE_BRUSH_SVG("Starship/Persona/SkeletalMesh", Icon20x20));
 		Set("Persona.AssetClass.Blueprint", new IMAGE_BRUSH_SVG( "Starship/MainToolbar/blueprints", Icon20x20) );
 		Set("Persona.AssetClass.Physics", new IMAGE_BRUSH_SVG("Starship/Persona/Physics", Icon20x20));
