Sequencer: Added mising Sequencer clock source icons

#rnx

diff --git a/Engine/Source/Editor/EditorStyle/Private/StarshipStyle.cpp b/Engine/Source/Editor/EditorStyle/Private/StarshipStyle.cpp
--- a/Engine/Source/Editor/EditorStyle/Private/StarshipStyle.cpp
+++ b/Engine/Source/Editor/EditorStyle/Private/StarshipStyle.cpp
@@ -2175,6 +2175,12 @@ void FStarshipEditorStyle::FStyle::SetupSequencerStyles()
 		Set("Sequencer.CursorDecorator_Retime", new IMAGE_BRUSH("Sequencer/CursorDecorator_Retime", Icon16x16));
 		Set("Sequencer.CursorDecorator_EasingHandle", new IMAGE_BRUSH("Sequencer/CursorDecorator_EasingHandle", Icon16x16));
 
+		Set("Sequencer.ClockSource.Platform", new IMAGE_BRUSH("Sequencer/Main_Icons/Icon_ClockSource_Platform_16x", Icon14x14));
+		Set("Sequencer.ClockSource.Audio", new IMAGE_BRUSH("Sequencer/Main_Icons/Icon_ClockSource_Audio_16x", Icon14x14));
+		Set("Sequencer.ClockSource.RelativeTimecode", new IMAGE_BRUSH("Sequencer/Main_Icons/Icon_ClockSource_RelativeTimecode_16x", Icon14x14));
+		Set("Sequencer.ClockSource.Timecode", new IMAGE_BRUSH("Sequencer/Main_Icons/Icon_ClockSource_Timecode_16x", Icon14x14));
+		Set("Sequencer.ClockSource.Custom", new IMAGE_BRUSH("Sequencer/Main_Icons/Icon_ClockSource_Custom_16x", Icon14x14));
+
 		Set("Sequencer.BreadcrumbText", FTextBlockStyle(NormalText).SetFont(FStyleFonts::Get().NormalBold));
 		Set("Sequencer.BreadcrumbIcon", new IMAGE_BRUSH("Common/SmallArrowRight", Icon10x10));
 
