Adding the combo chevron to the status bar button

diff --git a/Engine/Source/Editor/EditorStyle/Private/StarshipStyle.cpp b/Engine/Source/Editor/EditorStyle/Private/StarshipStyle.cpp
--- a/Engine/Source/Editor/EditorStyle/Private/StarshipStyle.cpp
+++ b/Engine/Source/Editor/EditorStyle/Private/StarshipStyle.cpp
@@ -7756,8 +7756,10 @@ void FStarshipEditorStyle::FStyle::SetupStatusBarStyle()
 		.SetPressedPadding(FMargin(2, 3, 2, 1));
 
 	const FComboButtonStyle StatusBarComboButton = FComboButtonStyle(FStarshipCoreStyle::GetCoreStyle().GetWidgetStyle<FComboButtonStyle>("ComboButton"))
-		.SetDownArrowImage(FSlateNoResource())
-		.SetButtonStyle(StatusBarButton);
+		.SetDownArrowImage(CORE_IMAGE_BRUSH_SVG("Starship/CoreWidgets/ComboBox/corner-dropdown", FVector2D(7.0f, 7.0f)))
+		.SetButtonStyle(StatusBarButton)
+		.SetDownArrowPadding(FMargin(0.0f))
+		.SetDownArrowAlignment(EVerticalAlignment::VAlign_Bottom);
 		
 	Set("StatusBar.StatusBarButton", StatusBarButton);
 	Set("StatusBar.StatusBarComboButton", StatusBarComboButton);
