Updated the Developer Folder Icons in the content browser and the lock icons

diff --git a/Engine/Source/Editor/EditorStyle/Private/StarshipStyle.cpp b/Engine/Source/Editor/EditorStyle/Private/StarshipStyle.cpp
--- a/Engine/Source/Editor/EditorStyle/Private/StarshipStyle.cpp
+++ b/Engine/Source/Editor/EditorStyle/Private/StarshipStyle.cpp
@@ -6084,7 +6084,7 @@ void FStarshipEditorStyle::FStyle::SetupContentBrowserStyle()
 		Set("ContentBrowser.AssetTreeFolderClosed", new CORE_IMAGE_BRUSH_SVG("Starship/Common/folder-closed", Icon16x16));
 		Set("ContentBrowser.AssetTreeFolderOpen", new CORE_IMAGE_BRUSH_SVG("Starship/Common/folder-open", Icon16x16));
 
-		Set( "ContentBrowser.AssetTreeFolderDeveloper", new IMAGE_BRUSH( "Icons/FolderDeveloper", FVector2D(18, 16) ) );
+		Set( "ContentBrowser.AssetTreeFolderDeveloper", new IMAGE_BRUSH_SVG( "Starship/ContentBrowser/folder-developer", Icon16x16));
 		Set( "ContentBrowser.AssetTreeFolderOpenCode", new IMAGE_BRUSH( "Icons/FolderOpen_Code", FVector2D(18, 16) ) );
 		Set( "ContentBrowser.AssetTreeFolderClosedCode", new IMAGE_BRUSH( "Icons/FolderClosed_Code", FVector2D(18, 16) ) );
 
@@ -6273,10 +6273,10 @@ void FStarshipEditorStyle::FStyle::SetupContentBrowserStyle()
 		Set( "ContentBrowser.ColumnViewAssetIcon", new IMAGE_BRUSH( "Icons/doc_16x", Icon16x16 ) );
 
 		Set( "ContentBrowser.ColumnViewFolderIcon", new CORE_IMAGE_BRUSH_SVG( "Starship/Common/folder-closed", Icon16x16 ) );
-		Set( "ContentBrowser.ColumnViewDeveloperFolderIcon", new IMAGE_BRUSH( "Icons/FolderDeveloper", FVector2D(18, 16) ) );
+		Set( "ContentBrowser.ColumnViewDeveloperFolderIcon", new IMAGE_BRUSH_SVG( "Starship/ContentBrowser/folder-developer", Icon16x16 ) );
 
 		Set( "ContentBrowser.ListViewFolderIcon", new IMAGE_BRUSH_SVG( "Starship/ContentBrowser/folder", FVector2D(64, 64) ) );
-		Set( "ContentBrowser.ListViewDeveloperFolderIcon", new IMAGE_BRUSH( "Icons/Folders/FolderDev_Base_256x", FVector2D(256, 256) ) );
+		Set( "ContentBrowser.ListViewDeveloperFolderIcon", new IMAGE_BRUSH_SVG("Starship/ContentBrowser/FolderLargeDeveloper", FVector2D(64, 64)));
 
 		Set("ContentBrowser.AssetTileItem.FolderAreaHoveredBackground", new FSlateRoundedBoxBrush(FStyleColors::Secondary, 4.0f));
 		Set("ContentBrowser.AssetTileItem.ThumbnailAreaBackground", new FSlateRoundedBoxBrush(FStyleColors::Recessed, 4.0f));
