Sequencer: UX Cleanup

- Tighten up Sequencer and Curve Editor toolbars
- Fix  key navigation sizing/spacing
- Remove transparent style for pill text input which makes it hard to read and inconsistent with the new starship style.

#rnx

diff --git a/Engine/Source/Editor/CurveEditor/Private/SCurveEditorPanel.cpp b/Engine/Source/Editor/CurveEditor/Private/SCurveEditorPanel.cpp
--- a/Engine/Source/Editor/CurveEditor/Private/SCurveEditorPanel.cpp
+++ b/Engine/Source/Editor/CurveEditor/Private/SCurveEditorPanel.cpp
@@ -1031,7 +1031,7 @@ TSharedPtr<FExtender> SCurveEditorPanel::GetToolbarExtender()
 		static void FillToolbar(FToolBarBuilder& ToolBarBuilder, TSharedRef<SCurveKeyDetailPanel> InKeyDetailsPanel, TSharedRef<SCurveEditorPanel> InEditorPanel)
 		{
 			ToolBarBuilder.BeginSection("View");
-			ToolBarBuilder.BeginStyleOverride("AssetEditorToolbar");
+			ToolBarBuilder.BeginStyleOverride("CurveEditorToolbar");
 			{
 				// Dropdown Menu for choosing your viewing mode
 				TAttribute<FSlateIcon> ViewModeIcon;
diff --git a/Engine/Source/Editor/EditorStyle/Private/SlateEditorStyle.cpp b/Engine/Source/Editor/EditorStyle/Private/SlateEditorStyle.cpp
--- a/Engine/Source/Editor/EditorStyle/Private/SlateEditorStyle.cpp
+++ b/Engine/Source/Editor/EditorStyle/Private/SlateEditorStyle.cpp
@@ -2465,13 +2465,6 @@ void FSlateEditorStyle::FStyle::SetupSequencerStyles()
 			.SetNormalPadding( FMargin(0,0,0,0) )
 			.SetPressedPadding( FMargin(0,0,0,0) )
 		);
-		Set("Sequencer.ExposedNamePill.Input", FEditableTextBoxStyle()
-			.SetBackgroundImageNormal(FSlateColorBrush(FLinearColor::Transparent))
-			.SetBackgroundImageHovered(FSlateColorBrush(FLinearColor::Transparent))
-			.SetBackgroundImageFocused(FSlateColorBrush(FLinearColor::Transparent))
-			.SetBackgroundImageReadOnly(FSlateColorBrush(FLinearColor::Transparent))
-		);
-
 
 		Set("Sequencer.AnimationOutliner.ColorStrip", FButtonStyle()
 			.SetNormal(FSlateNoResource())
diff --git a/Engine/Source/Editor/EditorStyle/Private/StarshipStyle.cpp b/Engine/Source/Editor/EditorStyle/Private/StarshipStyle.cpp
--- a/Engine/Source/Editor/EditorStyle/Private/StarshipStyle.cpp
+++ b/Engine/Source/Editor/EditorStyle/Private/StarshipStyle.cpp
@@ -1956,6 +1956,17 @@ void FStarshipEditorStyle::FStyle::SetupSequencerStyles()
 	// Sequencer
 	if (IncludeEditorSpecificStyles())
 	{
+		FToolBarStyle SequencerToolbar = FStarshipCoreStyle::GetCoreStyle().GetWidgetStyle<FToolBarStyle>("AssetEditorToolbar");
+
+		SequencerToolbar.SetButtonPadding(       FMargin(6.0f, 0.0f));
+		SequencerToolbar.SetCheckBoxPadding(     FMargin(2.0f, 0.0f));
+		SequencerToolbar.SetComboButtonPadding(  FMargin(2.0f, 0.0f));
+		SequencerToolbar.SetIndentedBlockPadding(FMargin(2.0f, 0.0f));
+		SequencerToolbar.SetBlockPadding(        FMargin(2.0f, 0.0f));
+		SequencerToolbar.SetSeparatorPadding(    FMargin(6.0f, 0.0f));
+
+		Set("SequencerToolBar", SequencerToolbar);
+
 		Set("Sequencer.IconKeyAuto", new IMAGE_BRUSH("Sequencer/IconKeyAuto", Icon12x12));
 		Set("Sequencer.IconKeyBreak", new IMAGE_BRUSH("Sequencer/IconKeyBreak", Icon12x12));
 		Set("Sequencer.IconKeyConstant", new IMAGE_BRUSH("Sequencer/IconKeyConstant", Icon12x12));
@@ -2069,13 +2080,6 @@ void FStarshipEditorStyle::FStyle::SetupSequencerStyles()
 			.SetNormalPadding( FMargin(0,0,0,0) )
 			.SetPressedPadding( FMargin(0,0,0,0) )
 		);
-		Set("Sequencer.ExposedNamePill.Input", FEditableTextBoxStyle()
-			.SetBackgroundImageNormal(FSlateColorBrush(FLinearColor::Transparent))
-			.SetBackgroundImageHovered(FSlateColorBrush(FLinearColor::Transparent))
-			.SetBackgroundImageFocused(FSlateColorBrush(FLinearColor::Transparent))
-			.SetBackgroundImageReadOnly(FSlateColorBrush(FLinearColor::Transparent))
-		);
-
 
 		Set("Sequencer.AnimationOutliner.ColorStrip", FButtonStyle()
 			.SetNormal(FSlateNoResource())
@@ -6901,6 +6905,17 @@ void FStarshipEditorStyle::FStyle::SetupToolkitStyles()
 
 	// New Curve Editor 
 	{
+		FToolBarStyle CurveEditorToolbar = FStarshipCoreStyle::GetCoreStyle().GetWidgetStyle<FToolBarStyle>("AssetEditorToolbar");
+
+		CurveEditorToolbar.SetButtonPadding(       FMargin(6.0f, 0.0f));
+		CurveEditorToolbar.SetCheckBoxPadding(     FMargin(0.0f, 0.0f));
+		CurveEditorToolbar.SetComboButtonPadding(  FMargin(0.0f, 0.0f));
+		CurveEditorToolbar.SetIndentedBlockPadding(FMargin(0.0f, 0.0f));
+		CurveEditorToolbar.SetBlockPadding(        FMargin(0.0f, 0.0f));
+		CurveEditorToolbar.SetSeparatorPadding(    FMargin(2.0f, 0.0f));
+
+		Set("CurveEditorToolBar", CurveEditorToolbar);
+
 		// Tab
 		Set("GenericCurveEditor.TabIcon", new IMAGE_BRUSH("GenericCurveEditor/Icons/CurveEditor_32x", Icon16x16, FLinearColor(1.f, 1.f, 1.f, 0.8f)));
 
diff --git a/Engine/Source/Editor/LevelEditor/Private/LevelEditorToolBar.cpp b/Engine/Source/Editor/LevelEditor/Private/LevelEditorToolBar.cpp
--- a/Engine/Source/Editor/LevelEditor/Private/LevelEditorToolBar.cpp
+++ b/Engine/Source/Editor/LevelEditor/Private/LevelEditorToolBar.cpp
@@ -1969,7 +1969,7 @@ void FLevelEditorToolBar::RegisterCinematicsMenu()
 			// We hide the header row to keep the UI compact.
 			// @todo: Might be useful to have this sometimes, actually.  Ideally the user could summon it.
 			InitOptions.bShowHeaderRow = false;
-			InitOptions.bShowSearchBox = false;
+			InitOptions.bShowSearchBox = true;
 			InitOptions.bShowCreateNewFolder = false;
 
 			InitOptions.ColumnMap.Add(FSceneOutlinerBuiltInColumnTypes::Label(), FSceneOutlinerColumnInfo(ESceneOutlinerColumnVisibility::Visible, 0));
diff --git a/Engine/Source/Editor/MovieSceneTools/Private/SObjectBindingTag.cpp b/Engine/Source/Editor/MovieSceneTools/Private/SObjectBindingTag.cpp
--- a/Engine/Source/Editor/MovieSceneTools/Private/SObjectBindingTag.cpp
+++ b/Engine/Source/Editor/MovieSceneTools/Private/SObjectBindingTag.cpp
@@ -92,8 +92,6 @@ void SObjectBindingTag::Construct(const FArguments& InArgs)
 		[
 			SAssignNew(EditableText, SEditableTextBox)
 			.Font(FEditorStyle::GetFontStyle("TinyText"))
-			.Style(FEditorStyle::Get(), "Sequencer.ExposedNamePill.Input")
-			.BackgroundColor(FLinearColor::Transparent)
 			.OnTextCommitted(this, &SObjectBindingTag::OnNewTextCommitted)
 			.HintText(LOCTEXT("AddNew_Hint", "Enter New Name"))
 		];
@@ -168,6 +166,13 @@ void SObjectBindingTag::Construct(const FArguments& InArgs)
 			]
 		];
 	}
+	else if (OnCreateNew.IsBound())
+	{
+		ChildSlot
+		[
+			ContentBox
+		];
+	}
 	else
 	{
 		// else it's just a border
diff --git a/Engine/Source/Editor/Sequencer/Private/SKeyNavigationButtons.h b/Engine/Source/Editor/Sequencer/Private/SKeyNavigationButtons.h
--- a/Engine/Source/Editor/Sequencer/Private/SKeyNavigationButtons.h
+++ b/Engine/Source/Editor/Sequencer/Private/SKeyNavigationButtons.h
@@ -66,11 +66,11 @@ public:
 				.ColorAndOpacity(HoverTint)
 				[
 					SNew(SButton)
-					.ButtonStyle(FEditorStyle::Get(), "FlatButton")
+					.ButtonStyle(FEditorStyle::Get(), "HoverHintOnly")
 					.ToolTipText(PreviousKeyToolTip)
 					.OnClicked(this, &SKeyNavigationButtons::OnPreviousKeyClicked)
 					.ForegroundColor( FSlateColor::UseForeground() )
-					.ContentPadding(0)
+					.ContentPadding(2)
 					.IsFocusable(false)
 					[
 						SNew(STextBlock)
@@ -91,11 +91,11 @@ public:
 				.IsEnabled(!InDisplayNode->GetSequencer().IsReadOnly())
 				[
 					SNew(SButton)
-					.ButtonStyle(FEditorStyle::Get(), "FlatButton")
+					.ButtonStyle(FEditorStyle::Get(), "HoverHintOnly")
 					.ToolTipText(SetKeyToolTip)
 					.OnClicked(this, &SKeyNavigationButtons::OnAddKeyClicked)
 					.ForegroundColor( FSlateColor::UseForeground() )
-					.ContentPadding(0)
+					.ContentPadding(2)
 					.IsFocusable(false)
 					[
 						SNew(STextBlock)
@@ -115,10 +115,10 @@ public:
 				.ColorAndOpacity(HoverTint)
 				[
 					SNew(SButton)
-					.ButtonStyle(FEditorStyle::Get(), "FlatButton")
+					.ButtonStyle(FEditorStyle::Get(), "HoverHintOnly")
 					.ToolTipText(NextKeyToolTip)
 					.OnClicked(this, &SKeyNavigationButtons::OnNextKeyClicked)
-					.ContentPadding(0)
+					.ContentPadding(2)
 					.ForegroundColor( FSlateColor::UseForeground() )
 					.IsFocusable(false)
 					[
diff --git a/Engine/Source/Editor/Sequencer/Private/SSequencer.cpp b/Engine/Source/Editor/Sequencer/Private/SSequencer.cpp
--- a/Engine/Source/Editor/Sequencer/Private/SSequencer.cpp
+++ b/Engine/Source/Editor/Sequencer/Private/SSequencer.cpp
@@ -1365,6 +1365,7 @@ TSharedRef<SWidget> SSequencer::MakeToolBar()
 	FSlimHorizontalToolBarBuilder ToolBarBuilder( SequencerPtr.Pin()->GetCommandBindings(), FMultiBoxCustomization::None, Extender, true);
 	
 	ToolBarBuilder.BeginSection("Base Commands");
+	ToolBarBuilder.BeginStyleOverride("SequencerToolbar");
 	{
 		// General 
 		if (SequencerPtr.Pin()->IsLevelEditorSequencer())
@@ -1591,6 +1592,7 @@ TSharedRef<SWidget> SSequencer::MakeToolBar()
 		}
 	}
 	ToolBarBuilder.EndSection();
+	ToolBarBuilder.EndStyleOverride();
 
 	return ToolBarBuilder.MakeWidget();
 }
