Adjust viewport toolbar color and styles

diff --git a/Engine/Shaders/Private/SlateShaderCommon.ush b/Engine/Shaders/Private/SlateShaderCommon.ush
--- a/Engine/Shaders/Private/SlateShaderCommon.ush
+++ b/Engine/Shaders/Private/SlateShaderCommon.ush
@@ -84,9 +84,11 @@ float4 GetRoundedBoxElementColorInternal(float2 size, float2 UV, float thickness
 	float din  = GetRoundedBoxDistance(pos, center, max(radius - thickness, 0), thickness);
 
 	// Compute the border intensity and fill intensity with a smooth transition
-	float spread = 0.5;
-	float bi = smoothstep(spread, -spread, dext);
-	float fi = smoothstep(spread, -spread, din);
+	float bi_spread = 1.0;
+	float bi = smoothstep(bi_spread, -bi_spread, dext);
+
+	float fi_spread = .5;
+	float fi = smoothstep(fi_spread, -fi_spread, din);
 
 	float4 fill = fillColor;
 	float4 border = borderColor;
diff --git a/Engine/Source/Editor/EditorStyle/Private/StarshipStyle.cpp b/Engine/Source/Editor/EditorStyle/Private/StarshipStyle.cpp
--- a/Engine/Source/Editor/EditorStyle/Private/StarshipStyle.cpp
+++ b/Engine/Source/Editor/EditorStyle/Private/StarshipStyle.cpp
@@ -2264,9 +2264,6 @@ void FStarshipEditorStyle::FStyle::SetupViewportStyles()
 	{
 		FToolBarStyle ViewportToolbarStyle = FStarshipCoreStyle::GetCoreStyle().GetWidgetStyle<FToolBarStyle>("SlimToolBar");
 
-		FLinearColor TransInput = FStyleColors::Input.GetSpecifiedColor();
-		TransInput.A = .71;
-
 		FMargin ViewportMargin(4.f, 4.f);
 
 		const FCheckBoxStyle ViewportToggleButton = FCheckBoxStyle()
@@ -2280,11 +2277,11 @@ void FStarshipEditorStyle::FStyle::SetupViewportStyles()
 			.SetForegroundColor(FStyleColors::Foreground)
 			.SetHoveredForegroundColor(FStyleColors::ForegroundHover)
 			.SetPressedForegroundColor(FStyleColors::ForegroundHover)
-			.SetCheckedForegroundColor(FStyleColors::PrimaryHover)
-			.SetCheckedHoveredForegroundColor(FStyleColors::PrimaryHover)
+			.SetCheckedForegroundColor(FStyleColors::ForegroundHover)
+			.SetCheckedHoveredForegroundColor(FStyleColors::ForegroundHover)
 			.SetPadding(0);
 
-		FSlateRoundedBoxBrush* ViewportGroupBrush = new FSlateRoundedBoxBrush(TransInput, 12.f, FStyleColors::Dropdown, 1.0);
+		FSlateRoundedBoxBrush* ViewportGroupBrush = new FSlateRoundedBoxBrush(FStyleColors::Dropdown, 12.f, FStyleColors::Black, 1.0);
 		Set("EditorViewportToolBar.Group", ViewportGroupBrush);
 
 		FButtonStyle ViewportMenuButton = FButtonStyle()
@@ -2301,41 +2298,48 @@ void FStarshipEditorStyle::FStyle::SetupViewportStyles()
 
 		const FCheckBoxStyle ViewportMenuToggleLeftButtonStyle = FCheckBoxStyle(ViewportToggleButton)
 			.SetCheckBoxType(ESlateCheckBoxType::ToggleButton)
-			.SetUncheckedImage(		  BOX_BRUSH("Starship/EditorViewport/ToolBarLeftGroup", 12.f/25.f))
-			.SetUncheckedPressedImage(BOX_BRUSH("Starship/EditorViewport/ToolBarLeftGroup", 12.f/25.f))
-			.SetUncheckedHoveredImage(BOX_BRUSH("Starship/EditorViewport/ToolBarLeftGroup", 12.f/25.f))
-			.SetCheckedHoveredImage(  BOX_BRUSH("Starship/EditorViewport/ToolBarLeftGroup", 12.f/25.f))
-			.SetCheckedPressedImage(  BOX_BRUSH("Starship/EditorViewport/ToolBarLeftGroup", 12.f/25.f))
-			.SetCheckedImage(         BOX_BRUSH("Starship/EditorViewport/ToolBarLeftGroup", 12.f/25.f))
+			.SetUncheckedImage(		  BOX_BRUSH("Starship/EditorViewport/ToolBarLeftGroup", 12.f/25.f, FStyleColors::Dropdown))
+			.SetUncheckedPressedImage(BOX_BRUSH("Starship/EditorViewport/ToolBarLeftGroup", 12.f/25.f,FStyleColors::Header))
+			.SetUncheckedHoveredImage(BOX_BRUSH("Starship/EditorViewport/ToolBarLeftGroup", 12.f/25.f, FStyleColors::Hover))
+			.SetCheckedHoveredImage(  BOX_BRUSH("Starship/EditorViewport/ToolBarLeftGroup", 12.f/25.f, FStyleColors::PrimaryHover))
+			.SetCheckedPressedImage(  BOX_BRUSH("Starship/EditorViewport/ToolBarLeftGroup", 12.f/25.f, FStyleColors::PrimaryPress))
+			.SetCheckedImage(         BOX_BRUSH("Starship/EditorViewport/ToolBarLeftGroup", 12.f/25.f, FStyleColors::Primary))
 			.SetPadding(ViewportMargin);
 		Set("EditorViewportToolBar.ToggleButton.Start", ViewportMenuToggleLeftButtonStyle);
 
 		const FCheckBoxStyle ViewportMenuToggleMiddleButtonStyle = FCheckBoxStyle(ViewportToggleButton)
 			.SetCheckBoxType(ESlateCheckBoxType::ToggleButton)
-			.SetUncheckedImage(		  BOX_BRUSH("Starship/EditorViewport/ToolBarMiddleGroup", 12.f/25.f))
-			.SetUncheckedPressedImage(BOX_BRUSH("Starship/EditorViewport/ToolBarMiddleGroup", 12.f/25.f))
-			.SetUncheckedHoveredImage(BOX_BRUSH("Starship/EditorViewport/ToolBarMiddleGroup", 12.f/25.f))
-			.SetCheckedHoveredImage(  BOX_BRUSH("Starship/EditorViewport/ToolBarMiddleGroup", 12.f/25.f))
-			.SetCheckedPressedImage(  BOX_BRUSH("Starship/EditorViewport/ToolBarMiddleGroup", 12.f/25.f))
-			.SetCheckedImage(         BOX_BRUSH("Starship/EditorViewport/ToolBarMiddleGroup", 12.f/25.f))
+			.SetUncheckedImage(		  BOX_BRUSH("Starship/EditorViewport/ToolBarMiddleGroup", 12.f/25.f, FStyleColors::Dropdown))
+			.SetUncheckedPressedImage(BOX_BRUSH("Starship/EditorViewport/ToolBarMiddleGroup", 12.f/25.f, FStyleColors::Header))
+			.SetUncheckedHoveredImage(BOX_BRUSH("Starship/EditorViewport/ToolBarMiddleGroup", 12.f/25.f, FStyleColors::Hover))
+			.SetCheckedHoveredImage(  BOX_BRUSH("Starship/EditorViewport/ToolBarMiddleGroup", 12.f/25.f, FStyleColors::PrimaryHover))
+			.SetCheckedPressedImage(  BOX_BRUSH("Starship/EditorViewport/ToolBarMiddleGroup", 12.f/25.f, FStyleColors::PrimaryPress))
+			.SetCheckedImage(         BOX_BRUSH("Starship/EditorViewport/ToolBarMiddleGroup", 12.f/25.f, FStyleColors::Primary))
 			.SetPadding(ViewportMargin);
 		Set("EditorViewportToolBar.ToggleButton.Middle", ViewportMenuToggleMiddleButtonStyle);
 
 		const FCheckBoxStyle ViewportMenuToggleRightButtonStyle = FCheckBoxStyle(ViewportToggleButton)
 			.SetCheckBoxType(ESlateCheckBoxType::ToggleButton)
-			.SetUncheckedImage(		  BOX_BRUSH("Starship/EditorViewport/ToolBarRightGroup", 12.f/25.f))
-			.SetUncheckedPressedImage(BOX_BRUSH("Starship/EditorViewport/ToolBarRightGroup", 12.f/25.f))
-			.SetUncheckedHoveredImage(BOX_BRUSH("Starship/EditorViewport/ToolBarRightGroup", 12.f/25.f))
-			.SetCheckedHoveredImage(  BOX_BRUSH("Starship/EditorViewport/ToolBarRightGroup", 12.f/25.f))
-			.SetCheckedPressedImage(  BOX_BRUSH("Starship/EditorViewport/ToolBarRightGroup", 12.f/25.f))
-			.SetCheckedImage(         BOX_BRUSH("Starship/EditorViewport/ToolBarRightGroup", 12.f/25.f))
+			.SetUncheckedImage(		  BOX_BRUSH("Starship/EditorViewport/ToolBarRightGroup", 12.f/25.f, FStyleColors::Dropdown))
+			.SetUncheckedPressedImage(BOX_BRUSH("Starship/EditorViewport/ToolBarRightGroup", 12.f/25.f, FStyleColors::Header))
+			.SetUncheckedHoveredImage(BOX_BRUSH("Starship/EditorViewport/ToolBarRightGroup", 12.f/25.f, FStyleColors::Hover))
+			.SetCheckedHoveredImage(  BOX_BRUSH("Starship/EditorViewport/ToolBarRightGroup", 12.f/25.f, FStyleColors::PrimaryHover))
+			.SetCheckedPressedImage(  BOX_BRUSH("Starship/EditorViewport/ToolBarRightGroup", 12.f/25.f, FStyleColors::PrimaryPress))
+			.SetCheckedImage(         BOX_BRUSH("Starship/EditorViewport/ToolBarRightGroup", 12.f/25.f, FStyleColors::Primary))
 			.SetPadding(ViewportMargin);
 		Set("EditorViewportToolBar.ToggleButton.End", ViewportMenuToggleRightButtonStyle);
 
 		// We want a background-less version as the ComboMenu has its own unified background
 		const FToolBarStyle& SlimCoreToolBarStyle = FStarshipCoreStyle::GetCoreStyle().GetWidgetStyle<FToolBarStyle>("SlimToolBar");
 
-		Set("EditorViewportToolBar.ComboMenu.ButtonStyle", FButtonStyle(SlimCoreToolBarStyle.ButtonStyle).SetNormalPadding(0.0).SetPressedPadding(0.0));
+		FButtonStyle ComboMenuButtonStyle = FButtonStyle(SlimCoreToolBarStyle.ButtonStyle)
+			.SetNormal(BOX_BRUSH("Starship/EditorViewport/ToolBarRightGroup", 12.f/25.f, FStyleColors::Dropdown))
+			.SetPressed(BOX_BRUSH("Starship/EditorViewport/ToolBarRightGroup", 12.f/25.f, FStyleColors::Header))
+			.SetHovered(BOX_BRUSH("Starship/EditorViewport/ToolBarRightGroup", 12.f/25.f, FStyleColors::Hover))
+			.SetNormalPadding(0.0)
+			.SetPressedPadding(0.0);
+
+		Set("EditorViewportToolBar.ComboMenu.ButtonStyle", ComboMenuButtonStyle);
 		Set("EditorViewportToolBar.ComboMenu.ToggleButton", ViewportToggleButton);
 		Set("EditorViewportToolBar.ComboMenu.LabelStyle", SlimCoreToolBarStyle.LabelStyle);
 
diff --git a/Engine/Source/Editor/LevelEditor/Private/LevelEditorToolBar.cpp b/Engine/Source/Editor/LevelEditor/Private/LevelEditorToolBar.cpp
--- a/Engine/Source/Editor/LevelEditor/Private/LevelEditorToolBar.cpp
+++ b/Engine/Source/Editor/LevelEditor/Private/LevelEditorToolBar.cpp
@@ -1441,7 +1441,7 @@ TSharedRef< SWidget > FLevelEditorToolBar::MakeLevelEditorToolBar( const TShared
 			SNew(SHorizontalBox)
 			.IsEnabled(FSlateApplication::Get().GetNormalExecutionAttribute())
 			+ SHorizontalBox::Slot()
-			.Padding(7.0f, 0.0f, 0.0f, 0.0f)
+			.Padding(0.0f)
 			.AutoWidth()
 			[
 				UToolMenus::Get()->GenerateWidget("LevelEditor.LevelEditorToolBar.AssetsToolBar", MenuContext)
diff --git a/Engine/Source/Editor/LevelEditor/Private/SLevelViewport.cpp b/Engine/Source/Editor/LevelEditor/Private/SLevelViewport.cpp
--- a/Engine/Source/Editor/LevelEditor/Private/SLevelViewport.cpp
+++ b/Engine/Source/Editor/LevelEditor/Private/SLevelViewport.cpp
@@ -1859,6 +1859,7 @@ TSharedPtr<SWidget> SLevelViewport::MakeViewportToolbar()
 		.Visibility( EVisibility::SelfHitTestInvisible )
 		+SVerticalBox::Slot()
 		.AutoHeight()
+		.Padding(0, 1.0f, 0, 0)
 		.VAlign(VAlign_Top)
 		[
 			ToolBar
diff --git a/Engine/Source/Editor/UnrealEd/Private/SViewportToolBarComboMenu.cpp b/Engine/Source/Editor/UnrealEd/Private/SViewportToolBarComboMenu.cpp
--- a/Engine/Source/Editor/UnrealEd/Private/SViewportToolBarComboMenu.cpp
+++ b/Engine/Source/Editor/UnrealEd/Private/SViewportToolBarComboMenu.cpp
@@ -18,35 +18,32 @@ void SViewportToolBarComboMenu::Construct( const FArguments& InArgs )
 	EMultiBlockLocation::Type BlockLocation = InArgs._BlockLocation;
 
 	const FButtonStyle& ButtonStyle = FAppStyle::Get().GetWidgetStyle<FButtonStyle>("EditorViewportToolBar.ComboMenu.ButtonStyle");
-	const FCheckBoxStyle& CheckBoxStyle = FAppStyle::Get().GetWidgetStyle<FCheckBoxStyle>("EditorViewportToolBar.ComboMenu.ToggleButton");
+	const FCheckBoxStyle& CheckBoxStyle = FAppStyle::Get().GetWidgetStyle<FCheckBoxStyle>("EditorViewportToolBar.ToggleButton.Start");
 	const FTextBlockStyle& LabelStyle = FAppStyle::Get().GetWidgetStyle<FTextBlockStyle>("EditorViewportToolBar.ComboMenu.LabelStyle");
 
 	const FSlateIcon& Icon = InArgs._Icon.Get();
-
 	ParentToolBar = InArgs._ParentToolBar;
 
-	TSharedPtr<SCheckBox> ToggleControl;
-	{
-		ToggleControl = SNew(SCheckBox)
-		.Padding(0.f)
+
+	TSharedRef<SCheckBox> ToggleControl = SNew(SCheckBox)
 		.Style(&CheckBoxStyle)
-		.OnCheckStateChanged(InArgs._OnCheckStateChanged)
 		.ToolTipText(InArgs._ToggleButtonToolTip)
+		.OnCheckStateChanged(InArgs._OnCheckStateChanged)
 		.IsChecked(InArgs._IsChecked)
 		[
 			SNew(SImage)
 			.Image(Icon.GetIcon())
 			.ColorAndOpacity(FSlateColor::UseForeground())
 		];
-	}
+
 
 	{
 		TSharedRef<SWidget> ButtonContents =
 			SNew(SButton)
 			.ButtonStyle(&ButtonStyle)
-			.ContentPadding(0.f)
 			.ToolTipText(InArgs._MenuButtonToolTip)
 			.OnClicked(this, &SViewportToolBarComboMenu::OnMenuClicked)
+			.VAlign(VAlign_Center)
 			[
 				SNew(STextBlock)
 				.TextStyle(&LabelStyle)
@@ -74,28 +71,18 @@ void SViewportToolBarComboMenu::Construct( const FArguments& InArgs )
 
 	ChildSlot
 	[
-
-		SNew(SBorder)
-		.Padding(FMargin(6.f, 0.f, 6.f, 0.f))
-		.BorderImage(FAppStyle::Get().GetBrush("EditorViewportToolBar.Group"))
+		SNew(SHorizontalBox)
+		+SHorizontalBox::Slot()
+		.VAlign(VAlign_Center)
+		.AutoWidth()
 		[
-			SNew(SHorizontalBox)
-
-			+SHorizontalBox::Slot()
-			.VAlign(VAlign_Center)
-			.Padding(2.0f,0.0f)
-			.AutoWidth()
-			[
-				ToggleControl.ToSharedRef()
-			]
+			ToggleControl
+		]
 
-			+SHorizontalBox::Slot()
-			.Padding(2.0f, 0.0f)
-			.VAlign(VAlign_Center)
-			.AutoWidth()
-			[
-				MenuAnchor.ToSharedRef()
-			]
+		+SHorizontalBox::Slot()
+		.AutoWidth()
+		[
+			MenuAnchor.ToSharedRef()
 		]
 	];
 }
diff --git a/Engine/Source/Runtime/SlateCore/Private/Styling/StarshipCoreStyle.cpp b/Engine/Source/Runtime/SlateCore/Private/Styling/StarshipCoreStyle.cpp
--- a/Engine/Source/Runtime/SlateCore/Private/Styling/StarshipCoreStyle.cpp
+++ b/Engine/Source/Runtime/SlateCore/Private/Styling/StarshipCoreStyle.cpp
@@ -1882,7 +1882,7 @@ void FStarshipCoreStyle::SetupMultiboxStyles(TSharedRef<FStyle>& Style)
 			.SetButtonPadding(FMargin(4.0f, 0.0f))
 			.SetCheckBoxPadding(FMargin(8.0f, 0.0f))
 			.SetSeparatorBrush(FSlateColorBrush(FStyleColors::Recessed))
-			.SetSeparatorPadding(FMargin(8.f, -4.0f))
+			.SetSeparatorPadding(FMargin(4.f, -4.0f))
 			.SetLabelStyle(FTextBlockStyle(NormalText))
 			.SetComboButtonStyle(SlimToolBarComboButton)
 			.SetLabelPadding(FMargin(4.f, 0.f, 0.f, 0.f))
