Adjust sequencer and curve editor toolbar paddings to give more room back

diff --git a/Engine/Source/Editor/EditorStyle/Private/StarshipStyle.cpp b/Engine/Source/Editor/EditorStyle/Private/StarshipStyle.cpp
--- a/Engine/Source/Editor/EditorStyle/Private/StarshipStyle.cpp
+++ b/Engine/Source/Editor/EditorStyle/Private/StarshipStyle.cpp
@@ -1954,11 +1954,11 @@ void FStarshipEditorStyle::FStyle::SetupSequencerStyles()
 	{
 		FToolBarStyle SequencerToolbar = FStarshipCoreStyle::GetCoreStyle().GetWidgetStyle<FToolBarStyle>("AssetEditorToolbar");
 
-		SequencerToolbar.SetButtonPadding(       FMargin(6.0f, 0.0f));
-		SequencerToolbar.SetCheckBoxPadding(     FMargin(2.0f, 0.0f));
-		SequencerToolbar.SetComboButtonPadding(  FMargin(2.0f, 0.0f));
-		SequencerToolbar.SetIndentedBlockPadding(FMargin(2.0f, 0.0f));
-		SequencerToolbar.SetBlockPadding(        FMargin(2.0f, 0.0f));
+		SequencerToolbar.SetButtonPadding(       FMargin(0.0f, 0.0f));
+		SequencerToolbar.SetCheckBoxPadding(     FMargin(0.0f, 0.0f));
+		SequencerToolbar.SetComboButtonPadding(  FMargin(0.0f, 0.0f));
+		SequencerToolbar.SetIndentedBlockPadding(FMargin(0.0f, 0.0f));
+		SequencerToolbar.SetBlockPadding(        FMargin(0.0f, 0.0f));
 		SequencerToolbar.SetSeparatorPadding(    FMargin(6.0f, 0.0f));
 
 		Set("SequencerToolBar", SequencerToolbar);
@@ -7022,7 +7022,7 @@ void FStarshipEditorStyle::FStyle::SetupToolkitStyles()
 	{
 		FToolBarStyle CurveEditorToolbar = FStarshipCoreStyle::GetCoreStyle().GetWidgetStyle<FToolBarStyle>("AssetEditorToolbar");
 
-		CurveEditorToolbar.SetButtonPadding(       FMargin(6.0f, 0.0f));
+		CurveEditorToolbar.SetButtonPadding(       FMargin(0.0f, 0.0f));
 		CurveEditorToolbar.SetCheckBoxPadding(     FMargin(0.0f, 0.0f));
 		CurveEditorToolbar.SetComboButtonPadding(  FMargin(0.0f, 0.0f));
 		CurveEditorToolbar.SetIndentedBlockPadding(FMargin(0.0f, 0.0f));
