Fix level editor toolbar and viewport toolbar sizing and style issues

diff --git a/Engine/Source/Editor/EditorStyle/Private/StarshipStyle.cpp b/Engine/Source/Editor/EditorStyle/Private/StarshipStyle.cpp
--- a/Engine/Source/Editor/EditorStyle/Private/StarshipStyle.cpp
+++ b/Engine/Source/Editor/EditorStyle/Private/StarshipStyle.cpp
@@ -2286,7 +2286,7 @@ void FStarshipEditorStyle::FStyle::SetupViewportStyles()
 	{
 		FToolBarStyle ViewportToolbarStyle = FStarshipCoreStyle::GetCoreStyle().GetWidgetStyle<FToolBarStyle>("SlimToolBar");
 
-		FMargin ViewportMargin(4.f, 4.f);
+		FMargin ViewportMargin(6.f, 4.f, 3.f, 4.f);
 
 		const FCheckBoxStyle ViewportToggleButton = FCheckBoxStyle()
 			.SetCheckBoxType(ESlateCheckBoxType::ToggleButton)
@@ -2303,7 +2303,12 @@ void FStarshipEditorStyle::FStyle::SetupViewportStyles()
 			.SetCheckedHoveredForegroundColor(FStyleColors::ForegroundHover)
 			.SetPadding(0);
 
-		FSlateRoundedBoxBrush* ViewportGroupBrush = new FSlateRoundedBoxBrush(FStyleColors::Dropdown, 12.f, FStyleColors::Black, 1.0);
+
+		FLinearColor ToolbarBackgroundColor = FStyleColors::Dropdown.GetSpecifiedColor();
+		ToolbarBackgroundColor.A = .80f;
+
+
+		FSlateRoundedBoxBrush* ViewportGroupBrush = new FSlateRoundedBoxBrush(ToolbarBackgroundColor, 12.f, FLinearColor(0,0,0,.8), 1.0);
 		Set("EditorViewportToolBar.Group", ViewportGroupBrush);
 
 		FButtonStyle ViewportMenuButton = FButtonStyle()
@@ -2413,9 +2418,6 @@ void FStarshipEditorStyle::FStyle::SetupViewportStyles()
 			.SetDisabledForeground(FStyleColors::AccentYellow);
 		Set("EditorViewportToolBar.WarningButton", ViewportMenuWarningButton);
 
-		FLinearColor ToolbarBackgroundColor = FStyleColors::Foldout.GetSpecifiedColor();
-		ToolbarBackgroundColor.A = .75f;
-
 		Set("EditorViewportToolBar.Background", new FSlateNoResource());
 		Set("EditorViewportToolBar.OptionsDropdown", new IMAGE_BRUSH_SVG("Starship/EditorViewport/menu", Icon16x16));
 
@@ -5428,8 +5430,8 @@ void FStarshipEditorStyle::FStyle::SetupPersonaStyle()
 			.SetHovered(FSlateRoundedBoxBrush(FStyleColors::Dropdown, FVector4(4.0f, 0.0f, 0.0f, 4.0f)))
 			.SetPressed(FSlateRoundedBoxBrush(FStyleColors::Dropdown, FVector4(4.0f, 0.0f, 0.0f, 4.0f)))
 			.SetDisabled(FSlateRoundedBoxBrush(FStyleColors::Dropdown, FVector4(4.0f, 0.0f, 0.0f, 4.0f)))
-			.SetNormalPadding(FMargin(8.f, 4.f, 6.f, 4.f))
-			.SetPressedPadding(FMargin(8.f, 4.f, 6.f, 4.f));
+			.SetNormalPadding(FMargin(8.f, 2.f, 6.f, 2.f))
+			.SetPressedPadding(FMargin(8.f, 2.f, 6.f, 2.f));
 
 
 		MainToolbarLeftButton.SetButtonStyle(LeftToolbarButton);
@@ -5462,8 +5464,8 @@ void FStarshipEditorStyle::FStyle::SetupPersonaStyle()
 			.SetHovered(FSlateRoundedBoxBrush(FStyleColors::Dropdown, FVector4(0.0f, 0.0f, 0.0f, 0.0f)))
 			.SetPressed(FSlateRoundedBoxBrush(FStyleColors::Dropdown, FVector4(0.0f, 0.0f, 0.0f, 0.0f)))
 			.SetDisabled(FSlateRoundedBoxBrush(FStyleColors::Dropdown, FVector4(0.0f, 0.0f, 0.0f, 0.0f)))
-			.SetNormalPadding(FMargin(2.f, 4.f, 6.f, 4.f))
-			.SetPressedPadding(FMargin(2.f, 4.f, 6.f, 4.f));
+			.SetNormalPadding(FMargin(2.f, 2.f, 6.f, 2.f))
+			.SetPressedPadding(FMargin(2.f, 2.f, 6.f, 2.f));
 
 		MainToolbarCenterButton.SetButtonPadding(0.0f);
 		MainToolbarCenterButton.SetButtonStyle(CenterToolbarButton);
@@ -5496,8 +5498,8 @@ void FStarshipEditorStyle::FStyle::SetupPersonaStyle()
 			.SetHovered(FSlateRoundedBoxBrush(FStyleColors::Dropdown, FVector4(0.0f, 4.0f, 4.0f, 0.0f)))
 			.SetPressed(FSlateRoundedBoxBrush(FStyleColors::Dropdown, FVector4(0.0f, 4.0f, 4.0f, 0.0f)))
 			.SetDisabled(FSlateRoundedBoxBrush(FStyleColors::Dropdown, FVector4(0.0f, 4.0f, 4.0f, 0.0f)))
-			.SetNormalPadding(FMargin(2.f, 4.f, 8.f, 4.f))
-			.SetPressedPadding(FMargin(2.f, 4.f, 8.f, 4.f));
+			.SetNormalPadding(FMargin(2.f, 2.f, 8.f, 2.f))
+			.SetPressedPadding(FMargin(2.f, 2.f, 8.f, 2.f));
 
 		MainToolbarRightButton.SetButtonStyle(RightToolbarButton);
 		MainToolbarRightButton.SetButtonPadding(FMargin(0.0f, 0.0f, 8.0f, 0.0f));
@@ -5513,8 +5515,8 @@ void FStarshipEditorStyle::FStyle::SetupPersonaStyle()
 			.SetHovered(FSlateRoundedBoxBrush(FStyleColors::Dropdown, FVector4(0.0f, 4.0f, 4.0f, 0.0f)))
 			.SetPressed(FSlateRoundedBoxBrush(FStyleColors::Dropdown, FVector4(0.0f, 4.0f, 4.0f, 0.0f)))
 			.SetDisabled(FSlateRoundedBoxBrush(FStyleColors::Dropdown, FVector4(0.0f, 4.0f, 4.0f, 0.0f)))
-			.SetNormalPadding(FMargin(7.f, 4.f, 6.f, 4.f))
-			.SetPressedPadding(FMargin(7.f, 4.f, 6.f, 4.f));
+			.SetNormalPadding(FMargin(7.f, 2.f, 6.f, 2.f))
+			.SetPressedPadding(FMargin(7.f, 2.f, 6.f, 2.f));
 
 		FComboButtonStyle PlayToolbarComboButton = FComboButtonStyle(FStarshipCoreStyle::GetCoreStyle().GetWidgetStyle<FComboButtonStyle>("ComboButton"))
 			.SetDownArrowPadding(FMargin(-19.f, 0.f, 2.f, 0.f))
diff --git a/Engine/Source/Editor/UnrealEd/Private/STransformViewportToolbar.cpp b/Engine/Source/Editor/UnrealEd/Private/STransformViewportToolbar.cpp
--- a/Engine/Source/Editor/UnrealEd/Private/STransformViewportToolbar.cpp
+++ b/Engine/Source/Editor/UnrealEd/Private/STransformViewportToolbar.cpp
@@ -262,6 +262,7 @@ TSharedRef< SWidget > STransformViewportToolBar::MakeTransformToolBar( const TSh
 				.ToggleButtonToolTip(Command->GetDescription())
 				.MenuButtonToolTip(LOCTEXT("LocationGridSnap_ToolTip", "Set the Position Grid Snap value"))
 				.Icon(Command->GetIcon())
+				.MinDesiredButtonWidth(24.0f)
 				.ParentToolBar(SharedThis(this)),
 			PositionSnapName, 
 			false, 
@@ -390,6 +391,7 @@ TSharedRef< SWidget > STransformViewportToolBar::MakeTransformToolBar( const TSh
 				.ToggleButtonToolTip(Command->GetDescription())
 				.MenuButtonToolTip(LOCTEXT("ScaleGridSnap_ToolTip", "Set scaling options"))
 				.Icon(Command->GetIcon())
+				.MinDesiredButtonWidth(24.0f)
 				.ParentToolBar(SharedThis(this)),
 			ScaleSnapName,
 			false, 
diff --git a/Engine/Source/Editor/UnrealEd/Private/SViewportToolBarComboMenu.cpp b/Engine/Source/Editor/UnrealEd/Private/SViewportToolBarComboMenu.cpp
--- a/Engine/Source/Editor/UnrealEd/Private/SViewportToolBarComboMenu.cpp
+++ b/Engine/Source/Editor/UnrealEd/Private/SViewportToolBarComboMenu.cpp
@@ -44,6 +44,7 @@ void SViewportToolBarComboMenu::Construct( const FArguments& InArgs )
 			.ToolTipText(InArgs._MenuButtonToolTip)
 			.OnClicked(this, &SViewportToolBarComboMenu::OnMenuClicked)
 			.VAlign(VAlign_Center)
+			.HAlign(HAlign_Center)
 			[
 				SNew(STextBlock)
 				.TextStyle(&LabelStyle)
