Center icons in viewport toolbar rounded buttons

diff --git a/Engine/Source/Editor/EditorStyle/Private/StarshipStyle.cpp b/Engine/Source/Editor/EditorStyle/Private/StarshipStyle.cpp
--- a/Engine/Source/Editor/EditorStyle/Private/StarshipStyle.cpp
+++ b/Engine/Source/Editor/EditorStyle/Private/StarshipStyle.cpp
@@ -2319,8 +2319,8 @@ void FStarshipEditorStyle::FStyle::SetupViewportStyles()
 			.SetHoveredForeground(FStyleColors::ForegroundHover)
 			.SetPressedForeground(FStyleColors::ForegroundHover)
 			.SetDisabledForeground(FStyleColors::Foreground)
-			.SetNormalPadding(ViewportMargin)
-			.SetPressedPadding(ViewportMargin);
+			.SetNormalPadding(FMargin(4.0f, 4.0f, 3.0f, 4.0f))
+			.SetPressedPadding(FMargin(4.0f, 4.0f, 3.0f, 4.0f));
 		Set("EditorViewportToolBar.Button", ViewportMenuButton);
 
 		const FCheckBoxStyle ViewportMenuToggleLeftButtonStyle = FCheckBoxStyle(ViewportToggleButton)
