More viewport toolbar icon centering

diff --git a/Engine/Source/Editor/EditorStyle/Private/StarshipStyle.cpp b/Engine/Source/Editor/EditorStyle/Private/StarshipStyle.cpp
--- a/Engine/Source/Editor/EditorStyle/Private/StarshipStyle.cpp
+++ b/Engine/Source/Editor/EditorStyle/Private/StarshipStyle.cpp
@@ -2286,7 +2286,9 @@ void FStarshipEditorStyle::FStyle::SetupViewportStyles()
 	{
 		FToolBarStyle ViewportToolbarStyle = FStarshipCoreStyle::GetCoreStyle().GetWidgetStyle<FToolBarStyle>("SlimToolBar");
 
-		FMargin ViewportMargin(6.f, 4.f, 3.f, 4.f);
+		FMargin ViewportMarginLeft(6.f, 4.f, 3.f, 4.f);
+		FMargin ViewportMarginCenter(6.f, 4.f, 3.f, 4.f);
+		FMargin ViewportMarginRight(4.f, 4.f, 5.f, 4.f);
 
 		const FCheckBoxStyle ViewportToggleButton = FCheckBoxStyle()
 			.SetCheckBoxType(ESlateCheckBoxType::ToggleButton)
@@ -2331,7 +2333,7 @@ void FStarshipEditorStyle::FStyle::SetupViewportStyles()
 			.SetCheckedHoveredImage(  BOX_BRUSH("Starship/EditorViewport/ToolBarLeftGroup", 12.f/25.f, FStyleColors::PrimaryHover))
 			.SetCheckedPressedImage(  BOX_BRUSH("Starship/EditorViewport/ToolBarLeftGroup", 12.f/25.f, FStyleColors::PrimaryPress))
 			.SetCheckedImage(         BOX_BRUSH("Starship/EditorViewport/ToolBarLeftGroup", 12.f/25.f, FStyleColors::Primary))
-			.SetPadding(ViewportMargin);
+			.SetPadding(ViewportMarginLeft);
 		Set("EditorViewportToolBar.ToggleButton.Start", ViewportMenuToggleLeftButtonStyle);
 
 		const FCheckBoxStyle ViewportMenuToggleMiddleButtonStyle = FCheckBoxStyle(ViewportToggleButton)
@@ -2342,7 +2344,7 @@ void FStarshipEditorStyle::FStyle::SetupViewportStyles()
 			.SetCheckedHoveredImage(  BOX_BRUSH("Starship/EditorViewport/ToolBarMiddleGroup", 12.f/25.f, FStyleColors::PrimaryHover))
 			.SetCheckedPressedImage(  BOX_BRUSH("Starship/EditorViewport/ToolBarMiddleGroup", 12.f/25.f, FStyleColors::PrimaryPress))
 			.SetCheckedImage(         BOX_BRUSH("Starship/EditorViewport/ToolBarMiddleGroup", 12.f/25.f, FStyleColors::Primary))
-			.SetPadding(ViewportMargin);
+			.SetPadding(ViewportMarginCenter);
 		Set("EditorViewportToolBar.ToggleButton.Middle", ViewportMenuToggleMiddleButtonStyle);
 
 		const FCheckBoxStyle ViewportMenuToggleRightButtonStyle = FCheckBoxStyle(ViewportToggleButton)
@@ -2353,7 +2355,7 @@ void FStarshipEditorStyle::FStyle::SetupViewportStyles()
 			.SetCheckedHoveredImage(  BOX_BRUSH("Starship/EditorViewport/ToolBarRightGroup", 12.f/25.f, FStyleColors::PrimaryHover))
 			.SetCheckedPressedImage(  BOX_BRUSH("Starship/EditorViewport/ToolBarRightGroup", 12.f/25.f, FStyleColors::PrimaryPress))
 			.SetCheckedImage(         BOX_BRUSH("Starship/EditorViewport/ToolBarRightGroup", 12.f/25.f, FStyleColors::Primary))
-			.SetPadding(ViewportMargin);
+			.SetPadding(ViewportMarginRight);
 		Set("EditorViewportToolBar.ToggleButton.End", ViewportMenuToggleRightButtonStyle);
 
 		// We want a background-less version as the ComboMenu has its own unified background
@@ -2383,7 +2385,7 @@ void FStarshipEditorStyle::FStyle::SetupViewportStyles()
 			.SetCheckedForegroundColor(FStyleColors::Foreground)
 			.SetCheckedPressedForegroundColor(FStyleColors::ForegroundHover)
 			.SetCheckedHoveredForegroundColor(FStyleColors::ForegroundHover)
-			.SetPadding(ViewportMargin);
+			.SetPadding(FMargin(4.0f, 4.0f, 3.0f, 4.0f));
 		Set("EditorViewportToolBar.MaximizeRestoreButton", MaximizeRestoreButton);
 
 		Set("EditorViewportToolBar.Heading.Padding", FMargin(4.f));
@@ -2392,7 +2394,7 @@ void FStarshipEditorStyle::FStyle::SetupViewportStyles()
 		// SComboBox 
 		FComboButtonStyle ViewportComboButton = FComboButtonStyle()
 			.SetButtonStyle(ViewportMenuButton)
-			.SetContentPadding(ViewportMargin);
+			.SetContentPadding(ViewportMarginCenter);
 
 		ViewportToolbarStyle
 			.SetBackground(FSlateNoResource())
