New state machine icons

diff --git a/Engine/Source/Editor/AnimationBlueprintEditor/Private/AnimationStateNodes/SGraphNodeAnimStateAlias.cpp b/Engine/Source/Editor/AnimationBlueprintEditor/Private/AnimationStateNodes/SGraphNodeAnimStateAlias.cpp
--- a/Engine/Source/Editor/AnimationBlueprintEditor/Private/AnimationStateNodes/SGraphNodeAnimStateAlias.cpp
+++ b/Engine/Source/Editor/AnimationBlueprintEditor/Private/AnimationStateNodes/SGraphNodeAnimStateAlias.cpp
@@ -97,8 +97,7 @@ FText SGraphNodeAnimStateAlias::GetPreviewCornerText() const
 
 const FSlateBrush* SGraphNodeAnimStateAlias::GetNameIcon() const
 {
-	// @todo: Change this to a unique icon and/or look for better visual distinction between state/conduit/alias.
-	return FEditorStyle::GetBrush(TEXT("Graph.TransitionNode.Icon"));
+	return FEditorStyle::GetBrush(TEXT("Graph.AliasNode.Icon"));
 }
 
 TSharedPtr<SToolTip> SGraphNodeAnimStateAlias::GetComplexTooltip()
diff --git a/Engine/Source/Editor/EditorStyle/Private/StarshipStyle.cpp b/Engine/Source/Editor/EditorStyle/Private/StarshipStyle.cpp
--- a/Engine/Source/Editor/EditorStyle/Private/StarshipStyle.cpp
+++ b/Engine/Source/Editor/EditorStyle/Private/StarshipStyle.cpp
@@ -3578,8 +3578,9 @@ void FStarshipEditorStyle::FStyle::SetupGraphEditorStyles()
 				Set( "Graph.StateNode.Body", new BOX_BRUSH( "/Persona/StateMachineEditor/StateNode_Node_Body", FMargin(16.f/64.f, 25.f/64.f, 16.f/64.f, 16.f/64.f) ) );
 				Set( "Graph.StateNode.ColorSpill", new BOX_BRUSH( "/Persona/StateMachineEditor/StateNode_Node_ColorSpill", FMargin(4.0f/64.0f, 4.0f/32.0f) ) );
 
-				Set( "Graph.StateNode.Icon", new IMAGE_BRUSH( "/Persona/StateMachineEditor/State_Node_Icon_32x", Icon16x16 ) );
-				Set( "Graph.ConduitNode.Icon", new IMAGE_BRUSH( "/Persona/StateMachineEditor/Conduit_Node_Icon_32x", Icon16x16 ) );
+				Set( "Graph.StateNode.Icon", new IMAGE_BRUSH_SVG( "Starship/AnimationBlueprintEditor/AnimationState", Icon16x16 ) );
+				Set( "Graph.ConduitNode.Icon", new IMAGE_BRUSH_SVG( "Starship/AnimationBlueprintEditor/AnimationConduit", Icon16x16 ) );
+				Set( "Graph.AliasNode.Icon", new IMAGE_BRUSH_SVG( "Starship/AnimationBlueprintEditor/AnimationAlias", Icon16x16 ) );
 
 				Set( "Graph.StateNode.Pin.BackgroundHovered", new BOX_BRUSH( "/Persona/StateMachineEditor/StateNode_Pin_HoverCue", FMargin(12.0f/64.0f,12.0f/64.0f,12.0f/64.0f,12.0f/64.0f)));
 				Set( "Graph.StateNode.Pin.Background", new FSlateNoResource() );
