Update styling on mode palette toolbars

diff --git a/Engine/Source/Editor/EditorStyle/Private/StarshipStyle.cpp b/Engine/Source/Editor/EditorStyle/Private/StarshipStyle.cpp
--- a/Engine/Source/Editor/EditorStyle/Private/StarshipStyle.cpp
+++ b/Engine/Source/Editor/EditorStyle/Private/StarshipStyle.cpp
@@ -1633,7 +1633,10 @@ void FStarshipEditorStyle::FStyle::SetupGeneralStyles()
 
 		FToolBarStyle PaletteToolBarStyle = FStarshipCoreStyle::GetCoreStyle().GetWidgetStyle<FToolBarStyle>("SlimToolBar");
 
-		PaletteToolBarStyle.SetLabelStyle( GetParentStyle()->GetWidgetStyle<FTextBlockStyle>("SmallText") );
+		FTextBlockStyle PaletteToolbarLabelStyle = FTextBlockStyle(GetParentStyle()->GetWidgetStyle<FTextBlockStyle>("SmallText"));
+		PaletteToolbarLabelStyle.SetOverflowPolicy(ETextOverflowPolicy::Ellipsis);
+
+		PaletteToolBarStyle.SetLabelStyle(PaletteToolbarLabelStyle);
 		
 		PaletteToolBarStyle.SetBackground(FSlateColorBrush(FStyleColors::Recessed));
 
@@ -1669,37 +1672,8 @@ void FStarshipEditorStyle::FStyle::SetupGeneralStyles()
 
 		Set("PaletteToolBar.ExpandableAreaHeader", new FSlateRoundedBoxBrush(FStyleColors::Dropdown, FVector4(4.0, 4.0, 0.0, 0.0)));
 		Set("PaletteToolBar.ExpandableAreaBody", new FSlateRoundedBoxBrush(FStyleColors::Recessed, FVector4(0.0, 0.0, 4.0, 4.0)));
-		
-		FLinearColor ButtonHoverColor(FColor(81, 81, 81));
-
-			/* Create style for "ToolBar.ToggleButton" widget ... */
-		const FCheckBoxStyle ToolBarToggleButtonCheckBoxStyle = FCheckBoxStyle()
-				.SetCheckBoxType( ESlateCheckBoxType::ToggleButton )
-				.SetUncheckedImage(        FSlateNoResource() )
-				.SetUncheckedPressedImage( FSlateNoResource() )
-				.SetUncheckedHoveredImage( FSlateNoResource() )
-				.SetCheckedImage(          FSlateNoResource() )
-				.SetCheckedHoveredImage(   FSlateNoResource() )
-				.SetCheckedPressedImage(   FSlateNoResource() )
-				.SetForegroundColor(       FStyleColors::Foreground)
-				.SetHoveredForegroundColor(FStyleColors::ForegroundHover)
-				.SetPressedForegroundColor(FStyleColors::ForegroundHover)
-				.SetCheckedForegroundColor(FStyleColors::Primary)
-				.SetCheckedHoveredForegroundColor(FStyleColors::PrimaryHover)
-				.SetPadding(FMargin(2.f, 6.f));
-		
-		PaletteToolBarStyle.SetToggleButtonStyle(ToolBarToggleButtonCheckBoxStyle);
-		PaletteToolBarStyle.SetButtonStyle(FButtonStyle(Button)
-			.SetNormal(  FSlateNoResource() )
-			.SetPressed( FSlateNoResource() )
-			.SetHovered( FSlateNoResource() )
-			.SetNormalForeground(         FStyleColors::Foreground)
-			.SetHoveredForeground(        FStyleColors::ForegroundHover)
-			.SetPressedForeground(        FStyleColors::ForegroundHover)
-			.SetNormalPadding(FMargin(2.0f, 6.0f))
-			.SetPressedPadding(FMargin(2.0f, 6.0f))
-		);
 
+	
 		Set("PaletteToolBar", PaletteToolBarStyle);
 
 		Set("EditorModesPanel.CategoryFontStyle", DEFAULT_FONT( "Bold", 10 ));
