Updated Animation Attribute icons for visualizing within the Anim Graph.

diff --git a/Engine/Source/Editor/EditorStyle/Private/StarshipStyle.cpp b/Engine/Source/Editor/EditorStyle/Private/StarshipStyle.cpp
--- a/Engine/Source/Editor/EditorStyle/Private/StarshipStyle.cpp
+++ b/Engine/Source/Editor/EditorStyle/Private/StarshipStyle.cpp
@@ -5838,7 +5838,7 @@ void FStarshipEditorStyle::FStyle::SetupPersonaStyle()
 		Set("AnimGraph.Attribute.Attributes.Color", FStyleColors::AccentYellow);
 		Set("AnimGraph.Attribute.Sync.Icon", new IMAGE_BRUSH_SVG("Starship/AnimationBlueprintEditor/AnimationGraphSync", Icon16x16));
 		Set("AnimGraph.Attribute.Sync.Color", FStyleColors::AccentBlue);
-		Set("AnimGraph.Attribute.RootMotionDelta.Icon", new IMAGE_BRUSH_SVG("Starship/AnimationBlueprintEditor/AnimationGraphAttributes", Icon16x16));
+		Set("AnimGraph.Attribute.RootMotionDelta.Icon", new IMAGE_BRUSH_SVG("Starship/AnimationBlueprintEditor/AnimationGraphRootMotionDelta", Icon16x16));
 		Set("AnimGraph.Attribute.RootMotionDelta.Color", FStyleColors::AccentGreen);
 		Set("AnimGraph.Attribute.InertialBlending.Icon", new IMAGE_BRUSH_SVG("Starship/AnimationBlueprintEditor/AnimationGraphInertialBlending", Icon16x16));
 		Set("AnimGraph.Attribute.InertialBlending.Color", FStyleColors::AccentOrange);
