Adjust padding on mode palette toolbars so that none of the existing icons clip

diff --git a/Engine/Source/Editor/EditorStyle/Private/StarshipStyle.cpp b/Engine/Source/Editor/EditorStyle/Private/StarshipStyle.cpp
--- a/Engine/Source/Editor/EditorStyle/Private/StarshipStyle.cpp
+++ b/Engine/Source/Editor/EditorStyle/Private/StarshipStyle.cpp
@@ -1647,7 +1647,7 @@ void FStarshipEditorStyle::FStyle::SetupGeneralStyles()
 		PaletteToolBarStyle.SetComboButtonPadding(  FMargin(0.0f, 0.0f));
 		PaletteToolBarStyle.SetIndentedBlockPadding(FMargin(0.0f, 0.0f));
 		PaletteToolBarStyle.SetBlockPadding(        FMargin(0.0f, 0.0f));
-		PaletteToolBarStyle.ToggleButton.SetPadding(FMargin(2.0f, 6.0f));
+		PaletteToolBarStyle.ToggleButton.SetPadding(FMargin(0.0f, 6.0f));
 		PaletteToolBarStyle.ButtonStyle.SetNormalPadding(FMargin(2.0f, 6.0f));
 		PaletteToolBarStyle.ButtonStyle.SetPressedPadding(FMargin(2.0f, 6.0f));
 
diff --git a/Engine/Source/Runtime/Slate/Private/Framework/MultiBox/MultiBox.cpp b/Engine/Source/Runtime/Slate/Private/Framework/MultiBox/MultiBox.cpp
--- a/Engine/Source/Runtime/Slate/Private/Framework/MultiBox/MultiBox.cpp
+++ b/Engine/Source/Runtime/Slate/Private/Framework/MultiBox/MultiBox.cpp
@@ -1022,11 +1022,11 @@ void SMultiBoxWidget::BuildMultiBoxWidget()
 			[
 				SAssignNew(UniformToolbarPanel, SUniformWrapPanel)
 				.HAlign(HAlign_Left)
-				.MinDesiredSlotWidth(46.f)
+				.MinDesiredSlotWidth(50.f)
 				.MinDesiredSlotHeight(43.f)
-				.MaxDesiredSlotWidth(46.f)
+				.MaxDesiredSlotWidth(50.f)
 				.MaxDesiredSlotHeight(43.f)
-				.SlotPadding(FMargin(1.f, 1.f))
+				.SlotPadding(FMargin(2.f, 1.f))
 			];
 		}
 		break;
