[Slate] Allow the use of user log colors outside of WITH_EDITOR

diff --git a/Engine/Source/Editor/EditorStyle/Private/StarshipStyle.cpp b/Engine/Source/Editor/EditorStyle/Private/StarshipStyle.cpp
--- a/Engine/Source/Editor/EditorStyle/Private/StarshipStyle.cpp
+++ b/Engine/Source/Editor/EditorStyle/Private/StarshipStyle.cpp
@@ -47,7 +47,9 @@ void FStarshipEditorStyle::Initialize()
 	FPropertyEditorModule& PropertyEditorModule = FModuleManager::Get().GetModuleChecked<FPropertyEditorModule>("PropertyEditor");
 	PropertyEditorModule.RegisterCustomClassLayout("EditorStyleSettings", FOnGetDetailCustomizationInstance::CreateStatic(&FEditorStyleSettingsCustomization::MakeInstance));
 	PropertyEditorModule.RegisterCustomPropertyTypeLayout("StyleColorList", FOnGetPropertyTypeCustomizationInstance::CreateStatic(&FStyleColorListCustomization::MakeInstance));
+#endif // #if WITH_EDITOR
 
+#if WITH_EDITOR || IS_PROGRAM
 	// Selection highlight
 	USlateThemeManager::Get().SetDefaultColor(EStyleColor::User2, USlateThemeManager::Get().GetColor(EStyleColor::Highlight));
 	USlateThemeManager::Get().SetColorDisplayName(EStyleColor::User2, LOCTEXT("UserColor_OutputLogHighlight", "Log Highlight"));
@@ -57,8 +59,7 @@ void FStarshipEditorStyle::Initialize()
 	// Command
 	USlateThemeManager::Get().SetDefaultColor(EStyleColor::User4, USlateThemeManager::Get().GetColor(EStyleColor::AccentGreen));
 	USlateThemeManager::Get().SetColorDisplayName(EStyleColor::User4, LOCTEXT("UserColor_OutputLogCommand", "Log Command"));
-
-#endif
+#endif // #if WITH_EDITOR || IS_PROGRAM
 
 	const FString ThemesSubDir = TEXT("Slate/Themes");
 
