Replace all instances of old Source Control Icon with the new one

diff --git a/Engine/Source/Editor/EditorStyle/Private/StarshipStyle.cpp b/Engine/Source/Editor/EditorStyle/Private/StarshipStyle.cpp
--- a/Engine/Source/Editor/EditorStyle/Private/StarshipStyle.cpp
+++ b/Engine/Source/Editor/EditorStyle/Private/StarshipStyle.cpp
@@ -1349,7 +1349,7 @@ void FStarshipEditorStyle::FStyle::SetupGeneralStyles()
 
 
 		Set( "MainFrame.ChooseFilesToSave",       new IMAGE_BRUSH_SVG( "Starship/Common/SaveChoose", Icon16x16 ) );
-		Set( "MainFrame.ConnectToSourceControl",  new IMAGE_BRUSH_SVG( "Starship/Common/SourceControl", Icon16x16 ) );
+		Set( "MainFrame.ConnectToSourceControl",  new CORE_IMAGE_BRUSH_SVG("Starship/SourceControl/SourceControl", Icon16x16) );
 		Set( "MainFrame.OpenMarketplace",			new IMAGE_BRUSH_SVG("Starship/MainToolbar/marketplace", Icon16x16));
 
 		Set( "MainFrame.DebugTools.SmallFont", DEFAULT_FONT( "Regular", 8 ) );
@@ -7360,7 +7360,7 @@ void FStarshipEditorStyle::FStyle::SetupSourceControlStyles()
 		Set("SourceControl.ShelvedChangelist", new CORE_IMAGE_BRUSH_SVG("Starship/SourceControl/SCC_Changelist", Icon16x16, FStyleColors::AccentBlue));
 		Set("SourceControl.UncontrolledChangelist", new CORE_IMAGE_BRUSH_SVG("Starship/SourceControl/SCC_Changelist", Icon32x32, FStyleColors::AccentGray));
 		Set("SourceControl.UncontrolledChangelist_Small", new CORE_IMAGE_BRUSH_SVG("Starship/SourceControl/SCC_Changelist", Icon16x16, FStyleColors::AccentGray));
-		Set("SourceControl.OfflineFile_Small", new CORE_IMAGE_BRUSH_SVG("Starship/SourceControl/icon_SCC_Connect", Icon16x16, FStyleColors::AccentRed));
+		Set("SourceControl.OfflineFile_Small", new CORE_IMAGE_BRUSH_SVG("Starship/SourceControl/SourceControl", Icon16x16, FStyleColors::AccentRed));
 
 		Set("SourceControl.Add", new CORE_IMAGE_BRUSH_SVG( "Starship/SourceControl/SCC_ContentAdd",Icon16x16));
 		Set("SourceControl.Edit", new CORE_IMAGE_BRUSH_SVG( "Starship/SourceControl/SCC_CheckedOut", Icon16x16));
@@ -7381,7 +7381,7 @@ void FStarshipEditorStyle::FStyle::SetupSourceControlStyles()
 		Set("SourceControl.Actions.Submit", new CORE_IMAGE_BRUSH_SVG( "Starship/Common/arrow-left", Icon16x16));
 		Set("SourceControl.Actions.Diff", new CORE_IMAGE_BRUSH_SVG( "Starship/SourceControl/SCC_Action_Diff", Icon16x16));
 		Set("SourceControl.Actions.Revert", new CORE_IMAGE_BRUSH_SVG( "Starship/SourceControl/icon_SCC_Revert", Icon16x16));
-		Set("SourceControl.Actions.Connect", new CORE_IMAGE_BRUSH_SVG( "Starship/SourceControl/icon_SCC_Connect", Icon16x16));
+		Set("SourceControl.Actions.Connect", new CORE_IMAGE_BRUSH_SVG( "Starship/SourceControl/SourceControl", Icon16x16));
 		Set("SourceControl.Actions.History", new CORE_IMAGE_BRUSH_SVG( "Starship/SourceControl/icon_SCC_History", Icon16x16));
 		Set("SourceControl.Actions.Add", new CORE_IMAGE_BRUSH_SVG("Starship/SourceControl/SCC_CheckedOut", Icon16x16));
 		Set("SourceControl.Actions.ChangeSettings", new CORE_IMAGE_BRUSH_SVG("Starship/SourceControl/icon_SCC_Change_Source_Control_Settings", Icon16x16));
