Details Panel Name Area: Fix button padding/spacing issues

diff --git a/Engine/Source/Editor/EditorStyle/Private/StarshipStyle.cpp b/Engine/Source/Editor/EditorStyle/Private/StarshipStyle.cpp
--- a/Engine/Source/Editor/EditorStyle/Private/StarshipStyle.cpp
+++ b/Engine/Source/Editor/EditorStyle/Private/StarshipStyle.cpp
@@ -612,18 +612,6 @@ void FStarshipEditorStyle::FStyle::SetupGeneralStyles()
 
 	}
 
-	FButtonStyle CompactButton = FButtonStyle(Button)
-		.SetNormalPadding(FMargin(4))
-		.SetPressedPadding(FMargin(4));
-	Set("CompactButton", CompactButton);
-
-	Set("CompactComboButton", FComboButtonStyle(GetWidgetStyle<FComboButtonStyle>("ComboButton"))
-		.SetButtonStyle(CompactButton)
-		.SetDownArrowPadding(FMargin(4, 0, 0, 0))
-		.SetContentPadding(FMargin(4, 0, 0, 0))
-	);
-
-
 	// Error Reporting
 	{
 		Set( "InfoReporting.BackgroundColor", FLinearColor(0.1f, 0.33f, 1.0f));
@@ -3283,6 +3271,17 @@ void FStarshipEditorStyle::FStyle::SetupPropertyEditorStyles()
 
 		Set("DetailsView.DropZone.Below", new BOX_BRUSH("Common/VerticalBoxDropZoneIndicator_Below", FMargin(10.0f / 16.0f, 0, 0, 10.0f / 16.0f), SelectionColor_Subdued));
 
+		FButtonStyle NameAreaButton = FButtonStyle(Button)
+			.SetNormalPadding(FMargin(6, 3))
+			.SetPressedPadding(FMargin(6, 3));
+		Set("DetailsView.NameAreaButton", NameAreaButton);
+
+		Set("DetailsView.NameAreaComboButton", FComboButtonStyle(GetWidgetStyle<FComboButtonStyle>("ComboButton"))
+			.SetButtonStyle(NameAreaButton)
+			.SetDownArrowPadding(FMargin(4, 0, 0, 0))
+			.SetContentPadding(FMargin(4, 0, 0, 0))
+		);
+
 	}
 	}
 
diff --git a/Engine/Source/Editor/PropertyEditor/Private/SDetailNameArea.cpp b/Engine/Source/Editor/PropertyEditor/Private/SDetailNameArea.cpp
--- a/Engine/Source/Editor/PropertyEditor/Private/SDetailNameArea.cpp
+++ b/Engine/Source/Editor/PropertyEditor/Private/SDetailNameArea.cpp
@@ -167,7 +167,7 @@ TSharedRef< SWidget > SDetailNameArea::BuildObjectNameArea( const TArray< TWeakO
 					SNew( SButton )
 					.ButtonStyle( FEditorStyle::Get(), "SimpleButton" )
 					.OnClicked(	OnLockButtonClicked )
-					.ContentPadding(0)
+					.ContentPadding(FMargin(4,2))
 					.HAlign(HAlign_Center)
 					.VAlign(VAlign_Center)
 					.ToolTipText( LOCTEXT("LockSelectionButton_ToolTip", "Locks the current selection into the Details panel") )
diff --git a/Engine/Source/Editor/SubobjectEditor/Private/SSubobjectInstanceEditor.cpp b/Engine/Source/Editor/SubobjectEditor/Private/SSubobjectInstanceEditor.cpp
--- a/Engine/Source/Editor/SubobjectEditor/Private/SSubobjectInstanceEditor.cpp
+++ b/Engine/Source/Editor/SubobjectEditor/Private/SSubobjectInstanceEditor.cpp
@@ -97,7 +97,7 @@ void SSubobjectInstanceEditor::Construct(const FArguments& InArgs)
 	.AutoWidth()
 	[
 		SNew(SButton)
-		.ButtonStyle(FAppStyle::Get(), "CompactButton")
+		.ButtonStyle(FAppStyle::Get(), "DetailsView.NameAreaButton")
 		.ContentPadding(0)
 		.HAlign(HAlign_Center)
 		.VAlign(VAlign_Center)
@@ -119,9 +119,10 @@ void SSubobjectInstanceEditor::Construct(const FArguments& InArgs)
 	.AutoWidth()
 	[
 		SNew(SComboButton)
-		.ComboButtonStyle(FAppStyle::Get(), "CompactComboButton")
+		.ComboButtonStyle(FAppStyle::Get(), "DetailsView.NameAreaComboButton")
 		.HAlign(HAlign_Center)
 		.VAlign(VAlign_Center)
+		.ContentPadding(FMargin(0, 2, 0, 1))
 		.AddMetaData<FTagMetaData>(FTagMetaData(TEXT("Actor.EditBlueprint")))
 		.Visibility(this, &SSubobjectInstanceEditor::GetEditBlueprintButtonVisibility)
 		.ToolTipText(LOCTEXT("EditActorBlueprint_Tooltip", "Edit the Blueprint for this Actor"))
