Update Unreal and Epic logos in the Help menu

diff --git a/Engine/Source/Editor/EditorStyle/Private/StarshipStyle.cpp b/Engine/Source/Editor/EditorStyle/Private/StarshipStyle.cpp
--- a/Engine/Source/Editor/EditorStyle/Private/StarshipStyle.cpp
+++ b/Engine/Source/Editor/EditorStyle/Private/StarshipStyle.cpp
@@ -4305,8 +4305,8 @@ void FStarshipEditorStyle::FStyle::SetupLevelEditorStyle()
 		Set("MainFrame.VisitForums",               new IMAGE_BRUSH_SVG("Starship/Common/WebBrowser", Icon16x16));
 		Set("MainFrame.VisitSearchForAnswersPage", new IMAGE_BRUSH_SVG("Starship/Common/WebBrowser", Icon16x16));
 		Set("MainFrame.VisitSupportWebSite",       new IMAGE_BRUSH_SVG("Starship/Common/Support", Icon16x16));
-		Set("MainFrame.VisitEpicGamesDotCom",      new IMAGE_BRUSH("Icons/Help/icon_Help_epic_16x", Icon16x16));
-		Set("MainFrame.AboutUnrealEd",             new IMAGE_BRUSH("Icons/Help/icon_Help_unreal_16x", Icon16x16));
+		Set("MainFrame.VisitEpicGamesDotCom",      new IMAGE_BRUSH_SVG("About/EpicGamesLogo", Icon16x16));
+		Set("MainFrame.AboutUnrealEd",             new IMAGE_BRUSH_SVG("About/UnrealLogo", Icon16x16));
 		Set("MainFrame.CreditsUnrealEd",           new IMAGE_BRUSH_SVG("Starship/Common/Credits", Icon16x16));
 
 		Set( "EditorViewport.SelectMode", new IMAGE_BRUSH_SVG("Starship/EditorViewport/select", Icon16x16) );
