Update TranslateRotateMode icons

diff --git a/Engine/Source/Editor/EditorStyle/Private/StarshipStyle.cpp b/Engine/Source/Editor/EditorStyle/Private/StarshipStyle.cpp
--- a/Engine/Source/Editor/EditorStyle/Private/StarshipStyle.cpp
+++ b/Engine/Source/Editor/EditorStyle/Private/StarshipStyle.cpp
@@ -4315,8 +4315,8 @@ void FStarshipEditorStyle::FStyle::SetupLevelEditorStyle()
 		Set( "EditorViewport.RotateMode", new IMAGE_BRUSH_SVG("Starship/EditorViewport/rotate", Icon16x16 ) );
 		Set( "EditorViewport.ScaleMode", new IMAGE_BRUSH_SVG( "Starship/EditorViewport/scale", Icon16x16 ) );
 
-		Set( "EditorViewport.TranslateRotateMode", new IMAGE_BRUSH( "Icons/icon_translate_rotate_40x", Icon20x20 ) );
-		Set( "EditorViewport.TranslateRotate2DMode", new IMAGE_BRUSH("Icons/icon_translate_rotate_2d_40x", Icon20x20));
+		Set( "EditorViewport.TranslateRotateMode", new IMAGE_BRUSH_SVG("Starship/EditorViewport/TranslateRotate3D", Icon16x16 ) );
+		Set( "EditorViewport.TranslateRotate2DMode", new IMAGE_BRUSH_SVG("Starship/EditorViewport/TranslateRotate2D", Icon16x16 ) );
 
 		Set( "EditorViewport.ToggleRealTime", new IMAGE_BRUSH_SVG("Starship/Common/Realtime", Icon16x16));
 
