Sequencer: Icons for Auto change

diff --git a/Engine/Source/Editor/EditorStyle/Private/StarshipStyle.cpp b/Engine/Source/Editor/EditorStyle/Private/StarshipStyle.cpp
--- a/Engine/Source/Editor/EditorStyle/Private/StarshipStyle.cpp
+++ b/Engine/Source/Editor/EditorStyle/Private/StarshipStyle.cpp
@@ -2104,6 +2104,10 @@ void FStarshipEditorStyle::FStyle::SetupSequencerStyles()
 		Set("Sequencer.ShotFilter", new IMAGE_BRUSH("Sequencer/FilteredArea", FVector2D(74, 74), FLinearColor::White, ESlateBrushTileType::Both));
 		Set("Sequencer.KeyMark", new IMAGE_BRUSH("Sequencer/KeyMark", FVector2D(3, 21), FLinearColor::White, ESlateBrushTileType::NoTile));
 		Set("Sequencer.ToggleAutoKeyEnabled", new IMAGE_BRUSH_SVG("Starship/Sequencer/AutoKey", Icon20x20));
+		Set("Sequencer.SetAutoKey", new IMAGE_BRUSH_SVG("Starship/Sequencer/AutoKey", Icon20x20));
+		Set("Sequencer.SetAutoTrack", new IMAGE_BRUSH_SVG("Starship/Sequencer/AutoTrack", Icon20x20));
+		Set("Sequencer.SetAutoChangeAll", new IMAGE_BRUSH_SVG("Starship/Sequencer/AutoChangeAll", Icon20x20));
+		Set("Sequencer.SetAutoChangeNone", new IMAGE_BRUSH_SVG("Starship/Sequencer/AutoChangeNone", Icon20x20));
 		Set("Sequencer.AllowAllEdits", new IMAGE_BRUSH_SVG("Starship/Sequencer/AllowAllEdits", Icon20x20));
 		Set("Sequencer.AllowSequencerEditsOnly", new IMAGE_BRUSH_SVG("Starship/Sequencer/AllowSequencerEditsOnly", Icon20x20));
 		Set("Sequencer.AllowLevelEditsOnly", new IMAGE_BRUSH_SVG("Starship/Sequencer/AllowLevelEditsOnly", Icon20x20));
