Fixing Missing Slate Brushes

diff --git a/Engine/Plugins/Animation/ControlRig/Source/ControlRigEditor/Private/ControlRigEditorStyle.h b/Engine/Plugins/Animation/ControlRig/Source/ControlRigEditor/Private/ControlRigEditorStyle.h
--- a/Engine/Plugins/Animation/ControlRig/Source/ControlRigEditor/Private/ControlRigEditorStyle.h
+++ b/Engine/Plugins/Animation/ControlRig/Source/ControlRigEditor/Private/ControlRigEditorStyle.h
@@ -47,6 +47,7 @@ public:
 		const FString ControlRigPluginContentDir = FPaths::EnginePluginsDir() / TEXT("Animation/ControlRig/Content");
 		const FString EngineEditorSlateDir = FPaths::EngineContentDir() / TEXT("Editor/Slate");
 		SetContentRoot(ControlRigPluginContentDir);
+		SetCoreContentRoot(EngineEditorSlateDir);
 
 		const FSlateColor DefaultForeground(FLinearColor(0.72f, 0.72f, 0.72f, 1.f));
 
@@ -117,7 +118,7 @@ public:
 
 			Set("ControlRig.ResumeExecution", new IMAGE_BRUSH_SVG("Slate/simulate", Icon40x40));
 			Set("ControlRig.ReleaseMode", new IMAGE_BRUSH_SVG("Slate/animation", Icon40x40));
-			Set("ControlRig.DebugMode", new IMAGE_BRUSH_SVG("Slate/Bug", Icon40x40));
+			Set("ControlRig.DebugMode", new CORE_IMAGE_BRUSH_SVG("Starship/Common/Debug", Icon40x40));
 
 			{
 				FContentRootBracket Bracket(this, EngineEditorSlateDir);
diff --git a/Engine/Plugins/Editor/PluginBrowser/Source/PluginBrowser/Private/PluginStyle.cpp b/Engine/Plugins/Editor/PluginBrowser/Source/PluginBrowser/Private/PluginStyle.cpp
--- a/Engine/Plugins/Editor/PluginBrowser/Source/PluginBrowser/Private/PluginStyle.cpp
+++ b/Engine/Plugins/Editor/PluginBrowser/Source/PluginBrowser/Private/PluginStyle.cpp
@@ -9,6 +9,7 @@
 #include "Styling/StarshipCoreStyle.h"
 
 #define RootToContentDir FPluginStyle::InContent
+#define RootToCoreContentDir StyleSet->RootToCoreContentDir
 
 #include "Styling/SlateStyleMacros.h"
 #include "Styling/StyleColors.h"
@@ -59,7 +60,7 @@ void FPluginStyle::Initialize()
 
 	StyleSet->Set( "Plugins.TabIcon", new IMAGE_BRUSH_SVG( "Plugins", Icon16x16 ) );
 	StyleSet->Set( "Plugins.BreadcrumbArrow", new IMAGE_BRUSH( "SmallArrowRight", Icon10x10 ) );
-	StyleSet->Set( "Plugins.Documentation", new IMAGE_BRUSH_SVG("Starship/Common/Documentation", Icon16x16));
+	StyleSet->Set( "Plugins.Documentation", new CORE_IMAGE_BRUSH_SVG("../Editor/Slate/Starship/Common/Documentation", Icon16x16));
 	StyleSet->Set( "Plugins.ListBorder", new FSlateRoundedBoxBrush(FStyleColors::Recessed, 4.0f));
 	StyleSet->Set( "Plugins.RestartWarningBorder", new FSlateRoundedBoxBrush(FStyleColors::Panel , 5.0f, FStyleColors::Warning, 1.0f));
 
diff --git a/Engine/Plugins/Experimental/ChaosVehiclesPlugin/Source/ChaosVehiclesEditor/Private/ChaosVehiclesEditorStyle.h b/Engine/Plugins/Experimental/ChaosVehiclesPlugin/Source/ChaosVehiclesEditor/Private/ChaosVehiclesEditorStyle.h
--- a/Engine/Plugins/Experimental/ChaosVehiclesPlugin/Source/ChaosVehiclesEditor/Private/ChaosVehiclesEditorStyle.h
+++ b/Engine/Plugins/Experimental/ChaosVehiclesPlugin/Source/ChaosVehiclesEditor/Private/ChaosVehiclesEditorStyle.h
@@ -3,6 +3,7 @@
 #pragma once
 #include "Styling/SlateStyle.h"
 #include "Styling/SlateStyleRegistry.h"
+#include "Styling/SlateStyleMacros.h"
 
 class FChaosVehiclesEditorStyle final : public FSlateStyleSet
 {
@@ -16,8 +17,8 @@ public:
 		SetContentRoot(FPaths::EnginePluginsDir() / TEXT("Experimental/ChaosVehiclesPlugin/Resources"));
 #endif
 
-		Set("ClassIcon.ChaosVehicles", new FSlateVectorImageBrush(RootToContentDir(TEXT("ChaosVehicles_16.svg")), Icon16x16));
-		Set("ClassThumbnail.ChaosVehicles", new FSlateVectorImageBrush(RootToContentDir(TEXT("ChaosVehicles_64.svg")), Icon64x64));
+		Set("ClassIcon.ChaosVehicles", new IMAGE_BRUSH("ChaosVehicles_16x", Icon16x16));
+		Set("ClassThumbnail.ChaosVehicles", new IMAGE_BRUSH("ChaosVehicles_64x", Icon64x64));
 
 		FSlateStyleRegistry::RegisterSlateStyle(*this);
 	}
diff --git a/Engine/Plugins/FX/Niagara/Source/NiagaraEditor/Private/NiagaraEditorStyle.cpp b/Engine/Plugins/FX/Niagara/Source/NiagaraEditor/Private/NiagaraEditorStyle.cpp
--- a/Engine/Plugins/FX/Niagara/Source/NiagaraEditor/Private/NiagaraEditorStyle.cpp
+++ b/Engine/Plugins/FX/Niagara/Source/NiagaraEditor/Private/NiagaraEditorStyle.cpp
@@ -289,8 +289,8 @@ void InitParameters(TSharedRef< FSlateStyleSet > Style)
 	Style->Set("NiagaraEditor.StaticIcon", new IMAGE_PLUGIN_BRUSH("Icons/staticpill_16x", Icon16x16));
 	Style->Set("NiagaraEditor.Pins.StaticConnected", new IMAGE_PLUGIN_BRUSH("Icons/StaticPin_Connected", Icon16x16));
 	Style->Set("NiagaraEditor.Pins.StaticDisconnected", new IMAGE_PLUGIN_BRUSH("Icons/StaticPin_Disconnected", Icon16x16));
-	Style->Set("NiagaraEditor.Pins.StaticConnectedHovered", new IMAGE_BRUSH("Icons/StaticPin_Connected", Icon16x16, FLinearColor(0.8f, 0.8f, 0.8f)));
-	Style->Set("NiagaraEditor.Pins.StaticDisconnectedHovered", new IMAGE_BRUSH("Icons/StaticPin_Disconnected", Icon16x16, FLinearColor(0.8f, 0.8f, 0.8f)));
+	Style->Set("NiagaraEditor.Pins.StaticConnectedHovered", new IMAGE_PLUGIN_BRUSH("Icons/StaticPin_Connected", Icon16x16, FLinearColor(0.8f, 0.8f, 0.8f)));
+	Style->Set("NiagaraEditor.Pins.StaticDisconnectedHovered", new IMAGE_PLUGIN_BRUSH("Icons/StaticPin_Disconnected", Icon16x16, FLinearColor(0.8f, 0.8f, 0.8f)));
 }
 
 void InitParameterMapView(TSharedRef< FSlateStyleSet > Style)
@@ -333,7 +333,7 @@ void InitCodeView(TSharedRef< FSlateStyleSet > Style)
 
 
 	const FTextBlockStyle ErrorText = FTextBlockStyle(NormalText)
-		.SetUnderlineBrush(IMAGE_BRUSH("White", Icon8x8, FLinearColor::Red, ESlateBrushTileType::Both))
+		.SetUnderlineBrush(IMAGE_CORE_BRUSH("Old/White", Icon8x8, FLinearColor::Red, ESlateBrushTileType::Both))
 		.SetColorAndOpacity(FLinearColor::Red);
 		
 	Style->Set("TextEditor.NormalText", NormalText);
@@ -343,7 +343,7 @@ void InitCodeView(TSharedRef< FSlateStyleSet > Style)
 	FTextBlockStyle NormalHlslText = FTextBlockStyle(NormalText)
 		.SetFont(HlslFont);
 	const FTextBlockStyle HlslErrorText = FTextBlockStyle(NormalHlslText)
-		.SetUnderlineBrush(IMAGE_BRUSH("White", Icon8x8, FLinearColor::Red, ESlateBrushTileType::Both))
+		.SetUnderlineBrush(IMAGE_CORE_BRUSH("Old/White", Icon8x8, FLinearColor::Red, ESlateBrushTileType::Both))
 		.SetColorAndOpacity(FLinearColor::Red);
 	
 	Style->Set("SyntaxHighlight.HLSL.Normal", FTextBlockStyle(NormalHlslText).SetColorAndOpacity(FLinearColor(FColor(189, 183, 107))));
diff --git a/Engine/Source/Editor/EditorStyle/Private/StarshipStyle.cpp b/Engine/Source/Editor/EditorStyle/Private/StarshipStyle.cpp
--- a/Engine/Source/Editor/EditorStyle/Private/StarshipStyle.cpp
+++ b/Engine/Source/Editor/EditorStyle/Private/StarshipStyle.cpp
@@ -1023,7 +1023,7 @@ void FStarshipEditorStyle::FStyle::SetupGeneralStyles()
 	{
 		Set("PerfTools.TabIcon", new IMAGE_BRUSH( "Icons/icon_tab_PerfTools_16x", Icon16x16 ) );
 		Set("ClassViewer.TabIcon", new IMAGE_BRUSH_SVG( "Starship/Common/Class", Icon16x16 ) );
-		Set("StructViewer.TabIcon", new IMAGE_BRUSH_SVG( "Starship/Common/UserDefinedStruct", Icon16x16 ) );
+		Set("StructViewer.TabIcon", new IMAGE_BRUSH_SVG( "Starship/AssetIcons/UserDefinedStruct_16", Icon16x16 ) );
 		Set("BlueprintDebugger.TabIcon", new IMAGE_BRUSH_SVG( "Starship/Common/BlueprintDebugger", Icon16x16 ) );
 		Set("CollisionAnalyzer.TabIcon", new IMAGE_BRUSH_SVG("Starship/Common/Collision", Icon16x16));
 		Set("ObjectBrowser.TabIcon", new IMAGE_BRUSH_SVG( "Starship/Common/ObjectsBrowser", Icon16x16 ) );
@@ -1325,7 +1325,7 @@ void FStarshipEditorStyle::FStyle::SetupGeneralStyles()
 	{
 		Set( "MainFrame.AutoSaveImage", 	       new IMAGE_BRUSH_SVG( "Starship/Common/SaveCurrent", Icon16x16 ) );
 		Set( "MainFrame.SaveAll",                  new IMAGE_BRUSH_SVG( "Starship/Common/SaveAll", Icon16x16 ) );
-		Set( "MainFrame.ChoosePackagesToSave",     new IMAGE_BRUSH_SVG( "Starship/Common/icon_file_choosepackages_16px", Icon16x16 ) );
+		Set( "MainFrame.ChoosePackagesToSave",     new IMAGE_BRUSH_SVG( "Starship/Common/SaveChoose", Icon16x16 ) );
 		Set( "MainFrame.NewProject",               new IMAGE_BRUSH_SVG( "Starship/Common/ProjectNew", Icon16x16 ) );
 		Set( "MainFrame.OpenProject",              new IMAGE_BRUSH_SVG( "Starship/Common/ProjectOpen", Icon16x16 ) );
 		Set( "MainFrame.AddCodeToProject",         new IMAGE_BRUSH_SVG( "Starship/Common/ProjectC++", Icon16x16 ) );
@@ -4175,7 +4175,7 @@ void FStarshipEditorStyle::FStyle::SetupGraphEditorStyles()
 			Set( "TimelineEditor.AddEventTrack", new IMAGE_BRUSH_SVG( "Starship/Common/Event", Icon16x16) );
 			Set( "TimelineEditor.AddColorTrack", new IMAGE_BRUSH_SVG( "Starship/TimelineEditor/TrackTypeColor", Icon16x16) );
 			Set( "TimelineEditor.AddCurveAssetTrack", new IMAGE_BRUSH_SVG("Starship/TimelineEditor/TrackTypeCurve", Icon16x16));
-			Set( "TimelineEditor.DeleteTrack", new IMAGE_BRUSH_SVG( "Icons/icon_TrackDelete_36x24px", Icon36x24 ) );
+			Set( "TimelineEditor.DeleteTrack", new IMAGE_BRUSH( "Icons/icon_TrackDelete_36x24px", Icon36x24 ) );
 			Set("TimelineEditor.AutoPlay", new IMAGE_BRUSH_SVG("Starship/TimelineEditor/TimelineAutoPlay", Icon20x20));
 			Set("TimelineEditor.IgnoreTimeDilation", new IMAGE_BRUSH_SVG("Starship/TimelineEditor/TimelineIgnoreTimeDilation", Icon20x20));
 			Set("TimelineEditor.Replicated", new IMAGE_BRUSH_SVG("Starship/TimelineEditor/TimelineReplicated", Icon20x20));
@@ -4414,7 +4414,7 @@ void FStarshipEditorStyle::FStyle::SetupLevelEditorStyle()
 			Set( "LevelEditor.Tabs.EditorModes",            new IMAGE_BRUSH_SVG("Starship/Common/EditorModes", Icon16x16) );
 			Set( "LevelEditor.Tabs.Modes",                  new IMAGE_BRUSH_SVG("Starship/Common/EditorModes", Icon16x16) );
 			Set( "LevelEditor.Tabs.PlacementBrowser",       new IMAGE_BRUSH_SVG("Starship/Common/PlaceActors", Icon16x16) );
-			Set( "LevelEditor.Tabs.Properties",             new IMAGE_BRUSH_SVG("Starship/Common/Properties", Icon16x16) );
+			Set( "LevelEditor.Tabs.Properties",             new IMAGE_BRUSH_SVG("Starship/StaticMeshEditor/Properties", Icon16x16) );
 			Set( "LevelEditor.Tabs.Outliner",               new IMAGE_BRUSH_SVG("Starship/Common/WorldOutliner", Icon16x16) );
 			Set( "LevelEditor.Tabs.ContentBrowser",         new IMAGE_BRUSH_SVG("Starship/Common/ContentBrowser", Icon16x16) );
 
@@ -5186,7 +5186,7 @@ void FStarshipEditorStyle::FStyle::SetupPersonaStyle()
 		Set( "FullBlueprintEditor.SwitchToScriptingMode", new IMAGE_BRUSH_SVG( "Starship/Blueprints/icon_BlueprintEditor_EventGraph", Icon20x20));
 
 		// @todo - Icon Replacement - The UI commands using these icons are never visible in the editor
-		Set( "FullBlueprintEditor.SwitchToBlueprintDefaultsMode", new IMAGE_BRUSH_SVG("Starship/Common/Defaults", Icon20x20));
+		Set( "FullBlueprintEditor.SwitchToBlueprintDefaultsMode", new IMAGE_BRUSH_SVG("Starship/Common/blueprint", Icon20x20));
 		Set( "FullBlueprintEditor.SwitchToComponentsMode", new IMAGE_BRUSH_SVG("Starship/Common/Component", Icon20x20));
 
 		Set( "FullBlueprintEditor.EditGlobalOptions", new CORE_IMAGE_BRUSH_SVG( "Starship/Common/Settings", Icon20x20));
@@ -5275,8 +5275,8 @@ void FStarshipEditorStyle::FStyle::SetupPersonaStyle()
 		Set( "Kismet.Explorer.SearchDepthFont", DEFAULT_FONT( "Bold", 14) );
 
 		Set( "Kismet.Interfaces.Title",  FTextBlockStyle(NormalText) .SetFont( DEFAULT_FONT("Bold", 11 ) ) );
-		Set( "Kismet.Interfaces.Implement", new CORE_IMAGE_BRUSH_SVG( "Starship/Common/PreviousArrow", Icon16x16) );
-		Set( "Kismet.Interfaces.Remove", new CORE_IMAGE_BRUSH_SVG( "Icons/assign_right_16x", Icon16x16) );
+		Set( "Kismet.Interfaces.Implement", new CORE_IMAGE_BRUSH_SVG( "Starship/Common/arrow-left", Icon16x16) );
+		Set( "Kismet.Interfaces.Remove", new CORE_IMAGE_BRUSH_SVG( "Starship/Common/arrow-right", Icon16x16) );
 
 		Set( "Kismet.TypePicker.CategoryFont", DEFAULT_FONT( "BoldCondensedItalic", 11) );
 		Set( "Kismet.TypePicker.NormalFont", DEFAULT_FONT( "Regular", 11) );
@@ -5301,8 +5301,8 @@ void FStarshipEditorStyle::FStyle::SetupPersonaStyle()
 		Set( "BlueprintEditor.ShowGrid", new IMAGE_BRUSH_SVG( "Starship/Common/Grid", Icon16x16 ) );
 		Set( "BlueprintEditor.EnableSimulation", new IMAGE_BRUSH_SVG("Starship/MainToolbar/simulate", Icon20x20));
 		Set( "BlueprintEditor.EnableProfiling", new IMAGE_BRUSH_SVG("Starship/Common/Statistics", Icon20x20) );
-		Set( "SCS.NativeComponent", new IMAGE_BRUSH_SVG("Starship/AssetIcons/SceneComponent", Icon16x16));
-		Set( "SCS.Component", new IMAGE_BRUSH_SVG("Starship/AssetIcons/ActorComponent", Icon16x16));
+		Set( "SCS.NativeComponent", new IMAGE_BRUSH_SVG("Starship/AssetIcons/SceneComponent_16", Icon16x16));
+		Set( "SCS.Component", new IMAGE_BRUSH_SVG("Starship/AssetIcons/ActorComponent_16", Icon16x16));
 
 		// curve viewer
 		Set("AnimCurveViewer.MorphTargetOn", new IMAGE_BRUSH(TEXT("Persona/AnimCurveViewer/MorphTarget_On"), Icon16x16));
@@ -5959,7 +5959,7 @@ void FStarshipEditorStyle::FStyle::SetupClassIconsAndThumbnails()
 		Set("ClassIcon.ForceFeedbackComponent", new IMAGE_BRUSH_SVG("Starship/AssetIcons/ForceFeedbackEffect_16", Icon16x16));
 		Set("ClassIcon.LandscapeComponent", new IMAGE_BRUSH_SVG("Starship/AssetIcons/Landscape_16", Icon16x16));
 		Set("ClassIcon.ParticleSystemComponent", new IMAGE_BRUSH_SVG("Starship/AssetIcons/Emitter_16", Icon16x16));
-		Set("ClassIcon.PlanarReflectionComponent", new IMAGE_BRUSH_SVG("Starship/AssetIcons/PlaneReflectionCapture_16", Icon16x16));
+		Set("ClassIcon.PlanarReflectionComponent", new IMAGE_BRUSH_SVG("Starship/AssetIcons/PlanarReflectionCapture_16", Icon16x16));
 		Set("ClassIcon.PointLightComponent", new IMAGE_BRUSH_SVG("Starship/AssetIcons/PointLight_16", Icon16x16));
 		Set("ClassIcon.RectLightComponent", new IMAGE_BRUSH_SVG("Starship/AssetIcons/RectLight_16", Icon16x16));
 		Set("ClassIcon.RadialForceComponent", new IMAGE_BRUSH_SVG("Starship/AssetIcons/RadialForceActor_16", Icon16x16));
@@ -7201,9 +7201,9 @@ void FStarshipEditorStyle::FStyle::SetupToolkitStyles()
 
 	// General Curve Icons
 	{
-		Set("Curve.ZoomToFit", new IMAGE_BRUSH_SVG("Starship/Common/ZoomToFit", Icon20x20));
-		Set("Curve.ZoomToFitHorizontal", new IMAGE_BRUSH_SVG("Starship/Common/ZoomToFitHorizontal", Icon20x20));
-		Set("Curve.ZoomToFitVertical", new IMAGE_BRUSH_SVG("Starship/Common/ZoomToFitVertical", Icon20x20));
+		Set("Curve.ZoomToFit", new IMAGE_BRUSH_SVG("Starship/GenericCurveEditor/FramingAll", Icon20x20));
+		Set("Curve.ZoomToFitHorizontal", new IMAGE_BRUSH_SVG("Starship/GenericCurveEditor/FramingHorizontal", Icon20x20));
+		Set("Curve.ZoomToFitVertical", new IMAGE_BRUSH_SVG("Starship/GenericCurveEditor/FramingVertical", Icon20x20));
 	}
 
 	// PhysicsAssetEditor
