Adds ToggleButton to EditorViewportToolBar style, nDisplay fix

diff --git a/Engine/Source/Editor/EditorStyle/Private/StarshipStyle.cpp b/Engine/Source/Editor/EditorStyle/Private/StarshipStyle.cpp
--- a/Engine/Source/Editor/EditorStyle/Private/StarshipStyle.cpp
+++ b/Engine/Source/Editor/EditorStyle/Private/StarshipStyle.cpp
@@ -2371,6 +2371,23 @@ void FStarshipEditorStyle::FStyle::SetupViewportStyles()
 			.SetButtonStyle(ViewportMenuButton)
 			.SetContentPadding(ViewportMarginCenter);
 
+		// Non-grouped Toggle Button
+		FCheckBoxStyle SoloToggleButton = FCheckBoxStyle(ViewportToolbarStyle.ToggleButton)
+			.SetUncheckedImage(*ViewportGroupBrush)
+			.SetUncheckedPressedImage(*ViewportGroupPressedBrush)
+			.SetUncheckedHoveredImage(*ViewportGroupBrush)
+			.SetCheckedImage(FSlateRoundedBoxBrush(FStyleColors::Primary, 12.f, FLinearColor(0, 0, 0, .8), 1.0))
+			.SetCheckedHoveredImage(FSlateRoundedBoxBrush(FStyleColors::PrimaryHover, 12.f, FLinearColor(0, 0, 0, .8), 1.0))
+			.SetCheckedPressedImage(FSlateRoundedBoxBrush(FStyleColors::PrimaryPress, 12.f, FLinearColor(0, 0, 0, .8), 1.0))
+			.SetForegroundColor(FStyleColors::Foreground)
+			.SetPressedForegroundColor(FStyleColors::ForegroundHover)
+			.SetHoveredForegroundColor(FStyleColors::ForegroundHover)
+			.SetCheckedForegroundColor(FStyleColors::Foreground)
+			.SetCheckedPressedForegroundColor(FStyleColors::ForegroundHover)
+			.SetCheckedHoveredForegroundColor(FStyleColors::ForegroundHover)
+			.SetPadding(FMargin(6.0f, 4.0f, 6.0f, 4.0f));
+
+
 		ViewportToolbarStyle
 			.SetBackground(FSlateNoResource())
 			.SetIconSize(Icon16x16)
@@ -2382,6 +2399,7 @@ void FStarshipEditorStyle::FStyle::SetupViewportStyles()
 			.SetButtonPadding(FMargin(0))
 			.SetCheckBoxPadding(FMargin(4.0f, 0.0f))
 			.SetComboButtonStyle(ViewportComboButton)
+			.SetToggleButtonStyle(SoloToggleButton)
 			.SetButtonStyle(ViewportMenuButton)
 			.SetSeparatorBrush(FSlateNoResource())
 			.SetSeparatorPadding(FMargin(2.0f, 0.0f))
