Exposed enum EAssetEditorOpenLocation to Blueprint

diff --git a/Engine/Source/Editor/EditorStyle/Public/Classes/EditorStyleSettings.h b/Engine/Source/Editor/EditorStyle/Public/Classes/EditorStyleSettings.h
--- a/Engine/Source/Editor/EditorStyle/Public/Classes/EditorStyleSettings.h
+++ b/Engine/Source/Editor/EditorStyle/Public/Classes/EditorStyleSettings.h
@@ -10,7 +10,7 @@
 #include "Rendering/RenderingCommon.h"
 #include "EditorStyleSettings.generated.h"
 
-UENUM()
+UENUM(BlueprintType)
 enum class EAssetEditorOpenLocation : uint8
 {
 	/** Attempts to dock asset editors into either a new window, or the main window if they were docked there. */
