Replace tutorial images that are still being referred to in other editor styles.

diff --git a/Engine/Source/Editor/EditorStyle/Private/StarshipStyle.cpp b/Engine/Source/Editor/EditorStyle/Private/StarshipStyle.cpp
--- a/Engine/Source/Editor/EditorStyle/Private/StarshipStyle.cpp
+++ b/Engine/Source/Editor/EditorStyle/Private/StarshipStyle.cpp
@@ -2684,10 +2684,6 @@ void FStarshipEditorStyle::FStyle::SetupTutorialStyles()
 
 		Set( "Documentation.Separator", new BOX_BRUSH( "Common/Separator", 1/4.0f, FLinearColor(1,1,1,0.5f) ) );
 	}
-
-	{
-		Set("Documentation.ToolTip.Background", new BOX_BRUSH("Tutorials/TutorialContentBackground", FMargin(4 / 16.0f)));
-	}
 }
 
 void FStarshipEditorStyle::FStyle::SetupPropertyEditorStyles()
