Place Actors Panel: Fix checkbox styling

diff --git a/Engine/Source/Editor/EditorStyle/Private/StarshipStyle.cpp b/Engine/Source/Editor/EditorStyle/Private/StarshipStyle.cpp
--- a/Engine/Source/Editor/EditorStyle/Private/StarshipStyle.cpp
+++ b/Engine/Source/Editor/EditorStyle/Private/StarshipStyle.cpp
@@ -1705,12 +1705,12 @@ void FStarshipEditorStyle::FStyle::SetupGeneralStyles()
 		PaletteToolBarStyle.ButtonStyle.SetNormalPadding(FMargin(2.0f, 6.0f));
 		PaletteToolBarStyle.ButtonStyle.SetPressedPadding(FMargin(2.0f, 6.0f));
 
-		Set( "PaletteToolBar.Tab",  FCheckBoxStyle()
+		Set("PaletteToolBar.Tab",  FCheckBoxStyle()
 			.SetCheckBoxType(            ESlateCheckBoxType::ToggleButton)
 
-			.SetCheckedImage(            FSlateRoundedBoxBrush(FStyleColors::Input, 2.0f))
-			.SetCheckedHoveredImage(     FSlateRoundedBoxBrush(FStyleColors::Input, 2.0f))
-			.SetCheckedPressedImage(     FSlateRoundedBoxBrush(FStyleColors::Input, 2.0f))
+			.SetCheckedImage(            FSlateRoundedBoxBrush(FStyleColors::Primary, 2.0f))
+			.SetCheckedHoveredImage(     FSlateRoundedBoxBrush(FStyleColors::PrimaryHover, 2.0f))
+			.SetCheckedPressedImage(     FSlateRoundedBoxBrush(FStyleColors::Dropdown, 2.0f))
 
 			.SetUncheckedImage(          FSlateRoundedBoxBrush(FStyleColors::Secondary, 2.0f))
 			.SetUncheckedHoveredImage(   FSlateRoundedBoxBrush(FStyleColors::Hover, 2.0f))
@@ -1719,8 +1719,9 @@ void FStarshipEditorStyle::FStyle::SetupGeneralStyles()
 			.SetForegroundColor(         FStyleColors::Foreground)
 			.SetHoveredForegroundColor(  FStyleColors::ForegroundHover)
 			.SetPressedForegroundColor(  FStyleColors::ForegroundHover)
-			.SetCheckedForegroundColor(  FStyleColors::Primary)
-			.SetCheckedHoveredForegroundColor(FStyleColors::PrimaryHover)
+			.SetCheckedForegroundColor(  FStyleColors::ForegroundHover)
+			.SetCheckedHoveredForegroundColor(FStyleColors::ForegroundHover)
+			.SetCheckedPressedForegroundColor(FStyleColors::ForegroundHover)
 			.SetPadding(FMargin(2.f, 6.f))
 		);
 
diff --git a/Engine/Source/Editor/PlacementMode/Private/SPlacementModeTools.cpp b/Engine/Source/Editor/PlacementMode/Private/SPlacementModeTools.cpp
--- a/Engine/Source/Editor/PlacementMode/Private/SPlacementModeTools.cpp
+++ b/Engine/Source/Editor/PlacementMode/Private/SPlacementModeTools.cpp
@@ -593,7 +593,7 @@ void SPlacementModeTools::Construct( const FArguments& InArgs, TSharedRef<SDockT
 
 	SAssignNew(CategoryFilterPtr, SUniformWrapPanel)
 	.HAlign(HAlign_Center)
-	.SlotPadding(FMargin(2.0f, 0.0f));
+	.SlotPadding(FMargin(2.0f, 1.0f));
 
 	UpdatePlacementCategories();
 
