Re-Enables Editor Window Title Bar Color Override

diff --git a/Engine/Source/Editor/EditorStyle/Private/StarshipStyle.cpp b/Engine/Source/Editor/EditorStyle/Private/StarshipStyle.cpp
--- a/Engine/Source/Editor/EditorStyle/Private/StarshipStyle.cpp
+++ b/Engine/Source/Editor/EditorStyle/Private/StarshipStyle.cpp
@@ -166,13 +166,21 @@ FStarshipEditorStyle::FStyle::FStyle( const TWeakObjectPtr< UEditorStyleSettings
 {
 }
 
-
-void FStarshipEditorStyle::FStyle::SettingsChanged(UObject* ChangedObject, FPropertyChangedEvent& PropertyChangedEvent)
+FStarshipEditorStyle::FStyle::~FStyle()
 {
-	if ( ChangedObject == Settings.Get() )
+
+#ifdef WITH_EDITOR
+	if (Settings.IsValid())
 	{
-		SyncSettings();
+		Settings->OnSettingChanged().Remove(SettingChangedHandler);
 	}
+#endif
+
+}
+
+void FStarshipEditorStyle::FStyle::SettingsChanged(FName PropertyName)
+{
+	SyncSettings();
 }
 
 void FStarshipEditorStyle::FStyle::SyncSettings()
@@ -183,23 +191,17 @@ void FStarshipEditorStyle::FStyle::SyncSettings()
 		auto SubduedSelectionColor = Settings->GetSubduedSelectionColor();
 		SetColor(SelectionColor_Subdued_LinearRef, SubduedSelectionColor);
 
-		// Also sync the colors used by FStarshipCoreStyle, as FEditorStyle isn't yet being used as an override everywhere
-	/*	FStarshipCoreStyle::SetSelectorColor(Settings->KeyboardFocusColor);
-		FStarshipCoreStyle::SetSelectionColor(Settings->SelectionColor);
-		FStarshipCoreStyle::SetInactiveSelectionColor(Settings->InactiveSelectionColor);
-		FStarshipCoreStyle::SetPressedSelectionColor(Settings->PressedSelectionColor);
-*/
-
 		// Sync the window background settings
 		FWindowStyle& WindowStyle = const_cast<FWindowStyle&>(FStarshipCoreStyle::GetCoreStyle().GetWidgetStyle<FWindowStyle>("Window"));
-
 		if (Settings->bEnableEditorWindowBackgroundColor)
 		{
 			SetColor(WindowHighlightColor_LinearRef, Settings->EditorWindowBackgroundColor);
+			WindowTitleOverride->TintColor = WindowHighlightColor_LinearRef;
 		}
 		else
 		{
 			SetColor(WindowHighlightColor_LinearRef, FLinearColor(0, 0, 0, 0));
+			WindowTitleOverride->TintColor = FStyleColors::Title;
 		}
 	}
 }
@@ -296,6 +298,14 @@ void FStarshipEditorStyle::FStyle::Initialize()
 	AuditDuplicatedCoreStyles(*this);
 	
 	SyncSettings();
+
+#ifdef WITH_EDITOR
+	if (Settings.IsValid())
+	{
+		SettingChangedHandler = Settings->OnSettingChanged().AddRaw(this, &FStarshipEditorStyle::FStyle::SettingsChanged);
+	}
+#endif
+
 }
 
 void FStarshipEditorStyle::FStyle::SetupGeneralStyles()
@@ -2653,10 +2663,15 @@ void FStarshipEditorStyle::FStyle::SetupGeneralIcons()
 
 void FStarshipEditorStyle::FStyle::SetupWindowStyles()
 {
-	// Window styling
-	{
-		EditorWindowHighlightBrush = CORE_IMAGE_BRUSH("Common/Window/WindowTitle", FVector2D(74, 74), FLinearColor::White, ESlateBrushTileType::Horizontal);
-	}
+	// Override the core "Brushes.Title" brush in this editor style so we can overwrite the color it when the 
+	// EditorStyleSetting.bEnableEditorWindowBackgroundColor is enabled which allows users to customize
+	// the title bar area
+
+	// NOTE!  This raw pointer is "owned" by the style once we call Set. Therefore we let 
+	// the style destroy the brush rather than calling delete within the Editor Style.
+
+	WindowTitleOverride = new FSlateColorBrush(FStyleColors::Title);
+	Set("Brushes.Title", WindowTitleOverride); 
 }
 
 void FStarshipEditorStyle::FStyle::SetupProjectBadgeStyle()
diff --git a/Engine/Source/Editor/EditorStyle/Private/StarshipStyle.h b/Engine/Source/Editor/EditorStyle/Private/StarshipStyle.h
--- a/Engine/Source/Editor/EditorStyle/Private/StarshipStyle.h
+++ b/Engine/Source/Editor/EditorStyle/Private/StarshipStyle.h
@@ -40,6 +40,7 @@ public:
 	{
 	public:
 		FStyle(const TWeakObjectPtr< UEditorStyleSettings >& InSettings);
+		~FStyle();
 
 		void Initialize();
 		void SetupGeneralStyles();
@@ -74,7 +75,7 @@ public:
 		void SetupColorPickerStyle();
 		void SetupDerivedDataStyle();
 
-		void SettingsChanged(UObject* ChangedObject, FPropertyChangedEvent& PropertyChangedEvent);
+		void SettingsChanged(FName PropertyName);
 		void SyncSettings();
 		void SyncParentStyles();
 
@@ -136,9 +137,12 @@ public:
 		FButtonStyle NoBorder;
 		FScrollBarStyle ScrollBar;
 		FSlateFontInfo NormalFont;
-		FSlateBrush EditorWindowHighlightBrush;
+
+		FSlateBrush* WindowTitleOverride;
 
 		TWeakObjectPtr< UEditorStyleSettings > Settings;
+
+		FDelegateHandle SettingChangedHandler;
 	};
 
 	static TSharedRef<class FStarshipEditorStyle::FStyle> Create(const TWeakObjectPtr< UEditorStyleSettings >& InCustomization)
@@ -146,10 +150,6 @@ public:
 		TSharedRef<class FStarshipEditorStyle::FStyle> NewStyle = MakeShareable(new FStarshipEditorStyle::FStyle(InCustomization));
 		NewStyle->Initialize();
 
-#if WITH_EDITOR
-		FCoreUObjectDelegates::OnObjectPropertyChanged.AddSP(NewStyle, &FStarshipEditorStyle::FStyle::SettingsChanged);
-#endif
-
 		return NewStyle;
 	}
 
