StarshipStyle: Rename "EditorViewport.SubMenu.Layouts" to "Icons.Layout" and move it to StarshipCoreStyle
#rn deprecation editor

diff --git a/Engine/Source/Developer/OutputLog/Private/SOutputLog.cpp b/Engine/Source/Developer/OutputLog/Private/SOutputLog.cpp
--- a/Engine/Source/Developer/OutputLog/Private/SOutputLog.cpp
+++ b/Engine/Source/Developer/OutputLog/Private/SOutputLog.cpp
@@ -1834,7 +1834,7 @@ TSharedRef<SWidget> SOutputLog::CreateDrawerDockButton()
 				[
 					SNew(SImage)
 					.ColorAndOpacity(FSlateColor::UseForeground())
-					.Image(FAppStyle::Get().GetBrush("EditorViewport.SubMenu.Layouts"))
+					.Image(FAppStyle::Get().GetBrush("Icons.Layout"))
 				]
 				+ SHorizontalBox::Slot()
 				.VAlign(VAlign_Center)
diff --git a/Engine/Source/Editor/EditorStyle/Private/StarshipStyle.cpp b/Engine/Source/Editor/EditorStyle/Private/StarshipStyle.cpp
--- a/Engine/Source/Editor/EditorStyle/Private/StarshipStyle.cpp
+++ b/Engine/Source/Editor/EditorStyle/Private/StarshipStyle.cpp
@@ -4187,7 +4187,6 @@ void FStarshipEditorStyle::FStyle::SetupLevelEditorStyle()
 		Set("EditorViewport.SubMenu.Stats", new IMAGE_BRUSH_SVG("Starship/Common/Statistics", Icon16x16));
 		Set("EditorViewport.SubMenu.Bookmarks", new IMAGE_BRUSH_SVG("Starship/Common/Bookmarks", Icon16x16));
 		Set("EditorViewport.SubMenu.CreateCamera", new IMAGE_BRUSH_SVG("Starship/Common/CreateCamera", Icon16x16));
-		Set("EditorViewport.SubMenu.Layouts", new IMAGE_BRUSH_SVG("Starship/Common/Layout", Icon16x16));
 
 		Set("LevelViewport.ToggleGameView", new IMAGE_BRUSH_SVG("Starship/Common/GameView", Icon16x16));
 		Set("LevelViewport.ToggleImmersive", new IMAGE_BRUSH_SVG("Starship/Common/EnableFullscreen", Icon16x16));
diff --git a/Engine/Source/Editor/LevelEditor/Private/SLevelViewportToolBar.cpp b/Engine/Source/Editor/LevelEditor/Private/SLevelViewportToolBar.cpp
--- a/Engine/Source/Editor/LevelEditor/Private/SLevelViewportToolBar.cpp
+++ b/Engine/Source/Editor/LevelEditor/Private/SLevelViewportToolBar.cpp
@@ -768,7 +768,7 @@ void SLevelViewportToolBar::FillOptionsMenu(UToolMenu* Menu)
 				FText::GetEmpty(),
 				FNewToolMenuDelegate::CreateSP(this, &SLevelViewportToolBar::GenerateViewportConfigsMenu),
 				false,
-				FSlateIcon(FAppStyle::Get().GetStyleSetName(), "EditorViewport.SubMenu.Layouts")
+				FSlateIcon(FAppStyle::Get().GetStyleSetName(), "Icons.Layout")
 			);
 		}
 
diff --git a/Engine/Source/Editor/UMGEditor/Private/TabFactory/AnimationTabSummoner.cpp b/Engine/Source/Editor/UMGEditor/Private/TabFactory/AnimationTabSummoner.cpp
--- a/Engine/Source/Editor/UMGEditor/Private/TabFactory/AnimationTabSummoner.cpp
+++ b/Engine/Source/Editor/UMGEditor/Private/TabFactory/AnimationTabSummoner.cpp
@@ -723,7 +723,7 @@ private:
 					[
 						SNew(SImage)
 						.ColorAndOpacity(FSlateColor::UseForeground())
-						.Image(FAppStyle::Get().GetBrush("EditorViewport.SubMenu.Layouts"))
+						.Image(FAppStyle::Get().GetBrush("Icons.Layout"))
 					]
 					+ SHorizontalBox::Slot()
 					.VAlign(VAlign_Center)
diff --git a/Engine/Source/Runtime/SlateCore/Private/Styling/StarshipCoreStyle.cpp b/Engine/Source/Runtime/SlateCore/Private/Styling/StarshipCoreStyle.cpp
--- a/Engine/Source/Runtime/SlateCore/Private/Styling/StarshipCoreStyle.cpp
+++ b/Engine/Source/Runtime/SlateCore/Private/Styling/StarshipCoreStyle.cpp
@@ -396,6 +396,7 @@ TSharedRef<ISlateStyle> FStarshipCoreStyle::Create()
 		Style->Set("Icons.Rotate180", new IMAGE_BRUSH_SVG("Starship/Common/Rotate180", Icon16x16));
 		Style->Set("Icons.FlipHorizontal", new IMAGE_BRUSH_SVG("Starship/Common/FlipHorizontal", Icon16x16));
 		Style->Set("Icons.FlipVertical", new IMAGE_BRUSH_SVG("Starship/Common/FlipVertical", Icon16x16));
+		Style->Set("Icons.Layout", new IMAGE_BRUSH_SVG("Starship/Common/Layout", Icon16x16));
 
 		// Toolbar Size Icons
 		Style->Set("Icons.Toolbar.Settings", new IMAGE_BRUSH_SVG("Starship/Common/settings", Icon20x20));
