Marking FEditorStyle access as deprecated

diff --git a/Engine/Source/Editor/EditorStyle/Private/StarshipStyle.cpp b/Engine/Source/Editor/EditorStyle/Private/StarshipStyle.cpp
--- a/Engine/Source/Editor/EditorStyle/Private/StarshipStyle.cpp
+++ b/Engine/Source/Editor/EditorStyle/Private/StarshipStyle.cpp
@@ -22,6 +22,8 @@
 
 #define LOCTEXT_NAMESPACE "EditorStyle"
 
+FName FStarshipEditorStyle::StyleSetName = TEXT("EditorStyle");
+
 void FStarshipEditorStyle::Initialize()
 {
 	LLM_SCOPE_BYNAME(TEXT("FStarshipEditorStyle"));
@@ -57,6 +59,11 @@ void FStarshipEditorStyle::Shutdown()
 	StyleInstance.Reset();
 }
 
+const FName& FStarshipEditorStyle::GetStyleSetName()
+{
+	return FStarshipEditorStyle::StyleSetName;
+}
+
 BEGIN_SLATE_FUNCTION_BUILD_OPTIMIZATION
 
 /* FStarshipEditorStyle static initialization
@@ -86,7 +93,7 @@ bool FStarshipEditorStyle::FStyle::IncludeEditorSpecificStyles()
  *****************************************************************************/
 
 FStarshipEditorStyle::FStyle::FStyle( const TWeakObjectPtr< UEditorStyleSettings >& InSettings )
-	: FSlateStyleSet("EditorStyle")
+	: FSlateStyleSet(FStarshipEditorStyle::StyleSetName)
 
 	// Note, these sizes are in Slate Units.
 	// Slate Units do NOT have to map to pixels.
diff --git a/Engine/Source/Editor/EditorStyle/Private/StarshipStyle.h b/Engine/Source/Editor/EditorStyle/Private/StarshipStyle.h
--- a/Engine/Source/Editor/EditorStyle/Private/StarshipStyle.h
+++ b/Engine/Source/Editor/EditorStyle/Private/StarshipStyle.h
@@ -35,6 +35,8 @@ public:
 		FStarshipEditorStyle::StyleInstance->SyncSettings();
 	}
 
+	static const FName& GetStyleSetName();
+
 	class FStyle : public FSlateStyleSet
 	{
 	public:
@@ -154,4 +156,5 @@ public:
 
 	static TSharedPtr<FStarshipEditorStyle::FStyle> StyleInstance;
 	static TWeakObjectPtr<UEditorStyleSettings> Settings;
+	static FName StyleSetName;
 };
diff --git a/Engine/Source/Editor/EditorStyle/Public/EditorStyleSet.h b/Engine/Source/Editor/EditorStyle/Public/EditorStyleSet.h
--- a/Engine/Source/Editor/EditorStyle/Public/EditorStyleSet.h
+++ b/Engine/Source/Editor/EditorStyle/Public/EditorStyleSet.h
@@ -39,93 +39,110 @@ public:
 	* the appearance of the rest of the application, use FEditorStyle::GetEditorStyle().
 	*
 	*/
-
+	UE_DEPRECATED(5.1, "FEditorStyle::Get() is deprecated, use FAppStyle::Get() instead.")
 	static const ISlateStyle& Get( )
 	{
 		return FAppStyle::Get();
 	}
 
-	template< class T >            
+	template< class T >      
+	UE_DEPRECATED(5.1, "FEditorStyle::GetWidgetStyle() is deprecated, use FAppStyle::GetWidgetStyle() instead.")
 	static const T& GetWidgetStyle( FName PropertyName, const ANSICHAR* Specifier = NULL  ) 
 	{
 		return FAppStyle::Get().GetWidgetStyle< T >( PropertyName, Specifier );
 	}
 
+	UE_DEPRECATED(5.1, "FEditorStyle::GetFloat() is deprecated, use FAppStyle::GetFloat() instead.")
 	static float GetFloat( FName PropertyName, const ANSICHAR* Specifier = NULL )
 	{
 		return FAppStyle::Get().GetFloat( PropertyName, Specifier );
 	}
 
+	UE_DEPRECATED(5.1, "FEditorStyle::GetVector() is deprecated, use FAppStyle::GetVector() instead.")
 	static FVector2D GetVector( FName PropertyName, const ANSICHAR* Specifier = NULL ) 
 	{
 		return FAppStyle::Get().GetVector( PropertyName, Specifier );
 	}
 
+	UE_DEPRECATED(5.1, "FEditorStyle::GetColor() is deprecated, use FAppStyle::GetColor() instead.")
 	static const FLinearColor& GetColor( FName PropertyName, const ANSICHAR* Specifier = NULL )
 	{
 		return FAppStyle::Get().GetColor( PropertyName, Specifier );
 	}
 
+	UE_DEPRECATED(5.1, "FEditorStyle::GetSlateColor() is deprecated, use FAppStyle::GetSlateColor() instead.")
 	static const FSlateColor GetSlateColor( FName PropertyName, const ANSICHAR* Specifier = NULL )
 	{
 		return FAppStyle::Get().GetSlateColor( PropertyName, Specifier );
 	}
 
+	UE_DEPRECATED(5.1, "FEditorStyle::GetMargin() is deprecated, use FAppStyle::GetMargin() instead.")
 	static const FMargin& GetMargin( FName PropertyName, const ANSICHAR* Specifier = NULL )
 	{
 		return FAppStyle::Get().GetMargin( PropertyName, Specifier );
 	}
 
+	UE_DEPRECATED(5.1, "FEditorStyle::GetBrush() is deprecated, use FAppStyle::GetBrush() instead.")
 	static const FSlateBrush* GetBrush( FName PropertyName, const ANSICHAR* Specifier = NULL )
 	{
 		return FAppStyle::Get().GetBrush( PropertyName, Specifier );
 	}
 
+	UE_DEPRECATED(5.1, "FEditorStyle::GetDynamicImageBrush() is deprecated, use FAppStyle::GetDynamicImageBrush() instead.")
 	static const TSharedPtr< FSlateDynamicImageBrush > GetDynamicImageBrush( FName BrushTemplate, FName TextureName, const ANSICHAR* Specifier = NULL )
 	{
 		return FAppStyle::Get().GetDynamicImageBrush( BrushTemplate, TextureName, Specifier );
 	}
 
+	UE_DEPRECATED(5.1, "FEditorStyle::GetDynamicImageBrush() is deprecated, use FAppStyle::GetDynamicImageBrush() instead.")
 	static const TSharedPtr< FSlateDynamicImageBrush > GetDynamicImageBrush( FName BrushTemplate, const ANSICHAR* Specifier, class UTexture2D* TextureResource, FName TextureName )
 	{
 		return FAppStyle::Get().GetDynamicImageBrush( BrushTemplate, Specifier, TextureResource, TextureName );
 	}
 
+	UE_DEPRECATED(5.1, "FEditorStyle::GetDynamicImageBrush() is deprecated, use FAppStyle::GetDynamicImageBrush() instead.")
 	static const TSharedPtr< FSlateDynamicImageBrush > GetDynamicImageBrush( FName BrushTemplate, class UTexture2D* TextureResource, FName TextureName )
 	{
 		return FAppStyle::Get().GetDynamicImageBrush( BrushTemplate, TextureResource, TextureName );
 	}
 
+	UE_DEPRECATED(5.1, "FEditorStyle::GetSound() is deprecated, use FAppStyle::GetSound() instead.")
 	static const FSlateSound& GetSound( FName PropertyName, const ANSICHAR* Specifier = NULL )
 	{
 		return FAppStyle::Get().GetSound( PropertyName, Specifier );
 	}
 	
+	UE_DEPRECATED(5.1, "FEditorStyle::GetFontStyle() is deprecated, use FAppStyle::GetFontStyle() instead.")
 	static FSlateFontInfo GetFontStyle( FName PropertyName, const ANSICHAR* Specifier = NULL )
 	{
 		return FAppStyle::Get().GetFontStyle( PropertyName, Specifier );
 	}
 
+	UE_DEPRECATED(5.1, "FEditorStyle::GetDefaultBrush() is deprecated, use FAppStyle::GetDefaultBrush() instead.")
 	static const FSlateBrush* GetDefaultBrush()
 	{
 		return FAppStyle::Get().GetDefaultBrush();
 	}
 
+	UE_DEPRECATED(5.1, "FEditorStyle::GetNoBrush() is deprecated, use FAppStyle::GetNoBrush() instead.")
 	static const FSlateBrush* GetNoBrush()
 	{
 		return FStyleDefaults::GetNoBrush();
 	}
 
+	UE_DEPRECATED(5.1, "FEditorStyle::GetOptionalBrush() is deprecated, use FAppStyle::GetOptionalBrush() instead.")
 	static const FSlateBrush* GetOptionalBrush( FName PropertyName, const ANSICHAR* Specifier = NULL, const FSlateBrush* const DefaultBrush = FStyleDefaults::GetNoBrush() )
 	{
 		return FAppStyle::Get().GetOptionalBrush( PropertyName, Specifier, DefaultBrush );
 	}
 
+	UE_DEPRECATED(5.1, "FEditorStyle::GetResources() is deprecated, use FAppStyle::GetResources() instead.")
 	static void GetResources( TArray< const FSlateBrush* >& OutResources )
 	{
 		return FAppStyle::Get().GetResources( OutResources );
 	}
 
+	UE_DEPRECATED(5.1, "FEditorStyle::GetStyleSetName() is deprecated, use FAppStyle::GetAppStyleSetName() instead.")
 	static const FName& GetStyleSetName()
 	{
 		return Instance->GetStyleSetName();
@@ -140,6 +157,7 @@ public:
 	 *
 	 * @return New FName that is A concatenated with B.
 	 */
+	UE_DEPRECATED(5.1, "FEditorStyle::Join() is deprecated, use FAppStyle::Join() instead.")
 	static FName Join( FName A, const ANSICHAR* B )
 	{
 		if( B == NULL )
