Updating CrashReportClient style to only use images from certain folders, and moving developer tool style setup to individual style classes.

diff --git a/Engine/Config/BaseEditorPerProjectUserSettings.ini b/Engine/Config/BaseEditorPerProjectUserSettings.ini
--- a/Engine/Config/BaseEditorPerProjectUserSettings.ini
+++ b/Engine/Config/BaseEditorPerProjectUserSettings.ini
@@ -19,7 +19,7 @@
 ; -------------------------------------------------------------------------------------------------
 
 [/Script/UnrealEd.EditorPerProjectUserSettings]
-bEnableOutputLogWordWrap=False
+
 bEnableOutputLogClearOnPIE=False
 
 ; True if WASD keys should be remapped to flight controls while the right mouse button is held down
@@ -83,6 +83,7 @@ bEnableUserEditorLayoutManagement=True
 ColorVisionDeficiencyPreviewType=NormalVision
 ColorVisionDeficiencySeverity=3
 
+[/Script/OutputLog.OutputLogSettings]
 LogBackgroundColor=(R=0.015996,G=0.015996,B=0.015996,A=1.000000)
 LogSelectionBackgroundColor=(R=0.132868,G=0.132868,B=0.132868,A=1.000000)
 LogNormalColor=(R=0.72,G=0.72,B=0.72,A=1.000000)
diff --git a/Engine/Config/BaseEngine.ini b/Engine/Config/BaseEngine.ini
--- a/Engine/Config/BaseEngine.ini
+++ b/Engine/Config/BaseEngine.ini
@@ -1027,6 +1027,7 @@ DefaultWorldPartitionHLODLayer=/Engine/Maps/Templates/HLODs/HLODLayer_Instanced.
 +ClassRedirects=(OldName="EditorStyleSettings",NewName="/Script/UnrealEd.EditorStyleSettings")
 +EnumRedirects=(OldName="EAssetEditorOpenLocation",NewName="/Script/UnrealEd.EAssetEditorOpenLocation")
 +EnumRedirects=(OldName="ELogCategoryColorizationMode",NewName="/Script/UnrealEd.ELogCategoryColorizationMode")
++EnumRedirects=(OldName="/Script/UnrealEd.ELogCategoryColorizationMode",NewName="/Script/OutputLog.ELogCategoryColorizationMode")
 
 +PropertyRedirects=(OldName="PostProcessSettings.LocalExposureContrastReduction", NewName="PostProcessSettings.LocalExposureContrastScale")
 +PropertyRedirects=(OldName="PostProcessSettings.bOverride_LocalExposureContrastReduction", NewName="PostProcessSettings.bOverride_LocalExposureContrastScale")
diff --git a/Engine/Source/Developer/AutomationWindow/AutomationWindow.Build.cs b/Engine/Source/Developer/AutomationWindow/AutomationWindow.Build.cs
--- a/Engine/Source/Developer/AutomationWindow/AutomationWindow.Build.cs
+++ b/Engine/Source/Developer/AutomationWindow/AutomationWindow.Build.cs
@@ -20,7 +20,7 @@
 					"ApplicationCore",
                     "Slate",
                     "SlateCore",
-                    
+                    "ToolWidgets",
                     "CoreUObject",
                     "Json",
                     "JsonUtilities",
diff --git a/Engine/Source/Developer/AutomationWindow/Private/AutomationWindowModule.cpp b/Engine/Source/Developer/AutomationWindow/Private/AutomationWindowModule.cpp
--- a/Engine/Source/Developer/AutomationWindow/Private/AutomationWindowModule.cpp
+++ b/Engine/Source/Developer/AutomationWindow/Private/AutomationWindowModule.cpp
@@ -7,6 +7,7 @@
 #include "SAutomationWindow.h"
 
 #include "IAutomationWindowModule.h"
+#include "AutomationWindowStyle.h"
 
 
 /**
@@ -42,10 +43,12 @@
 
 	virtual void StartupModule( ) override
 	{
+		FAutomationWindowStyle::Get();
 	}
 
 	virtual void ShutdownModule( ) override
 	{
+		FAutomationWindowStyle::Shutdown();
 		ShutdownDelegate.ExecuteIfBound();	
 	}
 
diff --git a/Engine/Source/Developer/AutomationWindow/Private/AutomationWindowStyle.cpp b/Engine/Source/Developer/AutomationWindow/Private/AutomationWindowStyle.cpp
new file mode 100644
--- /dev/null
+++ b/Engine/Source/Developer/AutomationWindow/Private/AutomationWindowStyle.cpp
@@ -0,0 +1,119 @@
+// Copyright Epic Games, Inc. All Rights Reserved.
+
+#include "AutomationWindowStyle.h"
+#include "Styling/SlateStyle.h"
+#include "Styling/SlateStyleMacros.h"
+#include "Styling/SlateTypes.h"
+#include "Styling/StyleColors.h"
+#include "Styling/StarshipCoreStyle.h"
+#include "Styling/SlateStyleRegistry.h"
+#include "Styling/CoreStyle.h"
+
+FName FAutomationWindowStyle::StyleName("AutomationWindowStyle");
+TUniquePtr<FAutomationWindowStyle> FAutomationWindowStyle::Inst(nullptr);
+
+const FName& FAutomationWindowStyle::GetStyleSetName() const
+{
+	return StyleName;
+}
+
+const FAutomationWindowStyle& FAutomationWindowStyle::Get()
+{
+	if (!Inst.IsValid())
+	{
+		Inst = TUniquePtr<FAutomationWindowStyle>(new FAutomationWindowStyle);
+	}
+	return *(Inst.Get());
+}
+
+void FAutomationWindowStyle::Shutdown()
+{
+	Inst.Reset();
+}
+
+FAutomationWindowStyle::FAutomationWindowStyle()
+	: FSlateStyleSet(StyleName)
+{
+	const FVector2D Icon16x16(16.0f, 16.0f);
+	const FVector2D Icon20x20(20.0f, 20.0f);
+	const FVector2D Icon40x40(40.0f, 40.0f);
+
+	SetParentStyleName(FAppStyle::GetAppStyleSetName());
+
+	SetContentRoot(FPaths::EngineContentDir() / TEXT("Slate"));
+
+	{
+		FTextBlockStyle NormalText = FAppStyle::GetWidgetStyle<FTextBlockStyle>("NormalText");
+
+		Set("Automation.Header", FTextBlockStyle(NormalText)
+			.SetFont(DEFAULT_FONT("Mono", 12))
+			.SetColorAndOpacity(FStyleColors::ForegroundHeader));
+
+		Set("Automation.Normal", FTextBlockStyle(NormalText)
+			.SetFont(DEFAULT_FONT("Mono", 9))
+			.SetColorAndOpacity(FStyleColors::Foreground));
+
+		Set("Automation.Warning", FTextBlockStyle(NormalText)
+			.SetFont(DEFAULT_FONT("Mono", 9))
+			.SetColorAndOpacity(FStyleColors::Warning));
+
+		Set("Automation.Error", FTextBlockStyle(NormalText)
+			.SetFont(DEFAULT_FONT("Mono", 9))
+			.SetColorAndOpacity(FStyleColors::Error));
+
+		Set("Automation.ReportHeader", FTextBlockStyle(NormalText)
+			.SetFont(DEFAULT_FONT("Mono", 10))
+			.SetColorAndOpacity(FStyleColors::ForegroundHeader));
+
+		//state of individual tests
+		Set("Automation.Success", new IMAGE_BRUSH("Automation/Success", Icon16x16));
+		Set("Automation.Warning", new IMAGE_BRUSH("Automation/Warning", Icon16x16));
+		Set("Automation.Fail", new IMAGE_BRUSH("Automation/Fail", Icon16x16));
+		Set("Automation.InProcess", new IMAGE_BRUSH("Automation/InProcess", Icon16x16));
+		Set("Automation.NotRun", new IMAGE_BRUSH("Automation/NotRun", Icon16x16, FLinearColor(0.0f, 0.0f, 0.0f, 0.4f)));
+		Set("Automation.Skipped", new IMAGE_BRUSH("Automation/NoSessionWarning", Icon16x16));
+		Set("Automation.ParticipantsWarning", new IMAGE_BRUSH("Automation/ParticipantsWarning", Icon16x16));
+		Set("Automation.Participant", new IMAGE_BRUSH("Automation/Participant", Icon16x16));
+
+		//status as a regression test or not
+		Set("Automation.SmokeTest", new IMAGE_BRUSH("Automation/SmokeTest", Icon16x16));
+		Set("Automation.SmokeTestParent", new IMAGE_BRUSH("Automation/SmokeTestParent", Icon16x16));
+
+		//run icons
+		Set("AutomationWindow.RunTests", new IMAGE_BRUSH("Automation/RunTests", Icon40x40));
+		Set("AutomationWindow.RefreshTests", new IMAGE_BRUSH("Automation/RefreshTests", Icon40x40));
+		Set("AutomationWindow.FindWorkers", new IMAGE_BRUSH("Automation/RefreshWorkers", Icon40x40));
+		Set("AutomationWindow.StopTests", new IMAGE_BRUSH("Automation/StopTests", Icon40x40));
+		Set("AutomationWindow.RunTests.Small", new IMAGE_BRUSH("Automation/RunTests", Icon20x20));
+		Set("AutomationWindow.RefreshTests.Small", new IMAGE_BRUSH("Automation/RefreshTests", Icon20x20));
+		Set("AutomationWindow.FindWorkers.Small", new IMAGE_BRUSH("Automation/RefreshWorkers", Icon20x20));
+		Set("AutomationWindow.StopTests.Small", new IMAGE_BRUSH("Automation/StopTests", Icon20x20));
+
+		//filter icons
+		Set("AutomationWindow.ErrorFilter", new IMAGE_BRUSH("Automation/ErrorFilter", Icon40x40));
+		Set("AutomationWindow.WarningFilter", new IMAGE_BRUSH("Automation/WarningFilter", Icon40x40));
+		Set("AutomationWindow.SmokeTestFilter", new IMAGE_BRUSH("Automation/SmokeTestFilter", Icon40x40));
+		Set("AutomationWindow.DeveloperDirectoryContent", new IMAGE_BRUSH("Automation/DeveloperDirectoryContent", Icon40x40));
+		Set("AutomationWindow.ExcludedTestsFilter", new IMAGE_BRUSH("Automation/ExcludedTestsFilter", Icon40x40));
+		Set("AutomationWindow.ErrorFilter.Small", new IMAGE_BRUSH("Automation/ErrorFilter", Icon20x20));
+		Set("AutomationWindow.WarningFilter.Small", new IMAGE_BRUSH("Automation/WarningFilter", Icon20x20));
+		Set("AutomationWindow.SmokeTestFilter.Small", new IMAGE_BRUSH("Automation/SmokeTestFilter", Icon20x20));
+		Set("AutomationWindow.DeveloperDirectoryContent.Small", new IMAGE_BRUSH("Automation/DeveloperDirectoryContent", Icon20x20));
+		Set("AutomationWindow.TrackHistory", new IMAGE_BRUSH("Automation/TrackTestHistory", Icon40x40));
+
+		//device group settings
+		Set("AutomationWindow.GroupSettings", new IMAGE_BRUSH("Automation/Groups", Icon40x40));
+		Set("AutomationWindow.GroupSettings.Small", new IMAGE_BRUSH("Automation/Groups", Icon20x20));
+
+		//test backgrounds
+		Set("AutomationWindow.GameGroupBorder", new BOX_BRUSH("Automation/GameGroupBorder", FMargin(4.0f / 16.0f)));
+		Set("AutomationWindow.EditorGroupBorder", new BOX_BRUSH("Automation/EditorGroupBorder", FMargin(4.0f / 16.0f)));
+	}
+	
+	FSlateStyleRegistry::RegisterSlateStyle(*this);
+}
+
+FAutomationWindowStyle::~FAutomationWindowStyle()
+{
+	FSlateStyleRegistry::UnRegisterSlateStyle(*this);
+}
diff --git a/Engine/Source/Developer/AutomationWindow/Private/AutomationWindowStyle.h b/Engine/Source/Developer/AutomationWindow/Private/AutomationWindowStyle.h
new file mode 100644
--- /dev/null
+++ b/Engine/Source/Developer/AutomationWindow/Private/AutomationWindowStyle.h
@@ -0,0 +1,26 @@
+// Copyright Epic Games, Inc. All Rights Reserved.
+
+#pragma once
+
+#include "Styling/SlateStyle.h"
+#include "Styling/SlateColor.h"
+#include "UObject/NameTypes.h"
+#include "Templates/UniquePtr.h"
+
+class FAutomationWindowStyle
+	: public FSlateStyleSet
+{
+public:
+	virtual const FName& GetStyleSetName() const override;
+
+	static const FAutomationWindowStyle& Get();
+	static void Shutdown();
+	
+	~FAutomationWindowStyle();
+
+private:
+	FAutomationWindowStyle();
+
+	static FName StyleName;
+	static TUniquePtr<FAutomationWindowStyle> Inst;
+};
diff --git a/Engine/Source/Developer/AutomationWindow/Private/SAutomationTestItem.cpp b/Engine/Source/Developer/AutomationWindow/Private/SAutomationTestItem.cpp
--- a/Engine/Source/Developer/AutomationWindow/Private/SAutomationTestItem.cpp
+++ b/Engine/Source/Developer/AutomationWindow/Private/SAutomationTestItem.cpp
@@ -11,7 +11,7 @@
 #include "Widgets/Notifications/SProgressBar.h"
 #include "Widgets/Text/STextBlock.h"
 #include "Widgets/Input/SCheckBox.h"
-#include "Styling/AppStyle.h"
+#include "AutomationWindowStyle.h"
 #include "AutomationTestExcludelist.h"
 #include "SAutomationWindow.h"
 
@@ -25,6 +25,7 @@
 #endif
 
 #include "Widgets/Input/SHyperlink.h"
+#include "SSimpleButton.h"
 
 #define LOCTEXT_NAMESPACE "AutomationTestItem"
 
@@ -45,8 +46,6 @@
 BEGIN_SLATE_FUNCTION_BUILD_OPTIMIZATION
 TSharedRef<SWidget> SAutomationTestItem::GenerateWidgetForColumn( const FName& ColumnName )
 {
-	FSlateFontInfo ItemEditorFont = FAppStyle::GetFontStyle(TEXT("NormalFont"));
-
 	if( ColumnName == AutomationTestWindowConstants::Title)
 	{
 		TSharedRef<SWidget> TestNameWidget = SNullWidget::NullWidget;
@@ -56,7 +55,7 @@
 		{
 #if WITH_EDITOR
 			TestNameWidget = SNew(SHyperlink)
-				.Style(FAppStyle::Get(), "Common.GotoNativeCodeHyperlink")
+				.Style(FAutomationWindowStyle::Get(), "Common.GotoNativeCodeHyperlink")
 				.OnNavigate_Lambda([=] {
 					GEngine->Exec(nullptr, *TestStatus->GetOpenCommand());
 				})
@@ -67,7 +66,7 @@
 		{
 #if WITH_EDITOR
 			TestNameWidget = SNew(SHyperlink)
-				.Style(FAppStyle::Get(), "Common.GotoNativeCodeHyperlink")
+				.Style(FAutomationWindowStyle::Get(), "Common.GotoNativeCodeHyperlink")
 				.OnNavigate_Lambda([=] {
 					FString AssetPath = TestStatus->GetAssetPath();
 					FAssetRegistryModule& AssetRegistryModule = FModuleManager::LoadModuleChecked<FAssetRegistryModule>(TEXT("AssetRegistry"));
@@ -90,7 +89,7 @@
 		else if ( !TestStatus->GetSourceFile().IsEmpty() )
 		{
 			TestNameWidget = SNew(SHyperlink)
-				.Style(FAppStyle::Get(), "Common.GotoNativeCodeHyperlink")
+				.Style(FAutomationWindowStyle::Get(), "Common.GotoNativeCodeHyperlink")
 				.OnNavigate_Lambda([=] { FSlateApplication::Get().GotoLineInSource(TestStatus->GetSourceFile(), TestStatus->GetSourceFileLine()); })
 				.Text(FText::FromString(TestStatus->GetDisplayNameWithDecoration()));
 		}
@@ -148,7 +147,7 @@
 				HBox->AddSlot()
 					[
 						SNew( SImage )
-						.Image( FAppStyle::GetBrush("Automation.Participant") )
+						.Image(FAutomationWindowStyle::Get().GetBrush("Automation.Participant") )
 					];
 				HBox->AddSlot()
 					[
@@ -165,7 +164,7 @@
 					.HAlign(HAlign_Center)
 					[
 						SNew( SImage )
-						.Image( FAppStyle::GetBrush("Automation.ParticipantsWarning") )
+						.Image(FAutomationWindowStyle::Get().GetBrush("Automation.ParticipantsWarning") )
 						.ToolTipText( LOCTEXT("ParticipantsWarningToolTip", "Some tests require multiple participants") )
 					];
 			}
@@ -189,7 +188,7 @@
 				.FillWidth(1.0)
 				[			
 					SNew(SBorder)
-					.BorderImage( FAppStyle::GetBrush("ErrorReporting.Box") )
+					.BorderImage(FAutomationWindowStyle::Get().GetBrush("ErrorReporting.Box") )
 					.HAlign(HAlign_Center)
 					.VAlign(VAlign_Center)
 					.Padding( FMargin(3,0) )
@@ -231,7 +230,7 @@
 				.FillWidth(1.0)
 				[
 					SNew(SBorder)
-					.BorderImage( FAppStyle::GetBrush("ErrorReporting.Box") )
+					.BorderImage(FAutomationWindowStyle::Get().GetBrush("ErrorReporting.Box") )
 					.HAlign(HAlign_Center)
 					.VAlign(VAlign_Center)
 					.Padding( FMargin(3,0) )
@@ -291,16 +290,11 @@
 			.HAlign(HAlign_Center)
 			.AutoWidth()
 			[
-				SNew(SButton)
-				.ButtonStyle(FAppStyle::Get(), "SimpleButton")
+				SNew(SSimpleButton)
+				.Icon(FAutomationWindowStyle::Get().GetBrush("Icons.Edit"))
 				.Visibility(this, &SAutomationTestItem::IsDirectlyExcluded_GetVisibility)
 				.ToolTipText(LOCTEXT("EditExcludeOptions", "Edit exclude options"))
 				.OnClicked(FOnClicked::CreateSP(this, &SAutomationTestItem::OnEditExcludeOptionsClicked))
-				[
-					SNew(SImage)
-					.ColorAndOpacity(FSlateColor::UseForeground())
-				.Image(FAppStyle::Get().GetBrush("Icons.Edit"))
-				]
 #endif
 			];
 	}
@@ -320,11 +314,11 @@
 	{
 		if ( TestStatus->IsParent() )
 		{
-			ImageToUse = FAppStyle::GetBrush("Automation.SmokeTest");
+			ImageToUse = FAutomationWindowStyle::Get().GetBrush("Automation.SmokeTest");
 		}
 		else
 		{
-			ImageToUse = FAppStyle::GetBrush("Automation.SmokeTestParent");
+			ImageToUse = FAutomationWindowStyle::Get().GetBrush("Automation.SmokeTestParent");
 		}
 	}
 
@@ -598,32 +592,32 @@
 			//If there were ANY warnings in the results
 			if (CompleteState.NumEnabledTestsWarnings || CompleteState.NumDisabledTestsWarnings)
 			{
-				ImageToUse = FAppStyle::GetBrush("Automation.Warning");
+				ImageToUse = FAutomationWindowStyle::Get().GetBrush("Automation.Warning");
 			}
 			else
 			{
-				ImageToUse = FAppStyle::GetBrush("Automation.Success");
+				ImageToUse = FAutomationWindowStyle::Get().GetBrush("Automation.Success");
 			}
 		}
 		break;
 
 	case EAutomationState::Fail:
-		ImageToUse = FAppStyle::GetBrush("Automation.Fail");
+		ImageToUse = FAutomationWindowStyle::Get().GetBrush("Automation.Fail");
 		break;
 
 	case EAutomationState::NotRun:
 		{
-			ImageToUse = FAppStyle::GetBrush("Automation.NotRun");
+			ImageToUse = FAutomationWindowStyle::Get().GetBrush("Automation.NotRun");
 		}
 		break;
 
 	case EAutomationState::Skipped:
-		ImageToUse = FAppStyle::GetBrush("Automation.Skipped");
+		ImageToUse = FAutomationWindowStyle::Get().GetBrush("Automation.Skipped");
 		break;
 
 	default:
 	case EAutomationState::InProcess:
-		ImageToUse = FAppStyle::GetBrush("Automation.InProcess");
+		ImageToUse = FAutomationWindowStyle::Get().GetBrush("Automation.InProcess");
 		break;
 	}
 
diff --git a/Engine/Source/Developer/AutomationWindow/Private/SAutomationWindow.cpp b/Engine/Source/Developer/AutomationWindow/Private/SAutomationWindow.cpp
--- a/Engine/Source/Developer/AutomationWindow/Private/SAutomationWindow.cpp
+++ b/Engine/Source/Developer/AutomationWindow/Private/SAutomationWindow.cpp
@@ -22,7 +22,6 @@
 #include "Widgets/Layout/SScrollBox.h"
 #include "Widgets/Input/SCheckBox.h"
 #include "Widgets/Input/SSpinBox.h"
-#include "Styling/AppStyle.h"
 #include "SAutomationWindowCommandBar.h"
 #include "AutomationFilter.h"
 #include "AutomationPresetManager.h"
@@ -40,6 +39,7 @@
 #include "Widgets/Images/SThrobber.h"
 #include "Widgets/Input/SHyperlink.h"
 #include "Internationalization/Regex.h"
+#include "AutomationWindowStyle.h"
 
 
 #define LOCTEXT_NAMESPACE "AutomationTest"
@@ -55,7 +55,7 @@
 		: TCommands<FAutomationWindowCommands>(
 		TEXT("AutomationWindow"),
 		NSLOCTEXT("Contexts", "AutomationWindow", "Automation Window"),
-		NAME_None, FAppStyle::GetAppStyleSetName()
+		NAME_None, FAutomationWindowStyle::Get().GetStyleSetName()
 		)
 	{
 	}
@@ -230,7 +230,7 @@
 			SNew(SImage)
 			.ColorAndOpacity(FLinearColor(1.0f, 1.0f, 1.0f, 0.4f))
 			.ToolTipText( LOCTEXT( "Smoke Test", "Smoke Test" ) )
-			.Image(FAppStyle::GetBrush("Automation.SmokeTest"))
+			.Image(FAutomationWindowStyle::Get().GetBrush("Automation.SmokeTest"))
 		]
 
 		+ SHeaderRow::Column( AutomationTestWindowConstants::RequiredDeviceCount )
@@ -241,7 +241,7 @@
 		.VAlignCell(VAlign_Center)
 		[
 			SNew( SImage )
-			.Image( FAppStyle::GetBrush("Automation.ParticipantsWarning") )
+			.Image(FAutomationWindowStyle::Get().GetBrush("Automation.ParticipantsWarning") )
 			.ToolTipText( LOCTEXT( "RequiredDeviceCountWarningToolTip", "Number of devices required." ) )
 		]
 
@@ -478,7 +478,7 @@
 								[
 									//list of results for the selected test
 									SNew(SBorder)
-									.BorderImage(FAppStyle::GetBrush("MessageLog.ListBorder"))
+									.BorderImage(FAutomationWindowStyle::Get().GetBrush("Brushes.Panel"))
 									[
 										SNew(SScrollBox)
 										.Orientation(EOrientation::Orient_Horizontal)
@@ -499,7 +499,7 @@
 								.Padding(0.0f, 4.0f, 0.0f, 0.0f)
 								[
 									SNew(SBorder)
-									.BorderImage(FAppStyle::GetBrush("ToolPanel.GroupBorder"))
+									.BorderImage(FAutomationWindowStyle::Get().GetBrush("ToolPanel.GroupBorder"))
 									.Padding(FMargin(8.0f, 6.0f))
 									[
 										// Add the command bar
@@ -526,7 +526,7 @@
 			.Padding( 15.0f )
 			[
 				SNew(SBorder)
-					.BorderImage(FAppStyle::GetBrush("NotificationList.ItemBackground"))
+					.BorderImage(FAutomationWindowStyle::Get().GetBrush("NotificationList.ItemBackground"))
 					.Padding(8.0f)
 					.Visibility(this, &SAutomationWindow::HandleSelectSessionOverlayVisibility)
 					[
@@ -565,14 +565,14 @@
 	switch(TestBackgroundType)
 	{
 	case EAutomationTestBackgroundStyle::Game:
-		return FAppStyle::GetBrush("AutomationWindow.GameGroupBorder");
+		return FAutomationWindowStyle::Get().GetBrush("AutomationWindow.GameGroupBorder");
 
 	case EAutomationTestBackgroundStyle::Editor:
-		return FAppStyle::GetBrush("AutomationWindow.EditorGroupBorder");
+		return FAutomationWindowStyle::Get().GetBrush("AutomationWindow.EditorGroupBorder");
 
 	case EAutomationTestBackgroundStyle::Unknown:
 	default:
-		return FAppStyle::GetBrush("ToolPanel.GroupBorder");
+		return FAutomationWindowStyle::Get().GetBrush("ToolPanel.GroupBorder");
 	}
 }
 
@@ -648,7 +648,7 @@
 					FOnGetContent::CreateStatic( &SAutomationWindow::GenerateTestsOptionsMenuContent, InAutomationWindow ),
 					TAttribute<FText>(),
 					LOCTEXT( "TestOptionsToolTip", "Test Options" ),
-					FSlateIcon(FAppStyle::GetAppStyleSetName(), "AutomationWindow.TestOptions"),
+					FSlateIcon(FAutomationWindowStyle::Get().GetStyleSetName(), "AutomationWindow.TestOptions"),
 					false);
 
 				// Added button for running the currently open level test.
@@ -658,7 +658,7 @@
 					NAME_None,
 					TAttribute<FText>(),
 					LOCTEXT("RunLevelTest_ToolTip", "If the currently loaded editor level is a test map, click this to select the test and run it immediately."),
-					FSlateIcon(FAppStyle::GetAppStyleSetName(), "AutomationWindow.RunTests"));
+					FSlateIcon(FAutomationWindowStyle::Get().GetStyleSetName(), "AutomationWindow.RunTests"));
 #endif
 
 				ToolbarBuilder.AddToolBarButton( FAutomationWindowCommands::Get().RefreshTests );
@@ -680,7 +680,7 @@
 					FOnGetContent::CreateStatic( &SAutomationWindow::GenerateGroupOptionsMenuContent, InAutomationWindow ),
 					LOCTEXT( "GroupOptions_Label", "Device Groups" ),
 					LOCTEXT( "GroupOptionsToolTip", "Device Group Options" ),
-					FSlateIcon(FAppStyle::GetAppStyleSetName(), "AutomationWindow.GroupSettings"));
+					FSlateIcon(FAutomationWindowStyle::Get().GetStyleSetName(), "AutomationWindow.GroupSettings"));
 			}
 			ToolbarBuilder.EndSection();
 			ToolbarBuilder.BeginSection("Presets");
@@ -693,7 +693,6 @@
 
 	TSharedRef<SWidget> RunTests = 
 		SNew( SButton )
-		.ButtonStyle( FAppStyle::Get(), "ToggleButton" )
 		.ToolTipText( LOCTEXT( "StartStop Tests", "Start / Stop tests" ) )
 		.OnClicked( this, &SAutomationWindow::RunTests )
 		.IsEnabled( this, &SAutomationWindow::IsAutomationRunButtonEnabled )
@@ -721,9 +720,7 @@
 					[
 						SNew( STextBlock )
 						.Text( this, &SAutomationWindow::OnGetNumEnabledTestsString )
-						.ColorAndOpacity( FLinearColor::White )
-						.ShadowOffset( FVector2D::UnitVector )
-						.Font( FAppStyle::GetFontStyle( FName( "ToggleButton.LabelFont" ) ) )
+						.ColorAndOpacity(FSlateColor::UseForeground())
 					]
 				]
 				+ SVerticalBox::Slot()
@@ -733,9 +730,7 @@
 					SNew( STextBlock )
 					.Visibility( this, &SAutomationWindow::GetLargeToolBarVisibility )
 					.Text( this, &SAutomationWindow::GetRunAutomationLabel )
-					.Font( FAppStyle::GetFontStyle( FName( "ToggleButton.LabelFont" ) ) )
-					.ColorAndOpacity(FLinearColor::White)
-					.ShadowOffset( FVector2D::UnitVector )
+					.ColorAndOpacity(FSlateColor::UseForeground())
 				]
 			]
 		];
@@ -805,14 +800,14 @@
 			.AutoWidth()
 			[
 				SNew( SButton )
-				.ButtonStyle( FAppStyle::Get(), "NoBorder" )
+				.ButtonStyle(FAutomationWindowStyle::Get(), "NoBorder" )
 				.OnClicked( this, &SAutomationWindow::HandleNewPresetClicked )
 				.ToolTipText( LOCTEXT("AutomationPresetNewButtonTooltip", "Create a new preset") )
 				.IsEnabled(this, &SAutomationWindow::IsAddButtonEnabled)
 				.Content()
 				[
 					SNew(SImage)
-					.Image(FAppStyle::Get().GetBrush("AutomationWindow.PresetNew"))
+					.Image(FAutomationWindowStyle::Get().GetBrush("Icons.PlusCircle"))
 				]
 			]
 
@@ -821,14 +816,14 @@
 			.AutoWidth()
 			[
 				SNew( SButton )
-				.ButtonStyle( FAppStyle::Get(), "NoBorder" )
+				.ButtonStyle(FAutomationWindowStyle::Get(), "NoBorder" )
 				.OnClicked( this, &SAutomationWindow::HandleSavePresetClicked )
 				.ToolTipText( LOCTEXT("AutomationPresetSaveButtonTooltip", "Save the current test list") )
 				.IsEnabled(this, &SAutomationWindow::IsSaveButtonEnabled)
 				.Content()
 				[
 					SNew(SImage)
-					.Image(FAppStyle::Get().GetBrush("AutomationWindow.PresetSave"))
+					.Image(FAutomationWindowStyle::Get().GetBrush("Icons.Save"))
 				]
 			]
 
@@ -838,14 +833,14 @@
 			[
 				// remove button
 				SNew(SButton)
-				.ButtonStyle( FAppStyle::Get(), "NoBorder" )
+				.ButtonStyle(FAutomationWindowStyle::Get(), "NoBorder" )
 				.OnClicked( this, &SAutomationWindow::HandleRemovePresetClicked )
 				.ToolTipText(LOCTEXT("AutomationPresetRemoveButtonTooltip", "Remove the selected preset"))
 				.IsEnabled(this, &SAutomationWindow::IsRemoveButtonEnabled)
 				.Content()
 				[
 					SNew(SImage)
-					.Image(FAppStyle::Get().GetBrush("AutomationWindow.PresetRemove"))
+					.Image(FAutomationWindowStyle::Get().GetBrush("Icons.Delete"))
 				]
 			]
 		];
@@ -862,7 +857,7 @@
 		[
 			SNew( SBorder )
 			.Padding(0)
-			.BorderImage( FAppStyle::GetBrush("NoBorder") )
+			.BorderImage(FAutomationWindowStyle::Get().GetBrush("NoBorder") )
 			.IsEnabled( FSlateApplication::Get().GetNormalExecutionAttribute() )
 			[
 				ToolbarBuilder.MakeWidget()
@@ -1580,7 +1575,7 @@
 		const FDataDrivenPlatformInfo* DDPI = TargetPlatformInfo ? TargetPlatformInfo->DataDrivenPlatformInfo : nullptr;
 		FName DeviceImageName = DDPI ? DDPI->GetIconStyleName(EPlatformIconSize::Normal) : NAME_None;
 
-		const FSlateBrush* ImageToUse = FAppStyle::GetBrush(DeviceImageName);
+		const FSlateBrush* ImageToUse = FAutomationWindowStyle::Get().GetBrush(DeviceImageName);
 
 		PlatformsHBox->AddSlot()
 		.AutoWidth()
@@ -1590,7 +1585,7 @@
 			+ SOverlay::Slot()
 			[
 				SNew(SBorder)
-				.BorderImage( FAppStyle::GetBrush("ErrorReporting.Box") )
+				.BorderImage(FAutomationWindowStyle::Get().GetBrush("ErrorReporting.Box") )
 				.HAlign(HAlign_Center)
 				.VAlign(VAlign_Center)
 				.Padding( FMargin(3,0) )
@@ -1689,8 +1684,8 @@
 		MessageString.RightChopInline(FileAndLineRegexMatcher.GetMatchEnding(), false);
 
 		SourceLink = SNew(SHyperlink)
-			.Style(FAppStyle::Get(), "Common.GotoNativeCodeHyperlink")
-			.TextStyle(FAppStyle::Get(), Message->Style)
+			.Style(FAutomationWindowStyle::Get(), "Common.GotoNativeCodeHyperlink")
+			.TextStyle(FAutomationWindowStyle::Get(), Message->Style)
 			.OnNavigate_Lambda([=] { FSlateApplication::Get().GotoLineInSource(FileName, LineNumber); })
 			.Text(FText::FromString(FileAndLineRegexMatcher.GetCaptureGroup(0)));
 	}
@@ -1711,7 +1706,7 @@
 			.Padding(0)
 			[
 				SNew(STextBlock)
-				.TextStyle( FAppStyle::Get(), Message->Style )
+				.TextStyle(FAutomationWindowStyle::Get(), Message->Style )
 				.Text(FText::FromString(MessageString))
 			]
 		];
@@ -2045,7 +2040,7 @@
 		Brush += TEXT( ".RunTests" );
 	}
 	Brush += GetSmallIconExtension();
-	return FAppStyle::GetBrush( *Brush );
+	return FAutomationWindowStyle::Get().GetBrush( *Brush );
 }
 
 
diff --git a/Engine/Source/Developer/DeviceManager/Private/Widgets/Browser/SDeviceBrowserFilterBar.cpp b/Engine/Source/Developer/DeviceManager/Private/Widgets/Browser/SDeviceBrowserFilterBar.cpp
--- a/Engine/Source/Developer/DeviceManager/Private/Widgets/Browser/SDeviceBrowserFilterBar.cpp
+++ b/Engine/Source/Developer/DeviceManager/Private/Widgets/Browser/SDeviceBrowserFilterBar.cpp
@@ -113,22 +113,20 @@
 			[
 				// platform filter
 				SNew(SComboButton)
-					.ComboButtonStyle(FAppStyle::Get(), "ToolbarComboButton")
-					.ForegroundColor(FLinearColor::White)
-					.ButtonContent()
-					[
-						SNew(STextBlock)
-						.TextStyle(FAppStyle::Get(), "Launcher.Filters.Text")
-						.Text(LOCTEXT("PlatformFiltersComboButtonText", "Platform Filters"))
-					]
-					.ContentPadding(0.0f)
-					.MenuContent()
-					[
-						SAssignNew(PlatformListView, SListView<TSharedPtr<FDeviceBrowserFilterEntry> >)
-							.ItemHeight(24.0f)
-							.ListItemsSource(&Filter->GetFilteredPlatforms())
-							.OnGenerateRow_Lambda(PlatformListViewGenerateRow)
-					]
+				.ButtonContent()
+				[
+					SNew(STextBlock)
+					.Font(FAppStyle::Get().GetFontStyle("NormalBold"))
+					.Text(LOCTEXT("PlatformFiltersComboButtonText", "Platform Filters"))
+				]
+				.ContentPadding(0.0f)
+				.MenuContent()
+				[
+					SAssignNew(PlatformListView, SListView<TSharedPtr<FDeviceBrowserFilterEntry> >)
+					.ItemHeight(24.0f)
+					.ListItemsSource(&Filter->GetFilteredPlatforms())
+					.OnGenerateRow_Lambda(PlatformListViewGenerateRow)
+				]
 			]
 
 		+ SHorizontalBox::Slot()
diff --git a/Engine/Source/Developer/LogVisualizer/Private/SVisualLoggerReport.cpp b/Engine/Source/Developer/LogVisualizer/Private/SVisualLoggerReport.cpp
--- a/Engine/Source/Developer/LogVisualizer/Private/SVisualLoggerReport.cpp
+++ b/Engine/Source/Developer/LogVisualizer/Private/SVisualLoggerReport.cpp
@@ -4,6 +4,7 @@
 #include "SlateOptMacros.h"
 #include "LogVisualizerStyle.h"
 #include "Widgets/Input/SSearchBox.h"
+#include "Styling/AppStyle.h"
 
 #define LOCTEXT_NAMESPACE "SVisualLoggerReport"
 
@@ -48,7 +49,7 @@
 				[
 					SNew(SBorder)
 					.Padding(5.0f)
-					.BorderImage(FCoreStyle::Get().GetBrush("BoxShadow"))
+					.BorderImage(FAppStyle::Get().GetBrush("BoxShadow"))
 					[
 						SNew(SBorder).Padding(2).HAlign(HAlign_Left)
 						.BorderImage(FLogVisualizerStyle::Get().GetBrush("RichText.Background"))
diff --git a/Engine/Source/Developer/OutputLog/OutputLog.Build.cs b/Engine/Source/Developer/OutputLog/OutputLog.Build.cs
--- a/Engine/Source/Developer/OutputLog/OutputLog.Build.cs
+++ b/Engine/Source/Developer/OutputLog/OutputLog.Build.cs
@@ -1,31 +1,32 @@
-// Copyright Epic Games, Inc. All Rights Reserved.
-
-using UnrealBuildTool;
-
-public class OutputLog : ModuleRules
-{
-	public OutputLog(ReadOnlyTargetRules Target) : base(Target)
-	{
-		PrivateDependencyModuleNames.AddRange(
-			new string[] {
-				"Core",
-				"CoreUObject", // @todo Mac: for some reason it's needed to link in debug on Mac
-                "InputCore",
-				"Slate",
-				"SlateCore",
-                "TargetPlatform",
-                "DesktopPlatform"
-			}
-		);
-
-		if (Target.bBuildEditor)
-        {
-            PrivateDependencyModuleNames.AddRange(
-				new string[] {
-						"EditorFramework",
-                        "UnrealEd"
-				}
-			);
-        }
-    }
-}
+// Copyright Epic Games, Inc. All Rights Reserved.
+
+using UnrealBuildTool;
+
+public class OutputLog : ModuleRules
+{
+	public OutputLog(ReadOnlyTargetRules Target) : base(Target)
+	{
+		PrivateDependencyModuleNames.AddRange(
+			new string[] {
+				"Core",
+				"CoreUObject", // @todo Mac: for some reason it's needed to link in debug on Mac
+                "InputCore",
+				"Slate",
+				"SlateCore",
+                "TargetPlatform",
+                "DesktopPlatform",
+				"ToolWidgets"
+			}
+		);
+
+		if (Target.bBuildEditor)
+        {
+            PrivateDependencyModuleNames.AddRange(
+				new string[] {
+						"EditorFramework",
+                        "UnrealEd"
+				}
+			);
+        }
+    }
+}
diff --git a/Engine/Source/Developer/OutputLog/Private/OutputLogModule.cpp b/Engine/Source/Developer/OutputLog/Private/OutputLogModule.cpp
--- a/Engine/Source/Developer/OutputLog/Private/OutputLogModule.cpp
+++ b/Engine/Source/Developer/OutputLog/Private/OutputLogModule.cpp
@@ -10,15 +10,24 @@
 #include "SDebugConsole.h"
 #include "SOutputLog.h"
 #include "SDeviceOutputLog.h"
-#include "Editor/WorkspaceMenuStructure/Public/WorkspaceMenuStructure.h"
-#include "Editor/WorkspaceMenuStructure/Public/WorkspaceMenuStructureModule.h"
+
 #include "Widgets/Docking/SDockTab.h"
 #include "Misc/ConfigCacheIni.h"
 
 #if WITH_EDITOR
 #include "Editor.h"
+#include "ISettingsModule.h"
 #endif
 
+#if WITH_EDITOR || (IS_PROGRAM && WITH_UNREAL_DEVELOPER_TOOLS)
+#include "Editor/WorkspaceMenuStructure/Public/WorkspaceMenuStructure.h"
+#include "Editor/WorkspaceMenuStructure/Public/WorkspaceMenuStructureModule.h"
+#include "OutputLogSettings.h"
+#endif
+
+#include "Internationalization/Internationalization.h"
+#include "OutputLogStyle.h"
+
 IMPLEMENT_MODULE(FOutputLogModule, OutputLog);
 
 namespace OutputLogModule
@@ -98,20 +107,34 @@
 
 void FOutputLogModule::StartupModule()
 {
+	FOutputLogStyle::Get();
+
+#if WITH_EDITOR
+	ISettingsModule* SettingsModule = FModuleManager::GetModulePtr<ISettingsModule>("Settings");
+	if (SettingsModule)
+	{
+		SettingsModule->RegisterSettings("Editor", "General", "Output Log",
+			NSLOCTEXT("OutputLog", "OutputLogSettingsName", "Output Log"),
+			NSLOCTEXT("OutputLog", "OutputLogSettingsDescription", "Set up preferences for the Output Log appearance and workflow."),
+			GetMutableDefault<UOutputLogSettings>()
+		);
+	}
+
+	FEditorDelegates::BeginPIE.AddRaw(this, &FOutputLogModule::ClearOnPIE);
+#endif
+
+#if WITH_EDITOR || (IS_PROGRAM && WITH_UNREAL_DEVELOPER_TOOLS)
 	FGlobalTabmanager::Get()->RegisterNomadTabSpawner(OutputLogModule::OutputLogTabName, FOnSpawnTab::CreateRaw(this, &FOutputLogModule::SpawnOutputLogTab))
-		.SetDisplayName(NSLOCTEXT("UnrealEditor", "OutputLogTab", "Output Log"))
-		.SetTooltipText(NSLOCTEXT("UnrealEditor", "OutputLogTooltipText", "Open the Output Log tab."))
+		.SetDisplayName(NSLOCTEXT("OutputLog", "OutputLogTab", "Output Log"))
+		.SetTooltipText(NSLOCTEXT("OutputLog", "OutputLogTooltipText", "Open the Output Log tab."))
 		.SetGroup(WorkspaceMenu::GetMenuStructure().GetDeveloperToolsLogCategory())
-		.SetIcon(FSlateIcon(FAppStyle::GetAppStyleSetName(), "Log.TabIcon"));
+		.SetIcon(FSlateIcon(FOutputLogStyle::Get().GetStyleSetName(), "Log.TabIcon"));
 
 	FGlobalTabmanager::Get()->RegisterNomadTabSpawner(OutputLogModule::DeviceOutputLogTabName, FOnSpawnTab::CreateRaw(this, &FOutputLogModule::SpawnDeviceOutputLogTab))
-		.SetDisplayName(NSLOCTEXT("UnrealEditor", "DeviceOutputLogTab", "Device Output Log"))
-		.SetTooltipText(NSLOCTEXT("UnrealEditor", "DeviceOutputLogTooltipText", "Open the Device Output Log tab."))
+		.SetDisplayName(NSLOCTEXT("OutputLog", "DeviceOutputLogTab", "Device Output Log"))
+		.SetTooltipText(NSLOCTEXT("OutputLog", "DeviceOutputLogTooltipText", "Open the Device Output Log tab."))
 		.SetGroup(WorkspaceMenu::GetMenuStructure().GetDeveloperToolsLogCategory())
-		.SetIcon(FSlateIcon(FAppStyle::GetAppStyleSetName(), "Log.TabIcon"));
-
-#if WITH_EDITOR
-	FEditorDelegates::BeginPIE.AddRaw(this, &FOutputLogModule::ClearOnPIE);
+		.SetIcon(FSlateIcon(FOutputLogStyle::Get().GetStyleSetName(), "Log.TabIcon"));
 #endif
 
 	OutputLogHistory = MakeShareable(new FOutputLogHistory);
@@ -119,16 +142,18 @@
 
 void FOutputLogModule::ShutdownModule()
 {
+#if WITH_EDITOR || (IS_PROGRAM && WITH_UNREAL_DEVELOPER_TOOLS)
 	if (FSlateApplication::IsInitialized())
 	{
 		FGlobalTabmanager::Get()->UnregisterNomadTabSpawner(OutputLogModule::OutputLogTabName);
 		FGlobalTabmanager::Get()->UnregisterNomadTabSpawner(OutputLogModule::DeviceOutputLogTabName);
 	}
+#endif
 
 #if WITH_EDITOR
 	FEditorDelegates::BeginPIE.RemoveAll(this);
 #endif
-
+	FOutputLogStyle::Shutdown();
 	OutputLogHistory.Reset();
 }
 
diff --git a/Engine/Source/Developer/OutputLog/Private/OutputLogStyle.cpp b/Engine/Source/Developer/OutputLog/Private/OutputLogStyle.cpp
new file mode 100644
--- /dev/null
+++ b/Engine/Source/Developer/OutputLog/Private/OutputLogStyle.cpp
@@ -0,0 +1,82 @@
+// Copyright Epic Games, Inc. All Rights Reserved.
+
+#include "OutputLogStyle.h"
+#include "Styling/SlateStyle.h"
+#include "Styling/SlateStyleMacros.h"
+#include "Styling/SlateTypes.h"
+#include "../Public/OutputLogSettings.h"
+#include "Styling/StyleColors.h"
+#include "Styling/StarshipCoreStyle.h"
+#include "Styling/SlateStyleRegistry.h"
+#include "Styling/CoreStyle.h"
+
+FName FOutputLogStyle::StyleName("OutputLogStyle");
+TUniquePtr<FOutputLogStyle> FOutputLogStyle::Inst(nullptr);
+
+const FName& FOutputLogStyle::GetStyleSetName() const
+{
+	return StyleName;
+}
+
+const FOutputLogStyle& FOutputLogStyle::Get()
+{
+	if (!Inst.IsValid())
+	{
+		Inst = TUniquePtr<FOutputLogStyle>(new FOutputLogStyle);
+	}
+	return *(Inst.Get());
+}
+
+void FOutputLogStyle::Shutdown()
+{
+	Inst.Reset();
+}
+
+FOutputLogStyle::FOutputLogStyle()
+	: FSlateStyleSet(StyleName)
+{
+	const FVector2D Icon16x16(16.0f, 16.0f);
+	const FVector2D Icon20x20(20.0f, 20.0f);
+
+	SetParentStyleName(FAppStyle::GetAppStyleSetName());
+
+	SetContentRoot(FPaths::EngineContentDir() / TEXT("Slate"));
+
+	// Output Log Window
+	{
+		const UOutputLogSettings* Settings = GetDefault<UOutputLogSettings>();
+		
+		const int32 LogFontSize = Settings ? Settings->LogFontSize : 9;
+
+		FTextBlockStyle NormalText = FAppStyle::GetWidgetStyle<FTextBlockStyle>("NormalText");
+
+		const FTextBlockStyle NormalLogText = FTextBlockStyle(NormalText)
+			.SetFont(DEFAULT_FONT("Mono", LogFontSize))
+			.SetColorAndOpacity(FStyleColors::Foreground)
+			.SetSelectedBackgroundColor(FStyleColors::Highlight)
+			.SetHighlightColor(FStyleColors::Black);
+
+		Set("Log.Normal", NormalLogText);
+
+		Set("Log.Command", FTextBlockStyle(NormalLogText)
+			.SetColorAndOpacity(FStyleColors::AccentGreen)
+		);
+
+		Set("Log.Warning", FTextBlockStyle(NormalLogText)
+			.SetColorAndOpacity(FStyleColors::Warning)
+		);
+
+		Set("Log.Error", FTextBlockStyle(NormalLogText)
+			.SetColorAndOpacity(FStyleColors::Error)
+		);
+
+		Set("DebugConsole.Icon", new IMAGE_BRUSH_SVG("Starship/Common/Console", Icon16x16));
+	}
+
+	FSlateStyleRegistry::RegisterSlateStyle(*this);
+}
+
+FOutputLogStyle::~FOutputLogStyle()
+{
+	FSlateStyleRegistry::UnRegisterSlateStyle(*this);
+}
diff --git a/Engine/Source/Developer/OutputLog/Private/OutputLogStyle.h b/Engine/Source/Developer/OutputLog/Private/OutputLogStyle.h
new file mode 100644
--- /dev/null
+++ b/Engine/Source/Developer/OutputLog/Private/OutputLogStyle.h
@@ -0,0 +1,29 @@
+// Copyright Epic Games, Inc. All Rights Reserved.
+
+#pragma once
+
+#include "Styling/SlateStyle.h"
+#include "Styling/SlateColor.h"
+#include "UObject/NameTypes.h"
+#include "Templates/UniquePtr.h"
+
+class FOutputLogStyle
+	: public FSlateStyleSet
+{
+public:
+	virtual const FName& GetStyleSetName() const override;
+
+	static const FOutputLogStyle& Get();
+	static void Shutdown();
+	
+	~FOutputLogStyle();
+
+private:
+	FOutputLogStyle();
+
+	void OnSettingsChanged(FName PropertyName);
+	void UpdateLogTextSize();
+
+	static FName StyleName;
+	static TUniquePtr<FOutputLogStyle> Inst;
+};
diff --git a/Engine/Source/Developer/OutputLog/Private/SDebugConsole.cpp b/Engine/Source/Developer/OutputLog/Private/SDebugConsole.cpp
--- a/Engine/Source/Developer/OutputLog/Private/SDebugConsole.cpp
+++ b/Engine/Source/Developer/OutputLog/Private/SDebugConsole.cpp
@@ -7,8 +7,8 @@
 #include "Widgets/Layout/SSpacer.h"
 #include "Widgets/Text/STextBlock.h"
 #include "Widgets/Layout/SBox.h"
-#include "Styling/AppStyle.h"
 #include "SOutputLog.h"
+#include "OutputLogStyle.h"
 
 namespace DebugConsoleDefs
 {
@@ -42,7 +42,7 @@
 				.HeightOverride( 200.0f )
 				[
 					SNew( SBorder )
-					.BorderImage( FAppStyle::GetBrush( "ToolPanel.GroupBorder" ) )
+					.BorderImage(FOutputLogStyle::Get().GetBrush( "ToolPanel.GroupBorder" ) )
 					.ColorAndOpacity( this, &SDebugConsole::GetAnimatedColorAndOpacity )
 					.BorderBackgroundColor( this, &SDebugConsole::GetAnimatedSlateColor )
 					[
@@ -63,7 +63,7 @@
 				SNew( SBorder )
 				.VAlign( VAlign_Center )
 				.Padding( FMargin(3.0f, 3.0f) )
-				.BorderImage(FAppStyle::Get().GetBrush("Brushes.Panel"))
+				.BorderImage(FOutputLogStyle::Get().GetBrush("Brushes.Panel"))
 				.ColorAndOpacity( this, &SDebugConsole::GetAnimatedColorAndOpacity )
 				.BorderBackgroundColor( this, &SDebugConsole::GetAnimatedSlateColor )
 				[
diff --git a/Engine/Source/Developer/OutputLog/Private/SDeviceOutputLog.cpp b/Engine/Source/Developer/OutputLog/Private/SDeviceOutputLog.cpp
--- a/Engine/Source/Developer/OutputLog/Private/SDeviceOutputLog.cpp
+++ b/Engine/Source/Developer/OutputLog/Private/SDeviceOutputLog.cpp
@@ -3,17 +3,17 @@
 #include "SDeviceOutputLog.h"
 #include "Framework/Text/TextLayout.h"
 #include "Widgets/Text/STextBlock.h"
-#include "Widgets/Input/SComboButton.h"
 #include "Misc/ScopeLock.h"
 #include "Modules/ModuleManager.h"
 #include "Widgets/Images/SImage.h"
 #include "Framework/Commands/UIAction.h"
 #include "Widgets/Input/SMultiLineEditableTextBox.h"
 #include "Framework/MultiBox/MultiBoxBuilder.h"
-#include "Styling/AppStyle.h"
 #include "Interfaces/ITargetPlatform.h"
 #include "Interfaces/ITargetPlatformManagerModule.h"
 #include "PlatformInfo.h"
+#include "OutputLogStyle.h"
+#include "SSimpleComboButton.h"
 
 static bool IsSupportedPlatform(ITargetPlatform* Platform)
 {
@@ -29,8 +29,8 @@
 	MessagesTextMarshaller = FOutputLogTextLayoutMarshaller::Create(TArray<TSharedPtr<FOutputLogMessage>>(), &Filter);
 
 	MessagesTextBox = SNew(SMultiLineEditableTextBox)
-		.Style(FAppStyle::Get(), "Log.TextBox")
-		.TextStyle(FAppStyle::Get(), "Log.Normal")
+		.Style(FOutputLogStyle::Get(), "Log.TextBox")
+		.TextStyle(FOutputLogStyle::Get(), "Log.Normal")
 		.ForegroundColor(FLinearColor::Gray)
 		.Marshaller(MessagesTextMarshaller)
 		.IsReadOnly(true)
@@ -59,8 +59,7 @@
 				.AutoWidth()
 				[
 					SAssignNew(TargetDeviceComboButton, SComboButton)
-					.ComboButtonStyle(FAppStyle::Get(), "GenericFilters.ComboButtonStyle")
-					.ForegroundColor(FLinearColor::White)
+					.ForegroundColor(FSlateColor::UseForeground())
 					.OnGetMenuContent(this, &SDeviceOutputLog::MakeDeviceComboButtonMenu)
 					.ContentPadding(FMargin(4.0f, 0.0f))
 					.ButtonContent()
@@ -81,7 +80,7 @@
 						.VAlign(VAlign_Center)
 						[
 							SNew(STextBlock)
-							.TextStyle(FAppStyle::Get(), "GenericFilters.TextStyle")
+							.Font(FOutputLogStyle::Get().GetFontStyle("NormalFontBold"))
 							.Text(this, &SDeviceOutputLog::GetSelectedTargetDeviceText)
 						]
 					]
@@ -278,7 +277,7 @@
 	TSharedPtr<FTargetDeviceEntry> DeviceEntry = MakeShareable(new FTargetDeviceEntry());
 	
 	DeviceEntry->DeviceId = TargetDevice->GetId();
-	DeviceEntry->DeviceIconBrush = FAppStyle::GetBrush(DeviceIconStyleName);
+	DeviceEntry->DeviceIconBrush = FOutputLogStyle::Get().GetBrush(DeviceIconStyleName);
 	DeviceEntry->DeviceWeakPtr = TargetDevice;
 	
 	DeviceList.Add(DeviceEntry);
@@ -364,7 +363,7 @@
 	}
 	else
 	{
-		return FAppStyle::GetBrush("Launcher.Instance_Unknown");
+		return FOutputLogStyle::Get().GetBrush("Launcher.Instance_Unknown");
 	}
 }
 
diff --git a/Engine/Source/Developer/OutputLog/Private/SOutputLog.cpp b/Engine/Source/Developer/OutputLog/Private/SOutputLog.cpp
--- a/Engine/Source/Developer/OutputLog/Private/SOutputLog.cpp
+++ b/Engine/Source/Developer/OutputLog/Private/SOutputLog.cpp
@@ -17,7 +17,6 @@
 #include "Widgets/Input/SMultiLineEditableTextBox.h"
 #include "Widgets/Input/SComboButton.h"
 #include "Widgets/Views/SListView.h"
-#include "Styling/AppStyle.h"
 #include "Widgets/Input/SSearchBox.h"
 #include "Widgets/Input/SCheckBox.h"
 #include "Widgets/Images/SImage.h"
@@ -30,10 +29,9 @@
 #include "Widgets/Docking/SDockTab.h"
 #include "OutputLogModule.h"
 #include "Widgets/Text/SlateEditableTextTypes.h"
+#include "OutputLogSettings.h"
+#include "OutputLogStyle.h"
 
-#if WITH_EDITOR
-#include "Settings/EditorStyleSettings.h"
-#endif
 
 #define LOCTEXT_NAMESPACE "SOutputLog"
 
@@ -181,7 +179,7 @@
 			[
 				SNew(SComboButton)
 				.IsEnabled(this, &SConsoleInputBox::IsCommandExecutorMenuEnabled)
-				.ComboButtonStyle(FAppStyle::Get(), "SimpleComboButton")
+				.ComboButtonStyle(FOutputLogStyle::Get(), "SimpleComboButton")
 				.ContentPadding(0)
 				.OnGetMenuContent(this, &SConsoleInputBox::GetCommandExecutorMenuContent)
 				.ButtonContent()
@@ -195,7 +193,7 @@
 					[
 						SNew(SImage)
 						.ColorAndOpacity(FSlateColor::UseForeground())
-						.Image(FAppStyle::Get().GetBrush("DebugConsole.Icon"))
+						.Image(FOutputLogStyle::Get().GetBrush("DebugConsole.Icon"))
 					]
 					+ SHorizontalBox::Slot()
 					.VAlign(VAlign_Center)
@@ -210,7 +208,7 @@
 			+SHorizontalBox::Slot()
 			[
 				SAssignNew(InputText, SMultiLineEditableTextBox)
-				.Font(FAppStyle::Get().GetWidgetStyle<FTextBlockStyle>("Log.Normal").Font)
+				.Font(FOutputLogStyle::Get().GetWidgetStyle<FTextBlockStyle>("Log.Normal").Font)
 				.HintText(this, &SConsoleInputBox::GetActiveCommandExecutorHintText)
 				.AllowMultiLine(this, &SConsoleInputBox::GetActiveCommandExecutorAllowMultiLine)
 				.OnTextCommitted(this, &SConsoleInputBox::OnTextCommitted)
@@ -225,7 +223,7 @@
 		.MenuContent
 		(
 			SNew(SBorder)
-			.BorderImage(FAppStyle::GetBrush("Menu.Background"))
+			.BorderImage(FOutputLogStyle::Get().GetBrush("Menu.Background"))
 			.Padding( FMargin(2) )
 			[
 				SNew(SBox)
@@ -315,7 +313,7 @@
 		[
 			SNew(STextBlock)
 			.Text(FText::FromString(SanitizedText))
-			.TextStyle(FAppStyle::Get(), "Log.Normal")
+			.TextStyle(FOutputLogStyle::Get(), "Log.Normal")
 			.HighlightText(Suggestions.SuggestionsHighlight)
 			.ColorAndOpacity(FSlateColor::UseForeground())
 		];
@@ -844,11 +842,8 @@
 		MakeDirty();
 	}
 
+	const ELogCategoryColorizationMode CategoryColorizationMode = GetDefault<UOutputLogSettings>()->CategoryColorizationMode;
 
-#if WITH_EDITOR
-	ELogCategoryColorizationMode CategoryColorizationMode = GetDefault<UEditorStyleSettings>()->CategoryColorizationMode;
-#endif
-
 	TArray<FTextLayout::FNewLineData> LinesToAdd;
 	LinesToAdd.Reserve(NumPendingMessages);
 	TArray<FTextLineHighlight> Highlights;
@@ -879,13 +874,13 @@
 
 		++NumAddedMessages;
 
-		const FTextBlockStyle& MessageTextStyle = FAppStyle::Get().GetWidgetStyle<FTextBlockStyle>(Message->Style);
+		const FTextBlockStyle& MessageTextStyle = FOutputLogStyle::Get().GetWidgetStyle<FTextBlockStyle>(Message->Style);
 
 		TSharedRef<FString> LineText = Message->Message;
 
 		TArray<TSharedRef<IRun>> Runs;
 
-#if WITH_EDITOR
+
 		switch (CategoryColorizationMode)
 		{
 		case ELogCategoryColorizationMode::None:
@@ -941,9 +936,6 @@
 			}
 			break;
 		}
-#else
-		Runs.Add(FSlateTextRun::Create(FRunInfo(), LineText, MessageTextStyle));
-#endif
 
 		if (!Message->Category.IsNone() && (Message->Category == CategoryToHighlight))
 		{
@@ -1058,8 +1050,8 @@
 	MessagesTextMarshaller = FOutputLogTextLayoutMarshaller::Create(InArgs._Messages, &Filter);
 
 	MessagesTextBox = SNew(SMultiLineEditableTextBox)
-		.Style(FAppStyle::Get(), "Log.TextBox")
-		.TextStyle(FAppStyle::Get(), "Log.Normal")
+		.Style(FOutputLogStyle::Get(), "Log.TextBox")
+		.TextStyle(FOutputLogStyle::Get(), "Log.Normal")
 		.Marshaller(MessagesTextMarshaller)
 		.IsReadOnly(true)
 		.AlwaysShowScrollbars(true)
@@ -1093,7 +1085,7 @@
 			.HAlign(HAlign_Left)
 			[
 				SNew(SComboButton)
-				.ComboButtonStyle(FAppStyle::Get(), "SimpleComboButton")
+				.ComboButtonStyle(FOutputLogStyle::Get(), "SimpleComboButton")
 				.ToolTipText(LOCTEXT("AddFilterToolTip", "Add an output log filter."))
 				.OnGetMenuContent(this, &SOutputLog::MakeAddFilterMenu)
 				.ButtonContent()
@@ -1103,7 +1095,7 @@
 					.AutoWidth()
 					[
 						SNew(SImage)
-						.Image(FAppStyle::Get().GetBrush("Icons.Filter"))
+						.Image(FOutputLogStyle::Get().GetBrush("Icons.Filter"))
 						.ColorAndOpacity(FSlateColor::UseForeground())
 					]
 					+SHorizontalBox::Slot()
@@ -1128,7 +1120,7 @@
 			.AutoWidth()
 			[
 				SNew(SComboButton)
-				.ComboButtonStyle(FAppStyle::Get(), "SimpleComboButton")
+				.ComboButtonStyle(FOutputLogStyle::Get(), "SimpleComboButton")
 				.OnGetMenuContent(this, &SOutputLog::GetViewButtonContent, InArgs._SettingsMenuFlags)
 				.ButtonContent()
 				[
@@ -1138,7 +1130,7 @@
 					.VAlign(VAlign_Center)
 					[
 						SNew(SImage)
-						.Image(FAppStyle::Get().GetBrush("Icons.Settings"))
+						.Image(FOutputLogStyle::Get().GetBrush("Icons.Settings"))
 						.ColorAndOpacity(FSlateColor::UseForeground())
 					]
 					+ SHorizontalBox::Slot()
@@ -1179,7 +1171,7 @@
 
 #if WITH_EDITOR
 	// Listen for style changes
-	UEditorStyleSettings* Settings = GetMutableDefault<UEditorStyleSettings>();
+	UOutputLogSettings* Settings = GetMutableDefault<UOutputLogSettings>();
 	SettingsWatchHandle = Settings->OnSettingChanged().AddRaw(this, &SOutputLog::HandleSettingChanged);
 #endif
 
@@ -1199,7 +1191,7 @@
 #if WITH_EDITOR
 	if (UObjectInitialized() && !GExitPurge)
 	{
-		UEditorStyleSettings* Settings = GetMutableDefault<UEditorStyleSettings>();
+		UOutputLogSettings* Settings = GetMutableDefault<UOutputLogSettings>();
 		Settings->OnSettingChanged().Remove(SettingsWatchHandle);
 	}
 #endif
@@ -1264,13 +1256,11 @@
 
 		// Determine how to format timestamps
 		static ELogTimes::Type LogTimestampMode = ELogTimes::None;
-#if WITH_EDITOR
 		if (UObjectInitialized() && !GExitPurge)
 		{
 			// Logging can happen very late during shutdown, even after the UObject system has been torn down, hence the init check above
-			LogTimestampMode = GetDefault<UEditorStyleSettings>()->LogTimestampMode;
+			LogTimestampMode = GetDefault<UOutputLogSettings>()->LogTimestampMode;
 		}
-#endif
 
 		const int32 OldNumMessages = OutMessages.Num();
 
@@ -1464,15 +1454,19 @@
 
 bool SOutputLog::IsWordWrapEnabled() const
 {
-	bool WordWrapEnabled = false;
-	GConfig->GetBool(TEXT("/Script/UnrealEd.EditorPerProjectUserSettings"), TEXT("bEnableOutputLogWordWrap"), WordWrapEnabled, GEditorPerProjectIni);
-	return WordWrapEnabled;
+	const UOutputLogSettings* Settings = GetDefault<UOutputLogSettings>();
+	return Settings ? Settings->bEnableOutputLogWordWrap : false;
 }
 
 void SOutputLog::SetWordWrapEnabled(ECheckBoxState InValue)
 {
-	const bool WordWrapEnabled = (InValue == ECheckBoxState::Checked);
-	GConfig->SetBool(TEXT("/Script/UnrealEd.EditorPerProjectUserSettings"), TEXT("bEnableOutputLogWordWrap"), WordWrapEnabled, GEditorPerProjectIni);
+	const bool bWordWrapEnabled = (InValue == ECheckBoxState::Checked);
+	UOutputLogSettings* Settings = GetMutableDefault<UOutputLogSettings>();
+	if (Settings)
+	{
+		Settings->bEnableOutputLogWordWrap = bWordWrapEnabled;
+		Settings->SaveConfig();
+	}
 
 	if (!bIsUserScrolled)
 	{
@@ -1480,18 +1474,24 @@
 	}
 }
 
+#if WITH_EDITOR
 bool SOutputLog::IsClearOnPIEEnabled() const
 {
-	bool ClearOnPIEEnabled = false;
-	GConfig->GetBool(TEXT("/Script/UnrealEd.EditorPerProjectUserSettings"), TEXT("bEnableOutputLogClearOnPIE"), ClearOnPIEEnabled, GEditorPerProjectIni);
-	return ClearOnPIEEnabled;
+	const UOutputLogSettings* Settings = GetDefault<UOutputLogSettings>();
+	return Settings ? Settings->bEnableOutputLogClearOnPIE : false;
 }
 
 void SOutputLog::SetClearOnPIE(ECheckBoxState InValue)
 {
-	const bool ClearOnPIEEnabled = (InValue == ECheckBoxState::Checked);
-	GConfig->SetBool(TEXT("/Script/UnrealEd.EditorPerProjectUserSettings"), TEXT("bEnableOutputLogClearOnPIE"), ClearOnPIEEnabled, GEditorPerProjectIni);
+	const bool bClearOnPIEEnabled = (InValue == ECheckBoxState::Checked);
+	UOutputLogSettings* Settings = GetMutableDefault<UOutputLogSettings>();
+	if (Settings)
+	{
+		Settings->bEnableOutputLogClearOnPIE = bClearOnPIEEnabled;
+		Settings->SaveConfig();
+	}
 }
+#endif
 
 void SOutputLog::BuildInitialLogCategoryFilter(const FArguments& InArgs)
 {
@@ -1804,7 +1804,7 @@
 		MenuBuilder.AddMenuEntry(
 			LOCTEXT("FindSourceFile", "Open Source Location"),
 			LOCTEXT("FindSourceFileTooltip", "Opens the folder containing the source of the Output Log."),
-			FSlateIcon(FAppStyle::GetAppStyleSetName(), "OutputLog.OpenSourceLocation"),
+			FSlateIcon(FOutputLogStyle::Get().GetStyleSetName(), "OutputLog.OpenSourceLocation"),
 			FUIAction(
 				FExecuteAction::CreateSP(this, &SOutputLog::OpenLogFileInExplorer)
 			)
@@ -1817,7 +1817,7 @@
 		MenuBuilder.AddMenuEntry(
 			LOCTEXT("OpenInExternalEditor", "Open In External Editor"),
 			LOCTEXT("OpenInExternalEditorTooltip", "Opens the Output Log in the default external editor."),
-			FSlateIcon(FAppStyle::GetAppStyleSetName(), "OutputLog.OpenInExternalEditor"),
+			FSlateIcon(FOutputLogStyle::Get().GetStyleSetName(), "OutputLog.OpenInExternalEditor"),
 			FUIAction(
 				FExecuteAction::CreateSP(this, &SOutputLog::OpenLogFileInExternalEditor)
 			)
@@ -1833,7 +1833,7 @@
 	{
 		return
 			SNew(SButton)
-			.ButtonStyle(FAppStyle::Get(), "SimpleButton")
+			.ButtonStyle(FOutputLogStyle::Get(), "SimpleButton")
 			.ToolTipText(LOCTEXT("DockInLayout_Tooltip", "Docks this output log in the current layout.\nThe drawer will still be usable as a temporary log."))
 			.ContentPadding(FMargin(1, 0))
 			.Visibility_Lambda(
@@ -1851,7 +1851,7 @@
 				[
 					SNew(SImage)
 					.ColorAndOpacity(FSlateColor::UseForeground())
-					.Image(FAppStyle::Get().GetBrush("Icons.Layout"))
+					.Image(FOutputLogStyle::Get().GetBrush("Icons.Layout"))
 				]
 				+ SHorizontalBox::Slot()
 				.VAlign(VAlign_Center)
diff --git a/Engine/Source/Developer/OutputLog/Private/SOutputLog.h b/Engine/Source/Developer/OutputLog/Private/SOutputLog.h
--- a/Engine/Source/Developer/OutputLog/Private/SOutputLog.h
+++ b/Engine/Source/Developer/OutputLog/Private/SOutputLog.h
@@ -461,9 +461,11 @@
 
 	void SetWordWrapEnabled(ECheckBoxState InValue);
 
+#if WITH_EDITOR
 	bool IsClearOnPIEEnabled() const;
 
 	void SetClearOnPIE(ECheckBoxState InValue);
+#endif
 
 	FSlateColor GetViewButtonForegroundColor() const;
 
diff --git a/Engine/Source/Developer/OutputLog/Public/OutputLogSettings.h b/Engine/Source/Developer/OutputLog/Public/OutputLogSettings.h
new file mode 100644
--- /dev/null
+++ b/Engine/Source/Developer/OutputLog/Public/OutputLogSettings.h
@@ -0,0 +1,98 @@
+// Copyright Epic Games, Inc. All Rights Reserved.
+
+#pragma once
+
+#include "UObject/ObjectMacros.h"
+#include "UObject/Object.h"
+#include "Styling/SlateBrush.h"
+#include "Containers/EnumAsByte.h"
+#include "Misc/OutputDevice.h"
+#include "Delegates/DelegateCombinations.h"
+#include "OutputLogSettings.generated.h"
+
+UENUM()
+enum class ELogCategoryColorizationMode : uint8
+{
+	/** Do not colorize based on log categories */
+	None,
+
+	/** Colorize the entire log line, but not warnings or errors */
+	ColorizeWholeLine,
+
+	/** Colorize only the category name (including on warnings and errors) */
+	ColorizeCategoryOnly,
+
+	/** Colorize the background of the category name (including on warnings and errors) */
+	ColorizeCategoryAsBadge
+};
+
+/**
+ * Implements the Editor style settings.
+ */
+UCLASS(config=EditorPerProjectUserSettings)
+class OUTPUTLOG_API UOutputLogSettings : public UObject
+{
+	GENERATED_BODY()
+
+public:
+	UOutputLogSettings()
+	{
+		bCycleToOutputLogDrawer = true;
+		LogTimestampMode = ELogTimes::None;
+	}
+
+	/** The font size used in the output log */
+	UPROPERTY(EditAnywhere, config, Category="Output Log", meta=(DisplayName="Log Font Size", ConfigRestartRequired=true))
+	int32 LogFontSize;
+
+	/** The display mode for timestamps in the output log window */
+	UPROPERTY(EditAnywhere, config, Category="Output Log", meta=(DisplayName = "Output Log Window Timestamp Mode"))
+	TEnumAsByte<ELogTimes::Type> LogTimestampMode;
+
+	/** How should categories be colorized in the output log? */
+	UPROPERTY(EditAnywhere, config, Category = "Output Log")
+	ELogCategoryColorizationMode CategoryColorizationMode;
+
+	/**
+	 * If checked pressing the console command shortcut will cycle between focusing the status bar console, opening the output log drawer, and back to the previous focus target. 
+	 * If unchecked, the console command shortcut will only focus the status bar console
+	 */
+	UPROPERTY(EditAnywhere, config, Category = "Output Log", meta = (DisplayName = "Open Output Log Drawer with Console Command Shortcut"))
+	bool bCycleToOutputLogDrawer;
+
+	UPROPERTY(EditAnywhere, config, Category = "Output Log")
+	bool bEnableOutputLogWordWrap;
+
+#if WITH_EDITORONLY_DATA
+	UPROPERTY(config)
+	bool bEnableOutputLogClearOnPIE;
+#endif
+
+public:
+
+	/**
+	 * Returns an event delegate that is executed when a setting has changed.
+	 *
+	 * @return The delegate.
+	 */
+	DECLARE_EVENT_OneParam(UOutputLogSettings, FSettingChangedEvent, FName /*PropertyName*/);
+	FSettingChangedEvent& OnSettingChanged() { return SettingChangedEvent; }
+
+protected:
+#if WITH_EDITOR
+	// UObject overrides
+	virtual void PostEditChangeProperty(struct FPropertyChangedEvent& PropertyChangedEvent) override
+	{
+		Super::PostEditChangeProperty(PropertyChangedEvent);
+
+		SaveConfig();
+		const FName PropertyName = (PropertyChangedEvent.Property != nullptr) ? PropertyChangedEvent.Property->GetFName() : NAME_None;
+		SettingChangedEvent.Broadcast(PropertyName);
+	}
+#endif
+
+private:
+
+	// Holds an event delegate that is executed when a setting has changed.
+	FSettingChangedEvent SettingChangedEvent;
+};
diff --git a/Engine/Source/Developer/Profiler/Private/ProfilerModule.cpp b/Engine/Source/Developer/Profiler/Private/ProfilerModule.cpp
--- a/Engine/Source/Developer/Profiler/Private/ProfilerModule.cpp
+++ b/Engine/Source/Developer/Profiler/Private/ProfilerModule.cpp
@@ -11,6 +11,7 @@
 #include "Widgets/SWidget.h"
 #include "Widgets/SProfilerWindow.h"
 #include "Widgets/Docking/SDockTab.h"
+#include "ProfilerStyle.h"
 
 /**
  * Implements the FProfilerModule module.
@@ -22,7 +23,10 @@
 
 	virtual TSharedRef<SWidget> CreateProfilerWindow( const TSharedRef<ISessionManager>& InSessionManager, const TSharedRef<SDockTab>& ConstructUnderMajorTab ) override;
 
-	virtual void StartupModule() override { }
+	virtual void StartupModule() override 
+	{ 
+		FProfilerStyle::Get();
+	}
 	virtual void ShutdownModule() override;
 
 	virtual bool SupportsDynamicReloading() override
@@ -65,6 +69,7 @@
 
 void FProfilerModule::ShutdownModule()
 {
+	FProfilerStyle::Shutdown();
 #if STATS
 	if (FProfilerManager::Get().IsValid())
 	{
diff --git a/Engine/Source/Developer/Profiler/Private/ProfilerStyle.cpp b/Engine/Source/Developer/Profiler/Private/ProfilerStyle.cpp
new file mode 100644
--- /dev/null
+++ b/Engine/Source/Developer/Profiler/Private/ProfilerStyle.cpp
@@ -0,0 +1,212 @@
+// Copyright Epic Games, Inc. All Rights Reserved.
+
+#include "ProfilerStyle.h"
+#include "Styling/SlateStyle.h"
+#include "Styling/SlateStyleMacros.h"
+#include "Styling/SlateTypes.h"
+#include "Styling/StyleColors.h"
+#include "Styling/StarshipCoreStyle.h"
+#include "Styling/SlateStyleRegistry.h"
+#include "Styling/CoreStyle.h"
+
+FName FProfilerStyle::StyleName("ProfilerStyle");
+TUniquePtr<FProfilerStyle> FProfilerStyle::Inst(nullptr);
+
+const FName& FProfilerStyle::GetStyleSetName() const
+{
+	return StyleName;
+}
+
+const FProfilerStyle& FProfilerStyle::Get()
+{
+	if (!Inst.IsValid())
+	{
+		Inst = TUniquePtr<FProfilerStyle>(new FProfilerStyle);
+	}
+	return *(Inst.Get());
+}
+
+void FProfilerStyle::Shutdown()
+{
+	Inst.Reset();
+}
+
+FProfilerStyle::FProfilerStyle()
+	: FSlateStyleSet(StyleName)
+{
+	const FVector2D Icon12x12(12.0f, 12.0f);
+	const FVector2D Icon16x16(16.0f, 16.0f);
+	const FVector2D Icon20x20(20.0f, 20.0f);
+	const FVector2D Icon32x32(32.0f, 32.0f);
+	const FVector2D Icon40x40(40.0f, 40.0f);
+
+	SetParentStyleName(FAppStyle::GetAppStyleSetName());
+
+	SetContentRoot(FPaths::EngineContentDir() / TEXT("Slate"));
+
+	// Profiler
+	{
+		FTextBlockStyle NormalText = FAppStyle::GetWidgetStyle<FTextBlockStyle>("NormalText");
+		// Profiler group brushes
+		Set("Profiler.Group.16", new BOX_BRUSH("Icons/Profiler/GroupBorder-16Gray", FMargin(4.0f / 16.0f)));
+
+		// Profiler toolbar icons
+		Set("Profiler.Tab", new IMAGE_BRUSH_SVG("Starship/Common/Visualizer", Icon16x16));
+		Set("Profiler.Tab.GraphView", new IMAGE_BRUSH("Icons/Profiler/Profiler_Graph_View_Tab_16x", Icon16x16));
+		Set("Profiler.Tab.EventGraph", new IMAGE_BRUSH("Icons/Profiler/profiler_OpenEventGraph_32x", Icon16x16));
+		Set("Profiler.Tab.FiltersAndPresets", new IMAGE_BRUSH("Icons/Profiler/Profiler_Filter_Presets_Tab_16x", Icon16x16));
+
+		Set("ProfilerCommand.ProfilerManager_Load", new IMAGE_BRUSH("Icons/Profiler/Profiler_Load_Profiler_40x", Icon40x40));
+		Set("ProfilerCommand.ProfilerManager_Load.Small", new IMAGE_BRUSH("Icons/Profiler/Profiler_Load_Profiler_40x", Icon20x20));
+
+		Set("ProfilerCommand.ProfilerManager_LoadMultiple", new IMAGE_BRUSH("Icons/Profiler/Profiler_LoadMultiple_Profiler_40x", Icon40x40));
+		Set("ProfilerCommand.ProfilerManager_LoadMultiple.Small", new IMAGE_BRUSH("Icons/Profiler/Profiler_LoadMultiple_Profiler_40x", Icon20x20));
+
+		Set("ProfilerCommand.ProfilerManager_Save", new IMAGE_BRUSH_SVG("Starship/Common/save", Icon16x16));
+		Set("ProfilerCommand.ProfilerManager_Save.Small", new IMAGE_BRUSH_SVG("Starship/Common/save", Icon16x16));
+
+		Set("ProfilerCommand.ProfilerManager_ToggleLivePreview", new IMAGE_BRUSH("Automation/RefreshTests", Icon40x40));
+		Set("ProfilerCommand.ProfilerManager_ToggleLivePreview.Small", new IMAGE_BRUSH("Automation/RefreshTests", Icon20x20));
+
+		Set("ProfilerCommand.StatsProfiler", new IMAGE_BRUSH("Icons/Profiler/profiler_stats_40x", Icon40x40));
+		Set("ProfilerCommand.StatsProfiler.Small", new IMAGE_BRUSH("Icons/Profiler/profiler_stats_40x", Icon20x20));
+
+		Set("ProfilerCommand.MemoryProfiler", new IMAGE_BRUSH("Icons/Profiler/profiler_mem_40x", Icon40x40));
+		Set("ProfilerCommand.MemoryProfiler.Small", new IMAGE_BRUSH("Icons/Profiler/profiler_mem_40x", Icon20x20));
+
+		Set("ProfilerCommand.FPSChart", new IMAGE_BRUSH("Icons/Profiler/Profiler_FPS_Chart_40x", Icon40x40));
+		Set("ProfilerCommand.FPSChart.Small", new IMAGE_BRUSH("Icons/Profiler/Profiler_FPS_Chart_40x", Icon20x20));
+
+		Set("ProfilerCommand.OpenSettings", new IMAGE_BRUSH("Icons/Profiler/Profiler_Settings_40x", Icon40x40));
+		Set("ProfilerCommand.OpenSettings.Small", new IMAGE_BRUSH("Icons/Profiler/Profiler_Settings_40x", Icon20x20));
+
+		Set("ProfilerCommand.ToggleDataPreview", new IMAGE_BRUSH("Icons/Profiler/profiler_sync_40x", Icon40x40));
+		Set("ProfilerCommand.ToggleDataPreview.Small", new IMAGE_BRUSH("Icons/Profiler/profiler_sync_40x", Icon20x20));
+
+		Set("ProfilerCommand.ToggleDataCapture", new IMAGE_BRUSH("Icons/Profiler/Profiler_Data_Capture_40x", Icon40x40));
+		Set("ProfilerCommand.ToggleDataCapture.Small", new IMAGE_BRUSH("Icons/Profiler/Profiler_Data_Capture_40x", Icon20x20));
+
+		Set("ProfilerCommand.ToggleDataCapture.Checked", new IMAGE_BRUSH_SVG("Starship/Common/stop", Icon40x40));
+		Set("ProfilerCommand.ToggleDataCapture.Checked.Small", new IMAGE_BRUSH_SVG("Starship/Common/stop", Icon20x20));
+
+		Set("ProfilerCommand.ToggleShowDataGraph", new IMAGE_BRUSH("Icons/Profiler/profiler_ShowGraphData_32x", Icon32x32));
+		Set("ProfilerCommand.OpenEventGraph", new IMAGE_BRUSH("Icons/Profiler/profiler_OpenEventGraph_32x", Icon16x16));
+
+		// Tooltip hint icon
+		Set("Profiler.Tooltip.HintIcon10", new IMAGE_BRUSH("Icons/Profiler/Profiler_Custom_Tooltip_12x", Icon12x12));
+
+		// Text styles
+		Set("Profiler.CaptionBold", FTextBlockStyle(NormalText)
+			.SetFont(DEFAULT_FONT("Bold", 10))
+			.SetColorAndOpacity(FLinearColor::White)
+			.SetShadowOffset(FVector2D(1.0f, 1.0f))
+			.SetShadowColorAndOpacity(FLinearColor(0.f, 0.f, 0.f, 0.8f))
+		);
+
+		Set("Profiler.TooltipBold", FTextBlockStyle(NormalText)
+			.SetFont(DEFAULT_FONT("Bold", 8))
+			.SetColorAndOpacity(FLinearColor(0.5f, 0.5f, 0.5f, 1.0f))
+			.SetShadowOffset(FVector2D(1.0f, 1.0f))
+			.SetShadowColorAndOpacity(FLinearColor(0.f, 0.f, 0.f, 0.8f))
+		);
+
+		Set("Profiler.Tooltip", FTextBlockStyle(NormalText)
+			.SetFont(DEFAULT_FONT("Regular", 8))
+			.SetColorAndOpacity(FLinearColor::White)
+			.SetShadowOffset(FVector2D(1.0f, 1.0f))
+			.SetShadowColorAndOpacity(FLinearColor(0.f, 0.f, 0.f, 0.8f))
+		);
+
+		// Event graph icons
+		Set("Profiler.EventGraph.SetRoot", new IMAGE_BRUSH("Icons/Profiler/profiler_SetRoot_32x", Icon32x32));
+		Set("Profiler.EventGraph.CullEvents", new IMAGE_BRUSH("Icons/Profiler/Profiler_Cull_Events_16x", Icon16x16));
+		Set("Profiler.EventGraph.FilterEvents", new IMAGE_BRUSH("Icons/Profiler/Profiler_Filter_Events_16x", Icon16x16));
+
+		Set("Profiler.EventGraph.SelectStack", new IMAGE_BRUSH("Icons/Profiler/profiler_SelectStack_32x", Icon32x32));
+
+		Set("Profiler.EventGraph.ExpandAll", new IMAGE_BRUSH("Icons/Profiler/profiler_ExpandAll_32x", Icon32x32));
+		Set("Profiler.EventGraph.CollapseAll", new IMAGE_BRUSH("Icons/Profiler/profiler_CollapseAll_32x", Icon32x32));
+
+		Set("Profiler.EventGraph.ExpandSelection", new IMAGE_BRUSH("Icons/Profiler/profiler_ExpandSelection_32x", Icon32x32));
+		Set("Profiler.EventGraph.CollapseSelection", new IMAGE_BRUSH("Icons/Profiler/profiler_CollapseSelection_32x", Icon32x32));
+
+		Set("Profiler.EventGraph.ExpandThread", new IMAGE_BRUSH("Icons/Profiler/profiler_ExpandThread_32x", Icon32x32));
+		Set("Profiler.EventGraph.CollapseThread", new IMAGE_BRUSH("Icons/Profiler/profiler_CollapseThread_32x", Icon32x32));
+
+		Set("Profiler.EventGraph.ExpandHotPath", new IMAGE_BRUSH("Icons/Profiler/profiler_ExpandHotPath_32x", Icon32x32));
+		Set("Profiler.EventGraph.HotPathSmall", new IMAGE_BRUSH("Icons/Profiler/profiler_HotPath_32x", Icon12x12));
+
+		Set("Profiler.EventGraph.ExpandHotPath16", new IMAGE_BRUSH("Icons/Profiler/profiler_HotPath_32x", Icon16x16));
+
+		Set("Profiler.EventGraph.GameThread", new IMAGE_BRUSH("Icons/Profiler/profiler_GameThread_32x", Icon32x32));
+		Set("Profiler.EventGraph.RenderThread", new IMAGE_BRUSH("Icons/Profiler/profiler_RenderThread_32x", Icon32x32));
+
+		Set("Profiler.EventGraph.ViewColumn", new IMAGE_BRUSH("Icons/Profiler/profiler_ViewColumn_32x", Icon32x32));
+		Set("Profiler.EventGraph.ResetColumn", new IMAGE_BRUSH("Icons/Profiler/profiler_ResetColumn_32x", Icon32x32));
+
+		Set("Profiler.EventGraph.HistoryBack", new IMAGE_BRUSH("Icons/Profiler/Profiler_History_Back_16x", Icon16x16));
+		Set("Profiler.EventGraph.HistoryForward", new IMAGE_BRUSH("Icons/Profiler/Profiler_History_Fwd_16x", Icon16x16));
+
+		Set("Profiler.EventGraph.MaximumIcon", new IMAGE_BRUSH("Icons/Profiler/Profiler_Max_Event_Graph_16x", Icon16x16));
+		Set("Profiler.EventGraph.AverageIcon", new IMAGE_BRUSH("Icons/Profiler/Profiler_Average_Event_Graph_16x", Icon16x16));
+
+		Set("Profiler.EventGraph.FlatIcon", new IMAGE_BRUSH("Icons/Profiler/Profiler_Events_Flat_16x", Icon16x16));
+		Set("Profiler.EventGraph.FlatCoalescedIcon", new IMAGE_BRUSH("Icons/Profiler/Profiler_Events_Flat_Coalesced_16x", Icon16x16));
+		Set("Profiler.EventGraph.HierarchicalIcon", new IMAGE_BRUSH("Icons/Profiler/Profiler_Events_Hierarchial_16x", Icon16x16));
+
+		Set("Profiler.EventGraph.HasCulledEventsSmall", new IMAGE_BRUSH("Icons/Profiler/Profiler_Has_Culled_Children_12x", Icon12x12));
+		Set("Profiler.EventGraph.CulledEvent", new IMAGE_BRUSH("Icons/Profiler/Profiler_Culled_12x", Icon12x12));
+		Set("Profiler.EventGraph.FilteredEvent", new IMAGE_BRUSH("Icons/Profiler/Profiler_Filtered_12x", Icon12x12));
+
+		Set("Profiler.EventGraph.DarkText", FTextBlockStyle(NormalText)
+			.SetFont(DEFAULT_FONT("Regular", 8))
+			.SetColorAndOpacity(FLinearColor::Black)
+			.SetShadowOffset(FVector2D(0.0f, 0.0f))
+		);
+
+		// Thread-view
+		Set("Profiler.ThreadView.SampleBorder", new BOX_BRUSH("Icons/Profiler/Profiler_ThreadView_SampleBorder_16x", FMargin(2.0f / 16.0f)));
+
+		// Event graph selected event border
+		Set("Profiler.EventGraph.Border.TB", new BOX_BRUSH("Icons/Profiler/Profiler_Border_TB_16x", FMargin(4.0f / 16.0f)));
+		Set("Profiler.EventGraph.Border.L", new BOX_BRUSH("Icons/Profiler/Profiler_Border_L_16x", FMargin(4.0f / 16.0f)));
+		Set("Profiler.EventGraph.Border.R", new BOX_BRUSH("Icons/Profiler/Profiler_Border_R_16x", FMargin(4.0f / 16.0f)));
+
+		// Misc
+
+		Set("Profiler.Misc.SortBy", new IMAGE_BRUSH("Icons/Profiler/profiler_SortBy_32x", Icon32x32));
+		Set("Profiler.Misc.SortAscending", new IMAGE_BRUSH("Icons/Profiler/profiler_SortAscending_32x", Icon32x32));
+		Set("Profiler.Misc.SortDescending", new IMAGE_BRUSH("Icons/Profiler/profiler_SortDescending_32x", Icon32x32));
+
+		Set("Profiler.Misc.ResetToDefault", new IMAGE_BRUSH("Icons/Profiler/profiler_ResetToDefault_32x", Icon32x32));
+
+		Set("Profiler.Misc.Reset16", new IMAGE_BRUSH("Icons/Profiler/profiler_ResetToDefault_32x", Icon16x16));
+
+		Set("Profiler.Type.Calls", new IMAGE_BRUSH("Icons/Profiler/profiler_Calls_32x", Icon16x16));
+		Set("Profiler.Type.Event", new IMAGE_BRUSH("Icons/Profiler/profiler_Event_32x", Icon16x16));
+		Set("Profiler.Type.Memory", new IMAGE_BRUSH("Icons/Profiler/profiler_Memory_32x", Icon16x16));
+		Set("Profiler.Type.Number", new IMAGE_BRUSH("Icons/Profiler/profiler_Number_32x", Icon16x16));
+
+		// NumberInt, NumberFloat, Memory, Hierarchical
+		Set("Profiler.Type.NumberInt", new IMAGE_BRUSH("Icons/Profiler/profiler_Number_32x", Icon16x16));
+		Set("Profiler.Type.NumberFloat", new IMAGE_BRUSH("Icons/Profiler/profiler_Number_32x", Icon16x16));
+		Set("Profiler.Type.Memory", new IMAGE_BRUSH("Icons/Profiler/profiler_Memory_32x", Icon16x16));
+		Set("Profiler.Type.Hierarchical", new IMAGE_BRUSH("Icons/Profiler/profiler_Event_32x", Icon16x16));
+
+		Set("Profiler.Misc.GenericFilter", new IMAGE_BRUSH("Icons/Profiler/profiler_GenericFilter_32x", Icon16x16));
+		Set("Profiler.Misc.GenericGroup", new IMAGE_BRUSH("Icons/Profiler/profiler_GenericGroup_32x", Icon16x16));
+		Set("Profiler.Misc.CopyToClipboard", new IMAGE_BRUSH("Icons/Profiler/profiler_CopyToClipboard_32x", Icon32x32));
+
+		Set("Profiler.Misc.Disconnect", new IMAGE_BRUSH("Icons/Profiler/profiler_Disconnect_32x", Icon32x32));
+
+		Set("TaskGraph.Mono", new BOX_BRUSH("Common/TaskGraph_Mono", FMargin(5.f / 12.f)));
+	}
+
+
+	FSlateStyleRegistry::RegisterSlateStyle(*this);
+}
+
+FProfilerStyle::~FProfilerStyle()
+{
+	FSlateStyleRegistry::UnRegisterSlateStyle(*this);
+}
diff --git a/Engine/Source/Developer/Profiler/Private/ProfilerStyle.h b/Engine/Source/Developer/Profiler/Private/ProfilerStyle.h
new file mode 100644
--- /dev/null
+++ b/Engine/Source/Developer/Profiler/Private/ProfilerStyle.h
@@ -0,0 +1,26 @@
+// Copyright Epic Games, Inc. All Rights Reserved.
+
+#pragma once
+
+#include "Styling/SlateStyle.h"
+#include "Styling/SlateColor.h"
+#include "UObject/NameTypes.h"
+#include "Templates/UniquePtr.h"
+
+class FProfilerStyle
+	: public FSlateStyleSet
+{
+public:
+	virtual const FName& GetStyleSetName() const override;
+
+	static const FProfilerStyle& Get();
+	static void Shutdown();
+	
+	~FProfilerStyle();
+
+private:
+	FProfilerStyle();
+
+	static FName StyleName;
+	static TUniquePtr<FProfilerStyle> Inst;
+};
diff --git a/Engine/Source/Developer/Profiler/Private/Widgets/SDataGraph.cpp b/Engine/Source/Developer/Profiler/Private/Widgets/SDataGraph.cpp
--- a/Engine/Source/Developer/Profiler/Private/Widgets/SDataGraph.cpp
+++ b/Engine/Source/Developer/Profiler/Private/Widgets/SDataGraph.cpp
@@ -21,10 +21,10 @@
 #include "Widgets/Text/STextBlock.h"
 #include "Framework/MultiBox/MultiBoxBuilder.h"
 #include "Widgets/Input/SButton.h"
-#include "Styling/AppStyle.h"
 #include "ProfilerDataProvider.h"
 #include "ProfilerManager.h"
 #include "Widgets/StatDragDropOp.h"
+#include "../ProfilerStyle.h"
 
 #define LOCTEXT_NAMESPACE "SDataGraph"
 
@@ -104,7 +104,7 @@
 			[
 				SNew(STextBlock)
 				.ColorAndOpacity( TextColor )
-				.TextStyle( FAppStyle::Get(), TEXT("Profiler.Tooltip") )
+				.TextStyle( FProfilerStyle::Get(), TEXT("Profiler.Tooltip") )
 				.Text( this, &SDataGraphSummary::SummaryInformation_GetGroupName )
 			]
 
@@ -117,7 +117,7 @@
 			[
 				SNew(STextBlock)
 				.ColorAndOpacity( TextColor )
-				.TextStyle( FAppStyle::Get(), TEXT("Profiler.Tooltip") )
+				.TextStyle(FProfilerStyle::Get(), TEXT("Profiler.Tooltip") )
 				.Text( this, &SDataGraphSummary::SummaryInformation_GetStatName )
 			]
 
@@ -130,7 +130,7 @@
 			[
 				SNew(STextBlock)
 				.ColorAndOpacity( TextColor )
-				.TextStyle( FAppStyle::Get(), TEXT("Profiler.Tooltip") )
+				.TextStyle(FProfilerStyle::Get(), TEXT("Profiler.Tooltip") )
 				.Text( this, &SDataGraphSummary::SummaryInformation_GetSummary )
 			]
 		];
@@ -306,8 +306,8 @@
 	// Rendering info.
 	const bool bEnabled  = ShouldBeEnabled( bParentEnabled );
 	ESlateDrawEffect DrawEffects = bEnabled ? ESlateDrawEffect::None : ESlateDrawEffect::DisabledEffect;
-	const FSlateBrush* TimelineAreaBrush = FAppStyle::GetBrush("Profiler.LineGraphArea");
-	const FSlateBrush* WhiteBrush = FAppStyle::GetBrush("WhiteTexture");
+	const FSlateBrush* TimelineAreaBrush = FProfilerStyle::Get().GetBrush("Brushes.White25");
+	const FSlateBrush* WhiteBrush = FProfilerStyle::Get().GetBrush("Brushes.White");
 
 	/** Width of the alloted geometry that is used to draw a data graph. */
 	const float AreaX0 = 0.0f;
diff --git a/Engine/Source/Developer/Profiler/Private/Widgets/SEventGraph.cpp b/Engine/Source/Developer/Profiler/Private/Widgets/SEventGraph.cpp
--- a/Engine/Source/Developer/Profiler/Private/Widgets/SEventGraph.cpp
+++ b/Engine/Source/Developer/Profiler/Private/Widgets/SEventGraph.cpp
@@ -21,6 +21,7 @@
 #include "Widgets/SEventGraphTooltip.h"
 #include "Widgets/Input/SSearchBox.h"
 #include "HAL/PlatformApplicationMisc.h"
+#include "../ProfilerStyle.h"
 
 #define LOCTEXT_NAMESPACE "SEventGraph"
 
@@ -246,7 +247,7 @@
 			[
 				SNew( SImage )
 				.Visibility( this, &SEventGraphTableCell::GetHotPathIconVisibility )
-				.Image( FAppStyle::GetBrush(TEXT("Profiler.EventGraph.HotPathSmall")) )
+				.Image(FProfilerStyle::Get().GetBrush(TEXT("Profiler.EventGraph.HotPathSmall")) )
 			]
 
 			+SHorizontalBox::Slot()
@@ -257,7 +258,7 @@
 			[
 				SNew( STextBlock )
 				.Text( FText::FromName(EventPtr->_StatName) )
-				.TextStyle( FAppStyle::Get(), TEXT("Profiler.Tooltip") )
+				.TextStyle(FProfilerStyle::Get(), TEXT("Profiler.Tooltip") )
 				.ColorAndOpacity( this, &SEventGraphTableCell::GetColorAndOpacity )
 				.ShadowColorAndOpacity( this, &SEventGraphTableCell::GetShadowColorAndOpacity )
 			]
@@ -275,7 +276,7 @@
 				[
 					SNew( SImage )
 					.Visibility( this, &SEventGraphTableCell::GetCulledEventsIconVisibility )
-					.Image( FAppStyle::GetBrush("Profiler.EventGraph.HasCulledEventsSmall") )
+					.Image(FProfilerStyle::Get().GetBrush("Profiler.EventGraph.HasCulledEventsSmall") )
 					.ToolTipText( LOCTEXT("HasCulledEvents_TT","This event contains culled children, if you want to see all children, please disable culling or use function details, or press this icon") )
 				]	
 			]
@@ -287,7 +288,7 @@
 			[
 				SNew( SImage )
 				.Visibility( this, &SEventGraphTableCell::GetHintIconVisibility )
-				.Image( FAppStyle::GetBrush("Profiler.Tooltip.HintIcon10") )
+				.Image(FProfilerStyle::Get().GetBrush("Profiler.Tooltip.HintIcon10") )
 				.ToolTip( SEventGraphTooltip::GetTableCellTooltip( EventPtr ) )
 			];
 		}
@@ -307,7 +308,7 @@
 			[
 				SNew( STextBlock )
 				.Text( FText::FromString(FormattedValue) )
-				.TextStyle( FAppStyle::Get(), TEXT("Profiler.Tooltip") )
+				.TextStyle(FProfilerStyle::Get(), TEXT("Profiler.Tooltip") )
 				.ColorAndOpacity( this, &SEventGraphTableCell::GetColorAndOpacity )
 				.ShadowColorAndOpacity( this, &SEventGraphTableCell::GetShadowColorAndOpacity )
 			];
@@ -473,7 +474,7 @@
 		.Padding( 0.0f )
 		[
 			SNew( SImage )
-			.Image( FAppStyle::GetBrush("Profiler.LineGraphArea") )
+			.Image(FProfilerStyle::Get().GetBrush("Brushes.White25") )
 			.ColorAndOpacity( this, &SEventGraphTableRow::GetBackgroundColorAndOpacity )
 		]
 
@@ -551,15 +552,15 @@
 		const FSlateBrush* Brush = nullptr;
 		if( Result == HAlign_Left )
 		{
-			Brush = FAppStyle::GetBrush("Profiler.EventGraph.Border.L");
+			Brush = FProfilerStyle::Get().GetBrush("Profiler.EventGraph.Border.L");
 		}
 		else if( Result == HAlign_Right )
 		{
-			Brush = FAppStyle::GetBrush("Profiler.EventGraph.Border.R");
+			Brush = FProfilerStyle::Get().GetBrush("Profiler.EventGraph.Border.R");
 		}
 		else
 		{
-			Brush = FAppStyle::GetBrush("Profiler.EventGraph.Border.TB");
+			Brush = FProfilerStyle::Get().GetBrush("Profiler.EventGraph.Border.TB");
 		}
 		return Brush;
 	}
@@ -792,7 +793,7 @@
 		.VAlign( VAlign_Center )
 		[
 			SNew(SImage)
-			.Image( FAppStyle::GetBrush( BrushName ) )
+			.Image(FProfilerStyle::Get().GetBrush( BrushName ) )
 		];
 	}
 
@@ -802,12 +803,12 @@
 	[
 		SNew( STextBlock )
 		.Text( FText::FromString(EEventGraphTypes::ToName( EventGraphType )) )
-		.TextStyle( FAppStyle::Get(), TEXT("Profiler.Caption") )
+		.ColorAndOpacity(FSlateColor::UseForeground())
 	];
 
 	return
 	SNew( SCheckBox )
-	.Style(FAppStyle::Get(), "ToggleButtonCheckbox")
+	.Style(FProfilerStyle::Get(), "ToggleButtonCheckbox")
 	.IsEnabled( this, &SEventGraph::EventGraphType_IsEnabled, EventGraphType )
 	.HAlign( HAlign_Center )
 	.Padding( 2.0f )
@@ -823,7 +824,7 @@
 {
 	return
 	SNew( SCheckBox )
-	.Style(FAppStyle::Get(), "ToggleButtonCheckbox")
+	.Style(FProfilerStyle::Get(), "ToggleButtonCheckbox")
 	.IsEnabled( this, &SEventGraph::EventGraph_IsEnabled )
 	.HAlign( HAlign_Center )
 	.Padding( 2.0f )
@@ -840,7 +841,7 @@
 		.VAlign( VAlign_Center )
 		[
 			SNew(SImage)
-			.Image( FAppStyle::GetBrush( EEventGraphViewModes::ToBrushName( EventGraphViewMode ) ) )
+			.Image(FProfilerStyle::Get().GetBrush( EEventGraphViewModes::ToBrushName( EventGraphViewMode ) ) )
 		]
 
 		+SHorizontalBox::Slot()
@@ -849,7 +850,7 @@
 		[
 			SNew( STextBlock )
 			.Text( EEventGraphViewModes::ToName( EventGraphViewMode ) )
-			.TextStyle( FAppStyle::Get(), TEXT("Profiler.Caption") )
+			.ColorAndOpacity(FSlateColor::UseForeground())
 		]
 	];
 }
@@ -858,7 +859,7 @@
 {
 	return
 	SNew( SBorder )
-	.BorderImage( FAppStyle::GetBrush("Profiler.Group.16") )
+	.BorderImage(FProfilerStyle::Get().GetBrush("Profiler.Group.16") )
 	.Padding(FMargin(2.0f, 0.0))
 	[
 		SNew( SHorizontalBox )
@@ -872,7 +873,7 @@
 		[	
 			SNew(STextBlock)
 			.Text( LOCTEXT("Toolbar_Type","Type") )
-			.TextStyle( FAppStyle::Get(), TEXT("Profiler.CaptionBold") )
+			.TextStyle(FProfilerStyle::Get(), TEXT("Profiler.CaptionBold") )
 		]
 
 		// One-frame
@@ -912,7 +913,7 @@
 {
 	return
 	SNew( SBorder )
-	.BorderImage( FAppStyle::GetBrush("Profiler.Group.16") )
+	.BorderImage(FProfilerStyle::Get().GetBrush("Profiler.Group.16") )
 	.Padding(FMargin(2.0f, 0.0))
 	[
 		SNew( SHorizontalBox )
@@ -926,7 +927,7 @@
 		[	
 			SNew(STextBlock)
 			.Text( LOCTEXT("Toolbar_ViewMode","View mode") )
-			.TextStyle( FAppStyle::Get(), TEXT("Profiler.CaptionBold") )
+			.TextStyle(FProfilerStyle::Get(), TEXT("Profiler.CaptionBold") )
 		]
 
 		// View mode - Hierarchical
@@ -984,7 +985,7 @@
 TSharedRef<SWidget> SEventGraph::GetWidgetForThreadFilter()
 {
 	return SNew(SBorder)
-		.BorderImage(FAppStyle::GetBrush("Profiler.Group.16"))
+		.BorderImage(FProfilerStyle::Get().GetBrush("Profiler.Group.16"))
 		.Padding(FMargin(2.0f, 0.0))
 		[
 			SNew(SHorizontalBox)
@@ -997,7 +998,7 @@
 			[
 				SNew ( STextBlock )
 				.Text( LOCTEXT( "Toolbar_Thread", "Thread" ) )
-				.TextStyle( FAppStyle::Get(), TEXT( "Profiler.CaptionBold" ) )
+				.TextStyle(FProfilerStyle::Get(), TEXT( "Profiler.CaptionBold" ) )
 			]
 
 			+SHorizontalBox::Slot()
@@ -1093,7 +1094,7 @@
 {
 	return
 	SNew( SBorder )
-	.BorderImage( FAppStyle::GetBrush("Profiler.Group.16") )
+	.BorderImage(FProfilerStyle::Get().GetBrush("Profiler.Group.16") )
 	.Padding( 0.0f )
 	[
 		SNew(SHorizontalBox)
@@ -1112,7 +1113,7 @@
 			.ContentPadding( 2.0f )
 			[
 				SNew(SImage)
-				.Image( FAppStyle::GetBrush("Profiler.EventGraph.HistoryBack") )
+				.Image(FProfilerStyle::Get().GetBrush("Profiler.EventGraph.HistoryBack") )
 			]
 		]
 
@@ -1130,7 +1131,7 @@
 			.ContentPadding( 2.0f )
 			[
 				SNew(SImage)
-				.Image( FAppStyle::GetBrush("Profiler.EventGraph.HistoryForward") )
+				.Image(FProfilerStyle::Get().GetBrush("Profiler.EventGraph.HistoryForward") )
 			]
 		]
 
@@ -1158,7 +1159,7 @@
 			.OnClicked( this, &SEventGraph::ExpandHotPath_OnClicked )
 			[
 				SNew( SImage )
-				.Image( FAppStyle::GetBrush(TEXT("Profiler.EventGraph.ExpandHotPath16")) )
+				.Image(FProfilerStyle::Get().GetBrush(TEXT("Profiler.EventGraph.ExpandHotPath16")) )
 			]
 		]
 
@@ -1273,7 +1274,7 @@
 	[
 		SNew(STextBlock)
 		.Text( Caption )
-		.TextStyle( FAppStyle::Get(), TEXT("Profiler.CaptionBold") )
+		.TextStyle(FProfilerStyle::Get(), TEXT("Profiler.CaptionBold") )
 	]
 
 	+SVerticalBox::Slot()
@@ -1309,7 +1310,7 @@
 	[
 		SNew(STextBlock)
 		.Text( LOCTEXT("FunctionDetails_CurrentFunction","Current Function") )
-		.TextStyle( FAppStyle::Get(), TEXT("Profiler.CaptionBold") )
+		.TextStyle(FProfilerStyle::Get(), TEXT("Profiler.CaptionBold") )
 	]
 
 	+SVerticalBox::Slot()
@@ -1597,7 +1598,7 @@
 		( 
 			SelectionStr, 
 			LOCTEXT("ContextMenu_Selection", "Currently selected events"), 
-			FSlateIcon(FAppStyle::GetAppStyleSetName(), "@missing.icon"), DummyUIAction, NAME_None, EUserInterfaceActionType::Button
+			FSlateIcon(FProfilerStyle::Get().GetStyleSetName(), "@missing.icon"), DummyUIAction, NAME_None, EUserInterfaceActionType::Button
 		);
 	}
 	MenuBuilder.EndSection();
@@ -1609,7 +1610,7 @@
 		( 
 			LOCTEXT("ContextMenu_Root_Set", "Set Root"), 
 			LOCTEXT("ContextMenu_Root_Set_Desc", "Sets the root to the selected event(s) and switches to hierarchical view"), 
-			FSlateIcon(FAppStyle::GetAppStyleSetName(), "Profiler.EventGraph.SetRoot"), SetRoot_Custom(), NAME_None, EUserInterfaceActionType::Button 
+			FSlateIcon(FProfilerStyle::Get().GetStyleSetName(), "Profiler.EventGraph.SetRoot"), SetRoot_Custom(), NAME_None, EUserInterfaceActionType::Button
 		);
 
 		FUIAction Action_AggregateForSelection;
@@ -1645,7 +1646,7 @@
 		( 
 			CullingDesc,
 			LOCTEXT("ContextMenu_Culling_Desc_TT","Culls the event graph based on the property value of the selected event"),
-			FSlateIcon(FAppStyle::GetAppStyleSetName(), "Profiler.EventGraph.CullEvents"), CullByProperty_Custom(SelectedEvent,HoveredColumnID,false), NAME_None, EUserInterfaceActionType::Button 
+			FSlateIcon(FProfilerStyle::Get().GetStyleSetName(), "Profiler.EventGraph.CullEvents"), CullByProperty_Custom(SelectedEvent,HoveredColumnID,false), NAME_None, EUserInterfaceActionType::Button
 		);
 
 	// Filtering menu
@@ -1668,7 +1669,7 @@
 		( 
 			FilteringDesc,
 			LOCTEXT("ContextMenu_Filtering_Desc_TT","Filters the event graph based on the property value of the selected event"),
-			FSlateIcon(FAppStyle::GetAppStyleSetName(), "Profiler.EventGraph.FilterEvents"), FilterOutByProperty_Custom(SelectedEvent,HoveredColumnID,false), NAME_None, EUserInterfaceActionType::Button 
+			FSlateIcon(FProfilerStyle::Get().GetStyleSetName(), "Profiler.EventGraph.FilterEvents"), FilterOutByProperty_Custom(SelectedEvent,HoveredColumnID,false), NAME_None, EUserInterfaceActionType::Button 
 		);
 
 		MenuBuilder.AddMenuSeparator();
@@ -1676,7 +1677,7 @@
 		( 
 			LOCTEXT("ContextMenu_ClearHistory", "Reset to default"), 
 			LOCTEXT("ContextMenu_ClearHistory_Reset_Desc", "For the selected event graph resets root/culling/filter to the default state and clears the history"), 
-			FSlateIcon(FAppStyle::GetAppStyleSetName(), "Profiler.Misc.ResetToDefault"), ClearHistory_Custom(), NAME_None, EUserInterfaceActionType::Button 
+			FSlateIcon(FProfilerStyle::Get().GetStyleSetName(), "Profiler.Misc.ResetToDefault"), ClearHistory_Custom(), NAME_None, EUserInterfaceActionType::Button 
 		);
 	}
 	MenuBuilder.EndSection();
@@ -1687,42 +1688,42 @@
 		( 
 			LOCTEXT("ContextMenu_Header_Expand_ExpandAll", "Expand All"), 
 			LOCTEXT("ContextMenu_Header_Expand_ExpandAll_Desc", "Expands all events"), 
-			FSlateIcon(FAppStyle::GetAppStyleSetName(), "Profiler.EventGraph.ExpandAll"), SetExpansionForEvents_Custom( ESelectedEventTypes::AllEvents, true ), NAME_None, EUserInterfaceActionType::Button 
+			FSlateIcon(FProfilerStyle::Get().GetStyleSetName(), "Profiler.EventGraph.ExpandAll"), SetExpansionForEvents_Custom( ESelectedEventTypes::AllEvents, true ), NAME_None, EUserInterfaceActionType::Button 
 		);
 
 		MenuBuilder.AddMenuEntry
 		( 
 			LOCTEXT("ContextMenu_Header_Expand_CollapseAll", "Collapse All"), 
 			LOCTEXT("ContextMenu_Header_Expand_CollapseAll_Desc", "Collapses all events"), 
-			FSlateIcon(FAppStyle::GetAppStyleSetName(), "Profiler.EventGraph.CollapseAll"), SetExpansionForEvents_Custom( ESelectedEventTypes::AllEvents, false ), NAME_None, EUserInterfaceActionType::Button 
+			FSlateIcon(FProfilerStyle::Get().GetStyleSetName(), "Profiler.EventGraph.CollapseAll"), SetExpansionForEvents_Custom( ESelectedEventTypes::AllEvents, false ), NAME_None, EUserInterfaceActionType::Button 
 		);
 
 		MenuBuilder.AddMenuEntry
 		( 
 			LOCTEXT("ContextMenu_Header_Expand_ExpandSelection", "Expand Selection"), 
 			LOCTEXT("ContextMenu_Header_Expand_ExpandSelection_Desc", "Expands selected events"), 
-			FSlateIcon(FAppStyle::GetAppStyleSetName(), "Profiler.EventGraph.ExpandSelection"), SetExpansionForEvents_Custom( ESelectedEventTypes::SelectedEvents, true ), NAME_None, EUserInterfaceActionType::Button 
+			FSlateIcon(FProfilerStyle::Get().GetStyleSetName(), "Profiler.EventGraph.ExpandSelection"), SetExpansionForEvents_Custom( ESelectedEventTypes::SelectedEvents, true ), NAME_None, EUserInterfaceActionType::Button 
 		);
 
 		MenuBuilder.AddMenuEntry
 		( 
 			LOCTEXT("ContextMenu_Header_Expand_CollapseSelection", "Collapse Selection"), 
 			LOCTEXT("ContextMenu_Header_Expand_CollapseSelection_Desc", "Collapses selected events"), 
-			FSlateIcon(FAppStyle::GetAppStyleSetName(), "Profiler.EventGraph.CollapseSelection"), SetExpansionForEvents_Custom( ESelectedEventTypes::SelectedEvents, false ), NAME_None, EUserInterfaceActionType::Button 
+			FSlateIcon(FProfilerStyle::Get().GetStyleSetName(), "Profiler.EventGraph.CollapseSelection"), SetExpansionForEvents_Custom( ESelectedEventTypes::SelectedEvents, false ), NAME_None, EUserInterfaceActionType::Button 
 		);
 
 		MenuBuilder.AddMenuEntry
 		( 
 			LOCTEXT("ContextMenu_Header_Expand_ExpandThread", "Expand Thread"), 
 			LOCTEXT("ContextMenu_Header_Expand_ExpandThread_Desc", "Expands selected threads"), 
-			FSlateIcon(FAppStyle::GetAppStyleSetName(), "Profiler.EventGraph.ExpandThread"), SetExpansionForEvents_Custom( ESelectedEventTypes::SelectedThreadEvents, true ), NAME_None, EUserInterfaceActionType::Button 
+			FSlateIcon(FProfilerStyle::Get().GetStyleSetName(), "Profiler.EventGraph.ExpandThread"), SetExpansionForEvents_Custom( ESelectedEventTypes::SelectedThreadEvents, true ), NAME_None, EUserInterfaceActionType::Button 
 		);
 
 		MenuBuilder.AddMenuEntry
 		( 
 			LOCTEXT("ContextMenu_Header_Expand_CollapseThread", "Collapse Thread"), 
 			LOCTEXT("ContextMenu_Header_Expand_CollapseThread_Desc", "Collapses selected threads"), 
-			FSlateIcon(FAppStyle::GetAppStyleSetName(), "Profiler.EventGraph.CollapseThread"), SetExpansionForEvents_Custom( ESelectedEventTypes::SelectedThreadEvents, false ), NAME_None, EUserInterfaceActionType::Button 
+			FSlateIcon(FProfilerStyle::Get().GetStyleSetName(), "Profiler.EventGraph.CollapseThread"), SetExpansionForEvents_Custom( ESelectedEventTypes::SelectedThreadEvents, false ), NAME_None, EUserInterfaceActionType::Button 
 		);
 
 		//-----------------------------------------------------------------------------
@@ -1736,7 +1737,7 @@
 		( 
 			LOCTEXT("ContextMenu_Header_Expand_ExpandHotPath", "Expand Hot Path"), 
 			LOCTEXT("ContextMenu_Header_Expand_ExpandHotPath_Desc", "Expands hot path for the selected events, based on the inclusive time, also enables descending sorting by inclusive time"), 
-			FSlateIcon(FAppStyle::GetAppStyleSetName(), "Profiler.EventGraph.ExpandHotPath"), Action_ExpandHotPath, NAME_None, EUserInterfaceActionType::Button 
+			FSlateIcon(FProfilerStyle::Get().GetStyleSetName(), "Profiler.EventGraph.ExpandHotPath"), Action_ExpandHotPath, NAME_None, EUserInterfaceActionType::Button 
 		);
 	}
 	MenuBuilder.EndSection();
@@ -1747,7 +1748,7 @@
 		( 
 			LOCTEXT("ContextMenu_Header_Navigation_ShowInHierarchicalView", "Show In Hierarchical View"), 
 			LOCTEXT("ContextMenu_Header_Navigation_ShowInHierarchicalView_Desc", "Switches to hierarchical view and expands selected events"), 
-			FSlateIcon(FAppStyle::GetAppStyleSetName(), EEventGraphViewModes::ToBrushName(EEventGraphViewModes::Hierarchical)), ShowSelectedEventsInViewMode_Custom(EEventGraphViewModes::Hierarchical), NAME_None, EUserInterfaceActionType::Check
+			FSlateIcon(FProfilerStyle::Get().GetStyleSetName(), EEventGraphViewModes::ToBrushName(EEventGraphViewModes::Hierarchical)), ShowSelectedEventsInViewMode_Custom(EEventGraphViewModes::Hierarchical), NAME_None, EUserInterfaceActionType::Check
 		);
 
 		//-----------------------------------------------------------------------------
@@ -1756,7 +1757,7 @@
 		( 
 			LOCTEXT("ContextMenu_Header_Navigation_ShowInFlatView", "Show In FlatInclusive View"), 
 			LOCTEXT("ContextMenu_Header_Navigation_ShowInFlatView_Desc", "Switches to flat view, also enables descending sorting by inclusive time"), 
-			FSlateIcon(FAppStyle::GetAppStyleSetName(), EEventGraphViewModes::ToBrushName(EEventGraphViewModes::FlatExclusive)), ShowSelectedEventsInViewMode_Custom(EEventGraphViewModes::FlatExclusive), NAME_None, EUserInterfaceActionType::Check 
+			FSlateIcon(FProfilerStyle::Get().GetStyleSetName(), EEventGraphViewModes::ToBrushName(EEventGraphViewModes::FlatExclusive)), ShowSelectedEventsInViewMode_Custom(EEventGraphViewModes::FlatExclusive), NAME_None, EUserInterfaceActionType::Check 
 		);
 
 		if( FProfilerManager::GetSettings().bShowCoalescedViewModesInEventGraph )
@@ -1765,7 +1766,7 @@
 			( 
 				LOCTEXT("ContextMenu_Header_Navigation_ShowInFlatCoalesced", "Show In FlatInclusive Coalesced"), 
 				LOCTEXT("ContextMenu_Header_Navigation_ShowInFlatCoalesced_Desc", "Switches to flat coalesced, also enables descending sorting by inclusive time"), 
-				FSlateIcon(FAppStyle::GetAppStyleSetName(), EEventGraphViewModes::ToBrushName(EEventGraphViewModes::FlatInclusiveCoalesced)), ShowSelectedEventsInViewMode_Custom(EEventGraphViewModes::FlatInclusiveCoalesced), NAME_None, EUserInterfaceActionType::Check 
+				FSlateIcon(FProfilerStyle::Get().GetStyleSetName(), EEventGraphViewModes::ToBrushName(EEventGraphViewModes::FlatInclusiveCoalesced)), ShowSelectedEventsInViewMode_Custom(EEventGraphViewModes::FlatInclusiveCoalesced), NAME_None, EUserInterfaceActionType::Check 
 			);
 		}
 
@@ -1775,7 +1776,7 @@
 		( 
 			LOCTEXT("ContextMenu_Header_Navigation_ShowInFlatExclusiveView", "Show In Flat Exclusive View"), 
 			LOCTEXT("ContextMenu_Header_Navigation_ShowInFlatExclusiveView_Desc", "Switches to flat exclusive view, also enables ascending sorting by exclusive time"), 
-			FSlateIcon(FAppStyle::GetAppStyleSetName(), EEventGraphViewModes::ToBrushName(EEventGraphViewModes::FlatExclusive)), ShowSelectedEventsInViewMode_Custom(EEventGraphViewModes::FlatExclusive), NAME_None, EUserInterfaceActionType::Check 
+			FSlateIcon(FProfilerStyle::Get().GetStyleSetName(), EEventGraphViewModes::ToBrushName(EEventGraphViewModes::FlatExclusive)), ShowSelectedEventsInViewMode_Custom(EEventGraphViewModes::FlatExclusive), NAME_None, EUserInterfaceActionType::Check 
 		);
 
 		if( FProfilerManager::GetSettings().bShowCoalescedViewModesInEventGraph )
@@ -1784,7 +1785,7 @@
 			( 
 				LOCTEXT("ContextMenu_Header_Navigation_ShowInFlatExclusiveCoalescedView", "Show In Flat Exclusive Coalesced View"), 
 				LOCTEXT("ContextMenu_Header_Navigation_ShowInFlatExclusiveCoalescedView_Desc", "Switches to flat exclusive coalesced view, also enables ascending sorting by exclusive time enabled"), 
-				FSlateIcon(FAppStyle::GetAppStyleSetName(), EEventGraphViewModes::ToBrushName(EEventGraphViewModes::FlatExclusiveCoalesced)), ShowSelectedEventsInViewMode_Custom(EEventGraphViewModes::FlatExclusiveCoalesced), NAME_None, EUserInterfaceActionType::Check 
+				FSlateIcon(FProfilerStyle::Get().GetStyleSetName(), EEventGraphViewModes::ToBrushName(EEventGraphViewModes::FlatExclusiveCoalesced)), ShowSelectedEventsInViewMode_Custom(EEventGraphViewModes::FlatExclusiveCoalesced), NAME_None, EUserInterfaceActionType::Check 
 			);
 		}
 
@@ -1813,7 +1814,7 @@
 		( 
 			LOCTEXT("ContextMenu_Header_Misc_CopySelectedToClipboard", "Copy To Clipboard"), 
 			LOCTEXT("ContextMenu_Header_Misc_CopySelectedToClipboard_Desc", "Copies selection to clipboard"), 
-			FSlateIcon(FAppStyle::GetAppStyleSetName(), "Profiler.Misc.CopyToClipboard"), Action_CopySelectedToClipboard, NAME_None, EUserInterfaceActionType::Button 
+			FSlateIcon(FProfilerStyle::Get().GetStyleSetName(), "Profiler.Misc.CopyToClipboard"), Action_CopySelectedToClipboard, NAME_None, EUserInterfaceActionType::Button 
 		);
 
 		FUIAction Action_SaveSelectedToFile;
@@ -1827,7 +1828,7 @@
 		( 
 			LOCTEXT("ContextMenu_Header_Misc_SelectStack", "Select Stack"), 
 			LOCTEXT("ContextMenu_Header_Misc_SelectStack_Desc", "Selects all events in the stack"), 
-			FSlateIcon(FAppStyle::GetAppStyleSetName(), "Profiler.EventGraph.SelectStack"), Action_SelectStack, NAME_None, EUserInterfaceActionType::Button 
+			FSlateIcon(FProfilerStyle::Get().GetStyleSetName(), "Profiler.EventGraph.SelectStack"), Action_SelectStack, NAME_None, EUserInterfaceActionType::Button 
 		);
 
 		MenuBuilder.AddSubMenu
@@ -1836,7 +1837,7 @@
 			LOCTEXT("ContextMenu_Header_Misc_Sort_Desc", "Sort by column"), 
 			FNewMenuDelegate::CreateSP( const_cast<SEventGraph*>(this), &SEventGraph::EventGraph_BuildSortByMenu ),
 			false, 
-			FSlateIcon(FAppStyle::GetAppStyleSetName(), "Profiler.Misc.SortBy")
+			FSlateIcon(FProfilerStyle::Get().GetStyleSetName(), "Profiler.Misc.SortBy")
 		);
 	}
 	MenuBuilder.EndSection();
@@ -1849,7 +1850,7 @@
 			LOCTEXT("ContextMenu_Header_Columns_View_Desc", "Hides or shows columns"), 
 			FNewMenuDelegate::CreateSP( const_cast<SEventGraph*>(this), &SEventGraph::EventGraph_BuildViewColumnMenu ),
 			false,
-			FSlateIcon(FAppStyle::GetAppStyleSetName(), "Profiler.EventGraph.ViewColumn")
+			FSlateIcon(FProfilerStyle::Get().GetStyleSetName(), "Profiler.EventGraph.ViewColumn")
 		);
 
 		FUIAction Action_ResetColumns
@@ -1861,7 +1862,7 @@
 		( 
 			LOCTEXT("ContextMenu_Header_Columns_ResetColumns", "Reset Columns To Default"), 
 			LOCTEXT("ContextMenu_Header_Columns_ResetColumns_Desc", "Resets columns to default"), 
-			FSlateIcon(FAppStyle::GetAppStyleSetName(), "Profiler.EventGraph.ResetColumn"), Action_ResetColumns, NAME_None, EUserInterfaceActionType::Button 
+			FSlateIcon(FProfilerStyle::Get().GetStyleSetName(), "Profiler.EventGraph.ResetColumn"), Action_ResetColumns, NAME_None, EUserInterfaceActionType::Button 
 		);
 	}
 	MenuBuilder.EndSection();
@@ -1909,7 +1910,7 @@
 		( 
 			LOCTEXT("ContextMenu_Header_Misc_Sort_SortAscending", "Sort Ascending"), 
 			LOCTEXT("ContextMenu_Header_Misc_Sort_SortAscending_Desc", "Sorts ascending"), 
-			FSlateIcon(FAppStyle::GetAppStyleSetName(), "Profiler.Misc.SortAscending"), Action_SortAscending, NAME_None, EUserInterfaceActionType::RadioButton
+			FSlateIcon(FProfilerStyle::Get().GetStyleSetName(), "Profiler.Misc.SortAscending"), Action_SortAscending, NAME_None, EUserInterfaceActionType::RadioButton
 		);
 
 		FUIAction Action_SortDescending
@@ -1922,7 +1923,7 @@
 		( 
 			LOCTEXT("ContextMenu_Header_Misc_Sort_SortDescending", "Sort Descending"), 
 			LOCTEXT("ContextMenu_Header_Misc_Sort_SortDescending_Desc", "Sorts descending"), 
-			FSlateIcon(FAppStyle::GetAppStyleSetName(), "Profiler.Misc.SortDescending"), Action_SortDescending, NAME_None, EUserInterfaceActionType::RadioButton 
+			FSlateIcon(FProfilerStyle::Get().GetStyleSetName(), "Profiler.Misc.SortDescending"), Action_SortDescending, NAME_None, EUserInterfaceActionType::RadioButton 
 		);
 	}
 	MenuBuilder.EndSection();
@@ -2200,7 +2201,7 @@
 			( 
 				LOCTEXT("ContextMenu_Header_Misc_Sort_SortAscending", "Sort Ascending"), 
 				LOCTEXT("ContextMenu_Header_Misc_Sort_SortAscending_Desc", "Sorts ascending"), 
-				FSlateIcon(FAppStyle::GetAppStyleSetName(), "Profiler.Misc.SortAscending"), Action_SortAscending, NAME_None, EUserInterfaceActionType::RadioButton
+				FSlateIcon(FProfilerStyle::Get().GetStyleSetName(), "Profiler.Misc.SortAscending"), Action_SortAscending, NAME_None, EUserInterfaceActionType::RadioButton
 			);
 
 			FUIAction Action_SortDescending
@@ -2213,7 +2214,7 @@
 			( 
 				LOCTEXT("ContextMenu_Header_Misc_Sort_SortDescending", "Sort Descending"), 
 				LOCTEXT("ContextMenu_Header_Misc_Sort_SortDescending_Desc", "Sorts descending"), 
-				FSlateIcon(FAppStyle::GetAppStyleSetName(), "Profiler.Misc.SortDescending"), Action_SortDescending, NAME_None, EUserInterfaceActionType::RadioButton 
+				FSlateIcon(FProfilerStyle::Get().GetStyleSetName(), "Profiler.Misc.SortDescending"), Action_SortDescending, NAME_None, EUserInterfaceActionType::RadioButton 
 			);
 			bIsMenuVisible = true;
 
@@ -2713,7 +2714,7 @@
 		SNew(STextBlock)
 		.WrapTextAt( 128.0f )
 		.Text( LOCTEXT("FunctionDetails_SelectOneEvent","Function details view works only if you select one event. Please select an individual event to proceed.") )
-		.TextStyle( FAppStyle::Get(), TEXT("Profiler.Tooltip") )
+		.TextStyle(FProfilerStyle::Get(), TEXT("Profiler.Tooltip") )
 	];
 
 	VerticalBox_TopCalling->ClearChildren();
@@ -2861,7 +2862,7 @@
 			SNew(SButton)
 			.HAlign( HAlign_Left )
 			.VAlign( VAlign_Center )
-			.TextStyle( FAppStyle::Get(), "Profiler.Tooltip" )
+			.TextStyle(FProfilerStyle::Get(), "Profiler.Tooltip" )
 			.ContentPadding( FMargin(4.0f, 1.0f) )
 			.OnClicked( this, &SEventGraph::CallingCalledFunctionButton_OnClicked, EventPtrAndPct.EventPtr )
 			[
@@ -2903,7 +2904,7 @@
 	[
 		SNew(STextBlock)
 		.Text( FText::FromString(GetEventDescription(EventPtr,Pct,bSimple)) )
-		.TextStyle( FAppStyle::Get(), bSimple ? TEXT("Profiler.Tooltip") : TEXT("Profiler.EventGraph.DarkText") )
+		.TextStyle(FProfilerStyle::Get(), bSimple ? TEXT("Profiler.Tooltip") : TEXT("Profiler.EventGraph.DarkText") )
 	];
 
 	if( EventPtr->_bIsCulled )
@@ -2914,7 +2915,7 @@
 		.VAlign(VAlign_Center)
 		[
 			SNew( SImage )
-			.Image( FAppStyle::GetBrush("Profiler.EventGraph.CulledEvent") )
+			.Image(FProfilerStyle::Get().GetBrush("Profiler.EventGraph.CulledEvent") )
 			.ToolTipText( LOCTEXT("Misc_EventCulled","Event is culled") )
 		];
 	}
@@ -2927,7 +2928,7 @@
 		.VAlign(VAlign_Center)
 		[
 			SNew( SImage )
-			.Image( FAppStyle::GetBrush("Profiler.EventGraph.FilteredEvent") )
+			.Image(FProfilerStyle::Get().GetBrush("Profiler.EventGraph.FilteredEvent") )
 			.ToolTipText( LOCTEXT("Misc_EventFiltered","Event is filtered") )
 		];
 	}
@@ -2938,7 +2939,7 @@
 	.VAlign(VAlign_Center)
 	[
 		SNew( SImage )
-		.Image( FAppStyle::GetBrush("Profiler.Tooltip.HintIcon10") )
+		.Image(FProfilerStyle::Get().GetBrush("Profiler.Tooltip.HintIcon10") )
 		.ToolTip( SEventGraphTooltip::GetTableCellTooltip( EventPtr ) )
 	];
 
diff --git a/Engine/Source/Developer/Profiler/Private/Widgets/SEventGraphTooltip.cpp b/Engine/Source/Developer/Profiler/Private/Widgets/SEventGraphTooltip.cpp
--- a/Engine/Source/Developer/Profiler/Private/Widgets/SEventGraphTooltip.cpp
+++ b/Engine/Source/Developer/Profiler/Private/Widgets/SEventGraphTooltip.cpp
@@ -16,7 +16,7 @@
 #include "Widgets/Text/STextBlock.h"
 #include "Widgets/Layout/SGridPanel.h"
 #include "Widgets/SToolTip.h"
-#include "Styling/AppStyle.h"
+#include "../ProfilerStyle.h"
 
 #define LOCTEXT_NAMESPACE "SEventGraphTooltip"
 
@@ -80,7 +80,7 @@
 						.AutoWidth()
 						[
 							SNew( SImage )
-							.Image( FAppStyle::GetBrush(TEXT("Profiler.EventGraph.HotPathSmall")) )
+							.Image(FProfilerStyle::Get().GetBrush(TEXT("Profiler.EventGraph.HotPathSmall")) )
 						]
 
 						+SHorizontalBox::Slot()
@@ -90,7 +90,6 @@
 							[
 								SNew(STextBlock)
 								.Text( LOCTEXT("HotPathInformation", "Hot path, should be investigated for bottlenecks") )
-								//.TextStyle( FAppStyle::Get(), TEXT("Profiler.TooltipBold") )
 								.ColorAndOpacity( ColorAndOpacity )
 								.Font( FCoreStyle::GetDefaultFontStyle("Regular", 8) )
 								.ShadowOffset( FVector2D(1.0f, 1.0f) )
@@ -121,14 +120,14 @@
 						[
 							SNew(STextBlock)
 							.Text(LOCTEXT("TT_Thread", "Thread:"))
-							.TextStyle( FAppStyle::Get(), TEXT("Profiler.TooltipBold") )
+							.TextStyle( FProfilerStyle::Get(), TEXT("Profiler.TooltipBold") )
 						]
 						+SGridPanel::Slot(1,0)
 						.Padding( 2.0f )
 						[
 							SNew(STextBlock)
 							.Text( FText::FromName(EventSample->_ThreadName ))
-							.TextStyle( FAppStyle::Get(), TEXT("Profiler.Tooltip") )
+							.TextStyle( FProfilerStyle::Get(), TEXT("Profiler.Tooltip") )
 						]
 
 						// Event: [_StatName]
@@ -137,14 +136,14 @@
 						[
 							SNew(STextBlock)
 							.Text(LOCTEXT("TT_Event", "Event:"))
-							.TextStyle( FAppStyle::Get(), TEXT("Profiler.TooltipBold") )
+							.TextStyle( FProfilerStyle::Get(), TEXT("Profiler.TooltipBold") )
 						]
 						+SGridPanel::Slot(3,0)
 						.Padding( 2.0f )
 						[
 							SNew(STextBlock)
 							.Text( FText::FromString(EventSample->GetShortEventName()) )
-							.TextStyle( FAppStyle::Get(), TEXT("Profiler.Tooltip") )
+							.TextStyle( FProfilerStyle::Get(), TEXT("Profiler.Tooltip") )
 						]
 
 						// Group: [_GroupName]
@@ -153,14 +152,14 @@
 						[
 							SNew(STextBlock)
 							.Text(LOCTEXT("TT_Group", "Group:"))
-							.TextStyle( FAppStyle::Get(), TEXT("Profiler.TooltipBold") )
+							.TextStyle( FProfilerStyle::Get(), TEXT("Profiler.TooltipBold") )
 						]
 						+SGridPanel::Slot(5,0)
 						.Padding( 2.0f )
 						[
 							SNew(STextBlock)
 							.Text( FText::FromName(EventSample->_GroupName) )
-							.TextStyle( FAppStyle::Get(), TEXT("Profiler.Tooltip") )
+							.TextStyle( FProfilerStyle::Get(), TEXT("Profiler.Tooltip") )
 						]
 
 						//-----------------------------------------------------------------------------
@@ -171,7 +170,7 @@
 						[
 							SNew(STextBlock)
 							.Text(LOCTEXT("TT_InclusiveTime", "Inclusive time:"))
-							.TextStyle( FAppStyle::Get(), TEXT("Profiler.TooltipBold") )
+							.TextStyle( FProfilerStyle::Get(), TEXT("Profiler.TooltipBold") )
 						]
 						// [InclusiveTimeMS] ms
 						+SGridPanel::Slot(1,1)
@@ -179,7 +178,7 @@
 						[
 							SNew(STextBlock)
 							.Text( FText::FromString(EventSample->GetFormattedValue(EEventPropertyIndex::InclusiveTimeMS)) )
-							.TextStyle( FAppStyle::Get(), TEXT("Profiler.Tooltip") )
+							.TextStyle( FProfilerStyle::Get(), TEXT("Profiler.Tooltip") )
 						]
 						// ([_InclusiveTimePct] % of the caller)
 						+SGridPanel::Slot(2,1)
@@ -187,7 +186,7 @@
 						[
 							SNew(STextBlock)
 							.Text( InclusiveTimePctCaller )
-							.TextStyle( FAppStyle::Get(), TEXT("Profiler.Tooltip") )
+							.TextStyle( FProfilerStyle::Get(), TEXT("Profiler.Tooltip") )
 						]
 						// Average inclusive time per call of all instances for this event, in milliseconds
 						+SGridPanel::Slot(3,1)
@@ -195,7 +194,7 @@
 						[
 							SNew(STextBlock)
 							.Text(LOCTEXT("TT_AvgIncTimePerCall", "Avg inc time per call:"))
-							.TextStyle( FAppStyle::Get(), TEXT("Profiler.TooltipBold") )
+							.TextStyle( FProfilerStyle::Get(), TEXT("Profiler.TooltipBold") )
 						]
 						
 						//-----------------------------------------------------------------------------
@@ -207,7 +206,7 @@
 						[
 							SNew(STextBlock)
 							.Text(LOCTEXT("TT_PctOfThread", "% of thread:"))
-							.TextStyle( FAppStyle::Get(), TEXT("Profiler.TooltipBold") )
+							.TextStyle( FProfilerStyle::Get(), TEXT("Profiler.TooltipBold") )
 						]
 						// [ThreadPct] %
 						+SGridPanel::Slot(1,2)
@@ -215,7 +214,7 @@
 						[
 							SNew(STextBlock)
 							.Text( FText::FromString(EventSample->GetFormattedValue(EEventPropertyIndex::ThreadPct)) )
-							.TextStyle( FAppStyle::Get(), TEXT("Profiler.Tooltip") )
+							.TextStyle( FProfilerStyle::Get(), TEXT("Profiler.Tooltip") )
 						]
 
 						+SGridPanel::Slot(3,2)
@@ -223,7 +222,7 @@
 						[
 							SNew(STextBlock)
 							.Text(LOCTEXT("TT_PctOfFrame", "% of frame:"))
-							.TextStyle( FAppStyle::Get(), TEXT("Profiler.TooltipBold") )
+							.TextStyle( FProfilerStyle::Get(), TEXT("Profiler.TooltipBold") )
 						]
 						// [FramePct] %
 						+SGridPanel::Slot(4,2)
@@ -231,7 +230,7 @@
 						[
 							SNew(STextBlock)
 							.Text( FText::FromString(EventSample->GetFormattedValue(EEventPropertyIndex::FramePct)) )
-							.TextStyle( FAppStyle::Get(), TEXT("Profiler.Tooltip") )
+							.TextStyle( FProfilerStyle::Get(), TEXT("Profiler.Tooltip") )
 						]
 
 						//-----------------------------------------------------------------------------
@@ -243,7 +242,7 @@
 						[
 							SNew(STextBlock)
 							.Text(LOCTEXT("TT_MinMaxAvgIncTime", "Min/Max/Avg inclusive time:"))
-							.TextStyle( FAppStyle::Get(), TEXT("Profiler.TooltipBold") )
+							.TextStyle( FProfilerStyle::Get(), TEXT("Profiler.TooltipBold") )
 						]
 						// [_MinInclusiveTimeMS] ms
 						+SGridPanel::Slot(1,3)
@@ -251,7 +250,7 @@
 						[
 							SNew(STextBlock)
 							.Text( FText::FromString(EventSample->GetFormattedValue(EEventPropertyIndex::MinInclusiveTimeMS)) )
-							.TextStyle( FAppStyle::Get(), TEXT("Profiler.Tooltip") )
+							.TextStyle( FProfilerStyle::Get(), TEXT("Profiler.Tooltip") )
 						]
 						// [_MaxInclusiveTimeMS] ms
 						+SGridPanel::Slot(2,3)
@@ -259,7 +258,7 @@
 						[
 							SNew(STextBlock)
 							.Text( FText::FromString(EventSample->GetFormattedValue(EEventPropertyIndex::MaxInclusiveTimeMS)) )
-							.TextStyle( FAppStyle::Get(), TEXT("Profiler.Tooltip") )
+							.TextStyle( FProfilerStyle::Get(), TEXT("Profiler.Tooltip") )
 						]
 						// [_AvgInclusiveTimeMS] ms
 						+SGridPanel::Slot(3,3)
@@ -267,7 +266,7 @@
 						[
 							SNew(STextBlock)
 							.Text( FText::FromString(EventSample->GetFormattedValue(EEventPropertyIndex::AvgInclusiveTimeMS)) )
-							.TextStyle( FAppStyle::Get(), TEXT("Profiler.Tooltip") )
+							.TextStyle( FProfilerStyle::Get(), TEXT("Profiler.Tooltip") )
 						]
 
 						//-----------------------------------------------------------------------------
@@ -279,7 +278,7 @@
 						[
 							SNew(STextBlock)
 							.Text(LOCTEXT("TT_ExclusiveTime", "Exclusive time:"))
-							.TextStyle( FAppStyle::Get(), TEXT("Profiler.TooltipBold") )
+							.TextStyle( FProfilerStyle::Get(), TEXT("Profiler.TooltipBold") )
 						]
 						// [_ExclusiveTimeMS] ms
 						+SGridPanel::Slot(1,4)
@@ -287,7 +286,7 @@
 						[
 							SNew(STextBlock)
 							.Text( FText::FromString(EventSample->GetFormattedValue(EEventPropertyIndex::ExclusiveTimeMS)) )
-							.TextStyle( FAppStyle::Get(), TEXT("Profiler.Tooltip") )
+							.TextStyle( FProfilerStyle::Get(), TEXT("Profiler.Tooltip") )
 						]
 						// ([_ExclusiveTimePct] % of this call's inc time)
 						+SGridPanel::Slot(2,4)
@@ -295,7 +294,7 @@
 						[
 							SNew(STextBlock)
 							.Text( ExclusiveTimePctCaller )
-							.TextStyle( FAppStyle::Get(), TEXT("Profiler.Tooltip") )
+							.TextStyle( FProfilerStyle::Get(), TEXT("Profiler.Tooltip") )
 						]
 
 						//-----------------------------------------------------------------------------
@@ -307,7 +306,7 @@
 						[
 							SNew(STextBlock)
 							.Text(LOCTEXT("TT_NulCallsPerFrame", "Num calls per frame:"))
-							.TextStyle( FAppStyle::Get(), TEXT("Profiler.TooltipBold") )
+							.TextStyle( FProfilerStyle::Get(), TEXT("Profiler.TooltipBold") )
 						]
 						// [_NumCallsPerFrame]
 						+SGridPanel::Slot(1,5)
@@ -315,7 +314,7 @@
 						[
 							SNew(STextBlock)
 							.Text( FText::FromString(EventSample->GetFormattedValue(EEventPropertyIndex::NumCallsPerFrame)) )
-							.TextStyle( FAppStyle::Get(), TEXT("Profiler.Tooltip") )
+							.TextStyle( FProfilerStyle::Get(), TEXT("Profiler.Tooltip") )
 						]
 
 
@@ -328,7 +327,7 @@
 						[
 							SNew( STextBlock )
 							.Text( LOCTEXT( "TT_MinMaxAvgNumCallsPerFrame", "Min/Max/Avg calls per frame:" ) )
-							.TextStyle( FAppStyle::Get(), TEXT( "Profiler.TooltipBold" ) )
+							.TextStyle( FProfilerStyle::Get(), TEXT( "Profiler.TooltipBold" ) )
 						]
 						// [_MinNumCallsPerFrame]
 						+ SGridPanel::Slot( 1, 6 )
@@ -336,7 +335,7 @@
 						[
 							SNew( STextBlock )
 							.Text( FText::FromString( EventSample->GetFormattedValue( EEventPropertyIndex::MinNumCallsPerFrame ) ) )
-							.TextStyle( FAppStyle::Get(), TEXT( "Profiler.Tooltip" ) )
+							.TextStyle( FProfilerStyle::Get(), TEXT( "Profiler.Tooltip" ) )
 						]
 						// [_MaxNumCallsPerFrame]
 						+ SGridPanel::Slot( 2, 6 )
@@ -344,7 +343,7 @@
 						[
 							SNew( STextBlock )
 							.Text( FText::FromString( EventSample->GetFormattedValue( EEventPropertyIndex::MaxNumCallsPerFrame ) ) )
-							.TextStyle( FAppStyle::Get(), TEXT( "Profiler.Tooltip" ) )
+							.TextStyle( FProfilerStyle::Get(), TEXT( "Profiler.Tooltip" ) )
 						]
 						// [_AvgNumCallsPerFrame]
 						+ SGridPanel::Slot( 3, 6 )
@@ -352,7 +351,7 @@
 						[
 							SNew( STextBlock )
 							.Text( FText::FromString( EventSample->GetFormattedValue( EEventPropertyIndex::AvgNumCallsPerFrame ) ) )
-							.TextStyle( FAppStyle::Get(), TEXT( "Profiler.Tooltip" ) )
+							.TextStyle( FProfilerStyle::Get(), TEXT( "Profiler.Tooltip" ) )
 						]
 						
 						//-----------------------------------------------------------------------------
@@ -364,7 +363,7 @@
 						[
 							SNew(STextBlock)
 							.Text(LOCTEXT("TT_ThreadDuration", "Thread duration:"))
-							.TextStyle( FAppStyle::Get(), TEXT("Profiler.TooltipBold") )
+							.TextStyle( FProfilerStyle::Get(), TEXT("Profiler.TooltipBold") )
 						]
 						// [_ThreadDurationMS]
 						+SGridPanel::Slot(1,7)
@@ -372,7 +371,7 @@
 						[
 							SNew(STextBlock)
 							.Text( FText::FromString(EventSample->GetFormattedValue(EEventPropertyIndex::ThreadDurationMS)) )
-							.TextStyle( FAppStyle::Get(), TEXT("Profiler.Tooltip") )
+							.TextStyle( FProfilerStyle::Get(), TEXT("Profiler.Tooltip") )
 						]
 
 						// Frame duration: 
@@ -381,7 +380,7 @@
 						[
 							SNew(STextBlock)
 							.Text(LOCTEXT("TT_FrameDuration", "Frame duration:"))
-							.TextStyle( FAppStyle::Get(), TEXT("Profiler.TooltipBold") )
+							.TextStyle( FProfilerStyle::Get(), TEXT("Profiler.TooltipBold") )
 						]
 						// [_FrameDurationMS]
 						+SGridPanel::Slot(3,7)
@@ -389,7 +388,7 @@
 						[
 							SNew(STextBlock)
 							.Text( FText::FromString(EventSample->GetFormattedValue(EEventPropertyIndex::FrameDurationMS)) )
-							.TextStyle( FAppStyle::Get(), TEXT("Profiler.Tooltip") )
+							.TextStyle( FProfilerStyle::Get(), TEXT("Profiler.Tooltip") )
 						]
 
 						/** Percent of time spent in the thread in relation to the entire frame. */
@@ -399,7 +398,7 @@
 						[
 							SNew(STextBlock)
 							.Text(LOCTEXT("ThreadToFrame", "Thread to Frame:"))
-							.TextStyle( FAppStyle::Get(), TEXT("Profiler.TooltipBold") )
+							.TextStyle( FProfilerStyle::Get(), TEXT("Profiler.TooltipBold") )
 						]
 						// [_FrameDurationMS]
 						+SGridPanel::Slot(5,7)
@@ -407,7 +406,7 @@
 						[
 							SNew(STextBlock)
 							.Text( FText::FromString(EventSample->GetFormattedValue(EEventPropertyIndex::ThreadToFramePct)) )
-							.TextStyle( FAppStyle::Get(), TEXT("Profiler.Tooltip") )
+							.TextStyle( FProfilerStyle::Get(), TEXT("Profiler.Tooltip") )
 						]
 					]
 
@@ -434,7 +433,7 @@
 			[
 				SNew(STextBlock)
 				.Text( FText::FromString(ParentName) )
-				.TextStyle( FAppStyle::Get(), TEXT("Profiler.Caption") )
+				.ColorAndOpacity(FSlateColor::UseForeground())
 			];
 
 		HBoxCaption->AddSlot()
@@ -442,6 +441,7 @@
 			[
 				SNew( SImage )
 				.Image( FAppStyle::GetBrush("BreadcrumbTrail.Delimiter") )
+				.ColorAndOpacity(FSlateColor::UseForeground())
 			];
 	}
 
@@ -450,7 +450,8 @@
 		[
 			SNew(STextBlock)
 			.Text( FText::FromString(EventSample->GetShortEventName()) )
-			.TextStyle( FAppStyle::Get(), TEXT("Profiler.CaptionBold") )
+			.Font( FProfilerStyle::Get().GetFontStyle("NormalBold") )
+			.ColorAndOpacity(FSlateColor::UseForeground())
 		];
 
 	if( bHasChildren )
@@ -493,6 +494,7 @@
 			[
 				SNew( SImage )
 				.Image( FAppStyle::GetBrush("BreadcrumbTrail.Delimiter") )
+				.ColorAndOpacity(FSlateColor::UseForeground())
 			];
 
 		HBoxCaption->AddSlot()
@@ -500,7 +502,7 @@
 			[
 				SNew(STextBlock)
 				.Text( FText::FromString(ChildrenNames) )
-				.TextStyle( FAppStyle::Get(), TEXT("Profiler.Caption") )
+				.ColorAndOpacity(FSlateColor::UseForeground())
 			];
 	}
 
diff --git a/Engine/Source/Developer/Profiler/Private/Widgets/SFiltersAndPresets.cpp b/Engine/Source/Developer/Profiler/Private/Widgets/SFiltersAndPresets.cpp
--- a/Engine/Source/Developer/Profiler/Private/Widgets/SFiltersAndPresets.cpp
+++ b/Engine/Source/Developer/Profiler/Private/Widgets/SFiltersAndPresets.cpp
@@ -10,11 +10,11 @@
 #include "Widgets/Layout/SGridPanel.h"
 #include "Widgets/SToolTip.h"
 #include "Widgets/Input/SCheckBox.h"
-#include "Styling/AppStyle.h"
 #include "ProfilerSession.h"
 #include "ProfilerManager.h"
 #include "Widgets/StatDragDropOp.h"
 #include "Widgets/Input/SSearchBox.h"
+#include "../ProfilerStyle.h"
 
 #define LOCTEXT_NAMESPACE "SFiltersAndPresets"
 
@@ -23,15 +23,15 @@
 {
 	static const FSlateBrush* GetIconForGroup()
 	{
-		return FAppStyle::GetBrush( TEXT( "Profiler.Misc.GenericGroup" ) );
+		return FProfilerStyle::Get().GetBrush( TEXT( "Profiler.Misc.GenericGroup" ) );
 	}
 
 	static const FSlateBrush* GetIconForStatType( const EProfilerSampleTypes::Type StatType )
 	{
-		const FSlateBrush* HierarchicalTimeIcon = FAppStyle::GetBrush( TEXT( "Profiler.Type.Hierarchical" ) );
-		const FSlateBrush* NumberIntIcon        = FAppStyle::GetBrush( TEXT( "Profiler.Type.NumberInt" ) );
-		const FSlateBrush* NumberFloatIcon      = FAppStyle::GetBrush( TEXT( "Profiler.Type.NumberFloat" ) );
-		const FSlateBrush* MemoryIcon           = FAppStyle::GetBrush( TEXT( "Profiler.Type.Memory" ) );
+		const FSlateBrush* HierarchicalTimeIcon = FProfilerStyle::Get().GetBrush( TEXT( "Profiler.Type.Hierarchical" ) );
+		const FSlateBrush* NumberIntIcon        = FProfilerStyle::Get().GetBrush( TEXT( "Profiler.Type.NumberInt" ) );
+		const FSlateBrush* NumberFloatIcon      = FProfilerStyle::Get().GetBrush( TEXT( "Profiler.Type.NumberFloat" ) );
+		const FSlateBrush* MemoryIcon           = FProfilerStyle::Get().GetBrush( TEXT( "Profiler.Type.Memory" ) );
 
 		const FSlateBrush* const StatIcons[ EProfilerSampleTypes::InvalidOrMax ] =
 		{
@@ -104,7 +104,7 @@
 		.ColumnSpan( 3 )
 		[
 			SNew(STextBlock)
-			.TextStyle( FAppStyle::Get(), TEXT("Profiler.TooltipBold") )
+			.TextStyle( FProfilerStyle::Get(), TEXT("Profiler.TooltipBold") )
 			.Text( LOCTEXT("NoStatData","There is no data for this stat") )
 		];
 		RowPos++;
@@ -119,7 +119,7 @@
 		.ColumnSpan( 3 )
 		[
 			SNew(STextBlock)
-			.TextStyle( FAppStyle::Get(), TEXT("Profiler.TooltipBold") )
+			.TextStyle( FProfilerStyle::Get(), TEXT("Profiler.TooltipBold") )
 			.Text( LOCTEXT("StatInstance","Stat information for profiler instance") )
 		];
 
@@ -128,7 +128,7 @@
 		.ColumnSpan( 3 )
 		[
 			SNew(STextBlock)
-			.TextStyle( FAppStyle::Get(), TEXT("Profiler.Tooltip") )
+			.TextStyle( FProfilerStyle::Get(), TEXT("Profiler.Tooltip") )
 			.Text( FText::FromString(InstanceName) )
 		];
 
@@ -145,7 +145,7 @@
 		.Padding( 2.0f )
 		[
 			SNew(STextBlock)
-			.TextStyle( FAppStyle::Get(), TEXT("Profiler.TooltipBold") )
+			.TextStyle( FProfilerStyle::Get(), TEXT("Profiler.TooltipBold") )
 			.Text( LOCTEXT("GroupDesc","Group:") )
 		];
 
@@ -154,7 +154,7 @@
 		.ColumnSpan( 2 )
 		[
 			SNew(STextBlock)
-			.TextStyle( FAppStyle::Get(), TEXT("Profiler.Tooltip") )
+			.TextStyle( FProfilerStyle::Get(), TEXT("Profiler.Tooltip") )
 			.Text( FText::FromName(ProfilerStat.OwningGroup().Name() ))
 		];
 		RowPos++;
@@ -163,7 +163,7 @@
 		.Padding( 2.0f )
 		[
 			SNew(STextBlock)
-			.TextStyle( FAppStyle::Get(), TEXT("Profiler.TooltipBold") )
+			.TextStyle( FProfilerStyle::Get(), TEXT("Profiler.TooltipBold") )
 			.Text( LOCTEXT("NameDesc","Name:") )
 		];
 
@@ -172,7 +172,7 @@
 		.ColumnSpan( 2 )
 		[
 			SNew(STextBlock)
-			.TextStyle( FAppStyle::Get(), TEXT("Profiler.Tooltip") )
+			.TextStyle( FProfilerStyle::Get(), TEXT("Profiler.Tooltip") )
 			.Text( FText::FromName(ProfilerStat.Name()) )
 		];
 		RowPos++;
@@ -181,7 +181,7 @@
 		.Padding( 2.0f )
 		[
 			SNew(STextBlock)
-			.TextStyle( FAppStyle::Get(), TEXT("Profiler.TooltipBold") )
+			.TextStyle( FProfilerStyle::Get(), TEXT("Profiler.TooltipBold") )
 			.Text( LOCTEXT("TypeDesc","Type:") )
 		];
 
@@ -205,7 +205,7 @@
 			[
 				SNew(STextBlock)
 				.Text( FText::FromString(EProfilerSampleTypes::ToDescription(SampleType)) )
-				.TextStyle( FAppStyle::Get(), TEXT("Profiler.Tooltip") )
+				.TextStyle( FProfilerStyle::Get(), TEXT("Profiler.Tooltip") )
 			]
 		];
 		RowPos++;
@@ -220,7 +220,7 @@
 		.ColumnSpan( 3 )
 		[
 			SNew(STextBlock)
-			.TextStyle( FAppStyle::Get(), TEXT("Profiler.TooltipBold") )
+			.TextStyle( FProfilerStyle::Get(), TEXT("Profiler.TooltipBold") )
 			.Text( LOCTEXT("ValueDesc","Value") )
 		];
 		RowPos++;
@@ -229,7 +229,7 @@
 		.Padding( 2.0f )
 		[
 			SNew(STextBlock)
-			.TextStyle( FAppStyle::Get(), TEXT("Profiler.Tooltip") )
+			.TextStyle( FProfilerStyle::Get(), TEXT("Profiler.Tooltip") )
 			.Text(FText::Format(LOCTEXT("MinDesc", "Min: {0}"), FText::FromString(Aggregated.GetFormattedValue(FProfilerAggregatedStat::EMinValue))))
 		];
 
@@ -237,7 +237,7 @@
 		.Padding( 2.0f )
 		[
 			SNew(STextBlock)
-			.TextStyle( FAppStyle::Get(), TEXT("Profiler.Tooltip") )
+			.TextStyle( FProfilerStyle::Get(), TEXT("Profiler.Tooltip") )
 			.Text(FText::Format(LOCTEXT("AvgDesc", "Avg: {0}"), FText::FromString(Aggregated.GetFormattedValue(FProfilerAggregatedStat::EAvgValue))))
 		];
 
@@ -245,7 +245,7 @@
 		.Padding( 2.0f )
 		[
 			SNew(STextBlock)
-			.TextStyle( FAppStyle::Get(), TEXT("Profiler.Tooltip") )
+			.TextStyle( FProfilerStyle::Get(), TEXT("Profiler.Tooltip") )
 			.Text(FText::Format(LOCTEXT("MaxDesc", "Max: {0}"), FText::FromString(Aggregated.GetFormattedValue(FProfilerAggregatedStat::EMaxValue))))
 		];
 		RowPos++;
@@ -265,7 +265,7 @@
 		.ColumnSpan( 3 )
 		[
 			SNew(STextBlock)
-			.TextStyle( FAppStyle::Get(), TEXT("Profiler.TooltipBold") )
+			.TextStyle( FProfilerStyle::Get(), TEXT("Profiler.TooltipBold") )
 			.Text(FText::Format(LOCTEXT("CallsFramesPctDesc", "Calls Frames with call: {0}"), FText::FromString(Aggregated.GetFormattedValue(FProfilerAggregatedStat::EFramesWithCallPct))))
 		];
 		RowPos++;
@@ -274,7 +274,7 @@
 		.Padding( 2.0f )
 		[
 			SNew(STextBlock)
-			.TextStyle( FAppStyle::Get(), TEXT("Profiler.Tooltip") )
+			.TextStyle( FProfilerStyle::Get(), TEXT("Profiler.Tooltip") )
 			.Text(FText::Format(LOCTEXT("MinDesc", "Min: {0}"), FText::FromString(Aggregated.GetFormattedValue(FProfilerAggregatedStat::EMinNumCalls))))
 		];
 
@@ -282,7 +282,7 @@
 		.Padding( 2.0f )
 		[
 			SNew(STextBlock)
-			.TextStyle( FAppStyle::Get(), TEXT("Profiler.Tooltip") )
+			.TextStyle( FProfilerStyle::Get(), TEXT("Profiler.Tooltip") )
 			.Text(FText::Format(LOCTEXT("AvgDesc", "Avg: {0}"), FText::FromString(Aggregated.GetFormattedValue(FProfilerAggregatedStat::EAvgNumCalls))))
 		];
 
@@ -290,7 +290,7 @@
 		.Padding( 2.0f )
 		[
 			SNew(STextBlock)
-			.TextStyle( FAppStyle::Get(), TEXT("Profiler.Tooltip") )
+			.TextStyle( FProfilerStyle::Get(), TEXT("Profiler.Tooltip") )
 			.Text(FText::Format(LOCTEXT("MaxDesc", "Max: {0}"), FText::FromString(Aggregated.GetFormattedValue(FProfilerAggregatedStat::EMaxNumCalls))))
 		];
 		RowPos++;
@@ -425,7 +425,7 @@
 				SNew( STextBlock )
 				.Text( this, &SGroupAndStatTableRow::GetText )
 				.HighlightText( InArgs._HighlightText )
-				.TextStyle( FAppStyle::Get(), TEXT("Profiler.Tooltip") )
+				.TextStyle( FProfilerStyle::Get(), TEXT("Profiler.Tooltip") )
 				.ColorAndOpacity( this, &SGroupAndStatTableRow::GetColorAndOpacity )
 			]
 
@@ -437,7 +437,7 @@
 			[
 				SNew( SImage )
 				.Visibility( !InGroupOrStatNode->IsGroup() ? EVisibility::Visible : EVisibility::Collapsed )
-				.Image( FAppStyle::GetBrush("Profiler.Tooltip.HintIcon10") )
+				.Image(FProfilerStyle::Get().GetBrush("Profiler.Tooltip.HintIcon10") )
 				.ToolTip( Tooltip )
 			]
 		];
@@ -1047,7 +1047,7 @@
 TSharedRef<SWidget> SFiltersAndPresets::GetToggleButtonForStatType( const EProfilerSampleTypes::Type StatType )
 {
 	return SNew( SCheckBox )
-		.Style( FAppStyle::Get(), "ToggleButtonCheckbox" )
+		.Style( FProfilerStyle::Get(), "ToggleButtonCheckbox" )
 		.HAlign( HAlign_Center )
 		.Padding( 2.0f )
 		.OnCheckStateChanged( this, &SFiltersAndPresets::FilterByStatType_OnCheckStateChanged, StatType )
@@ -1070,7 +1070,7 @@
 				[
 					SNew( STextBlock )
 						.Text( FText::FromString(EProfilerSampleTypes::ToName( StatType )) )
-						.TextStyle( FAppStyle::Get(), TEXT("Profiler.Caption") )
+						.ColorAndOpacity(FSlateColor::UseForeground() )
 				]
 		];
 }
diff --git a/Engine/Source/Developer/Profiler/Private/Widgets/SHistogram.cpp b/Engine/Source/Developer/Profiler/Private/Widgets/SHistogram.cpp
--- a/Engine/Source/Developer/Profiler/Private/Widgets/SHistogram.cpp
+++ b/Engine/Source/Developer/Profiler/Private/Widgets/SHistogram.cpp
@@ -11,8 +11,8 @@
 #include "Widgets/SOverlay.h"
 #include "Fonts/FontMeasure.h"
 #include "Framework/Application/SlateApplication.h"
-#include "Styling/AppStyle.h"
 #include "ProfilerFPSAnalyzer.h"
+#include "../ProfilerStyle.h"
 
 
 #define LOCTEXT_NAMESPACE "SHistogram"
@@ -40,9 +40,9 @@
 	// Rendering info.
 	const bool bEnabled  = ShouldBeEnabled( bParentEnabled );
 	ESlateDrawEffect DrawEffects = bEnabled ? ESlateDrawEffect::None : ESlateDrawEffect::DisabledEffect;
-	const FSlateBrush* TimelineAreaBrush = FAppStyle::GetBrush("Profiler.LineGraphArea");
-	const FSlateBrush* WhiteBrush = FAppStyle::GetBrush("WhiteTexture");
-	const FSlateBrush* FillImage  = FAppStyle::GetBrush("TaskGraph.Mono");
+	const FSlateBrush* TimelineAreaBrush = FProfilerStyle::Get().GetBrush("Brushes.White25");
+	const FSlateBrush* WhiteBrush = FProfilerStyle::Get().GetBrush("Brushes.White");
+	const FSlateBrush* FillImage  = FProfilerStyle::Get().GetBrush("TaskGraph.Mono");
 
 	// Draw background.
 	FSlateDrawElement::MakeBox
diff --git a/Engine/Source/Developer/Profiler/Private/Widgets/SProfilerMiniView.cpp b/Engine/Source/Developer/Profiler/Private/Widgets/SProfilerMiniView.cpp
--- a/Engine/Source/Developer/Profiler/Private/Widgets/SProfilerMiniView.cpp
+++ b/Engine/Source/Developer/Profiler/Private/Widgets/SProfilerMiniView.cpp
@@ -11,8 +11,8 @@
 #include "Fonts/FontMeasure.h"
 #include "Styling/CoreStyle.h"
 #include "Framework/Application/SlateApplication.h"
-#include "Styling/AppStyle.h"
 #include "ProfilerSession.h"
+#include "../ProfilerStyle.h"
 
 
 SProfilerMiniView::SProfilerMiniView()
@@ -102,8 +102,8 @@
 	// Rendering info.
 	const bool bEnabled = ShouldBeEnabled( bParentEnabled );
 	const ESlateDrawEffect DrawEffects = bEnabled ? ESlateDrawEffect::None : ESlateDrawEffect::DisabledEffect;
-	const FSlateBrush* MiniViewArea = FAppStyle::GetBrush( "Profiler.LineGraphArea" );
-	const FSlateBrush* WhiteBrush = FAppStyle::GetBrush( "WhiteTexture" );
+	const FSlateBrush* MiniViewArea = FProfilerStyle::Get().GetBrush( "Brushes.White25" );
+	const FSlateBrush* WhiteBrush = FProfilerStyle::Get().GetBrush( "Brushes.White" );
 
 	PaintState = new((void*)PaintStateMemory) FSlateOnPaintState( AllottedGeometry, MyCullingRect, OutDrawElements, LayerId, InWidgetStyle, DrawEffects );
 
@@ -248,7 +248,7 @@
 			OutDrawElements,
 			LayerId,
 			AllottedGeometry.ToPaintGeometry( FVector2D( SelectionBoxX0, 0.0f ), FVector2D( SelectionBoxX1 - SelectionBoxX0, AllottedGeometry.Size.Y ) ),
-			FAppStyle::GetBrush( "PlainBorder" ),
+			FProfilerStyle::Get().GetBrush( "PlainBorder" ),
 			DrawEffects,
 			FColorList::Green
 		);
diff --git a/Engine/Source/Developer/Profiler/Private/Widgets/SProfilerSettings.cpp b/Engine/Source/Developer/Profiler/Private/Widgets/SProfilerSettings.cpp
--- a/Engine/Source/Developer/Profiler/Private/Widgets/SProfilerSettings.cpp
+++ b/Engine/Source/Developer/Profiler/Private/Widgets/SProfilerSettings.cpp
@@ -16,8 +16,8 @@
 #include "Widgets/Layout/SGridPanel.h"
 #include "Widgets/Input/SButton.h"
 #include "Widgets/Input/SCheckBox.h"
-#include "Styling/AppStyle.h"
 #include "ProfilerManager.h"
+#include "../ProfilerStyle.h"
 
 
 #define LOCTEXT_NAMESPACE "SProfilerSettings"
@@ -167,7 +167,7 @@
 				.VAlign( VAlign_Center )
 				[
 					SNew( SImage )
-					.Image( FAppStyle::GetBrush("Profiler.Misc.Save16") )
+					.Image( FProfilerStyle::Get().GetBrush("Icons.Save") )
 				]
 
 				+SHorizontalBox::Slot()
@@ -195,7 +195,7 @@
 				.VAlign( VAlign_Center )
 				[
 					SNew( SImage )
-					.Image( FAppStyle::GetBrush("Profiler.Misc.Reset16") )
+					.Image(FProfilerStyle::Get().GetBrush("Profiler.Misc.Reset16") )
 				]
 
 				+SHorizontalBox::Slot()
diff --git a/Engine/Source/Developer/Profiler/Private/Widgets/SProfilerThreadView.cpp b/Engine/Source/Developer/Profiler/Private/Widgets/SProfilerThreadView.cpp
--- a/Engine/Source/Developer/Profiler/Private/Widgets/SProfilerThreadView.cpp
+++ b/Engine/Source/Developer/Profiler/Private/Widgets/SProfilerThreadView.cpp
@@ -86,8 +86,8 @@
 	// Rendering info.
 	const bool bEnabled = ShouldBeEnabled( bParentEnabled );
 	const ESlateDrawEffect DrawEffects = bEnabled ? ESlateDrawEffect::None : ESlateDrawEffect::DisabledEffect;
-	const FSlateBrush* BackgroundBrush = FAppStyle::GetBrush( "Profiler.LineGraphArea" );
-	const FSlateBrush* WhiteBrush = FAppStyle::GetBrush( "WhiteTexture" );
+	const FSlateBrush* BackgroundBrush = FAppStyle::GetBrush( "Brushes.White25" );
+	const FSlateBrush* WhiteBrush = FAppStyle::GetBrush( "Brushes.White" );
 
 	// Paint state for this call to OnPaint, valid only in this scope.
 	PaintState = new((void*)PaintStateMemory) FSlateOnPaintState( AllottedGeometry, MyCullingRect, OutDrawElements, LayerId, InWidgetStyle, DrawEffects );
diff --git a/Engine/Source/Developer/Profiler/Private/Widgets/SProfilerWindow.cpp b/Engine/Source/Developer/Profiler/Private/Widgets/SProfilerWindow.cpp
--- a/Engine/Source/Developer/Profiler/Private/Widgets/SProfilerWindow.cpp
+++ b/Engine/Source/Developer/Profiler/Private/Widgets/SProfilerWindow.cpp
@@ -27,6 +27,7 @@
 	#include "Runtime/Analytics/Analytics/Public/Interfaces/IAnalyticsProvider.h"
 	#include "EngineAnalytics.h"
 #endif // WITH_EDITOR
+#include "../ProfilerStyle.h"
 
 
 #define LOCTEXT_NAMESPACE "SProfilerWindow"
@@ -190,7 +191,7 @@
 													.AutoWidth()
 														[
 															SNew(SImage)
-															.Image(FAppStyle::GetBrush(TEXT("Profiler.Tab.FiltersAndPresets")))
+															.Image(FProfilerStyle::Get().GetBrush(TEXT("Profiler.Tab.FiltersAndPresets")))
 														]
 
 													+ SHorizontalBox::Slot()
@@ -225,14 +226,14 @@
 														.AutoWidth()
 														[
 															SNew(SImage)
-																.Image(FAppStyle::GetBrush(TEXT("Profiler.Tab.FiltersAndPresets")))
+															.Image(FProfilerStyle::Get().GetBrush(TEXT("Profiler.Tab.FiltersAndPresets")))
 														]
 
 													+ SHorizontalBox::Slot()
 														.AutoWidth()
 														[
 															SNew(STextBlock)
-																.Text( LOCTEXT("FiltersAndPresetsLabel", "Filters And Presets") )
+															.Text( LOCTEXT("FiltersAndPresetsLabel", "Filters And Presets") )
 														]
 												]
 
@@ -269,7 +270,7 @@
 																.AutoWidth()
 																[
 																	SNew(SImage)
-																		.Image(FAppStyle::GetBrush(TEXT("Profiler.Tab.GraphView")))
+																		.Image(FProfilerStyle::Get().GetBrush(TEXT("Profiler.Tab.GraphView")))
 																]
 
 															+ SHorizontalBox::Slot()
@@ -304,7 +305,7 @@
 				.VAlign(VAlign_Center)
 				[
 					SNew(SBorder)
-						.BorderImage(FAppStyle::GetBrush("NotificationList.ItemBackground"))
+						.BorderImage(FProfilerStyle::Get().GetBrush("NotificationList.ItemBackground"))
 						.Padding(8.0f)
 						.Visibility(this, &SProfilerWindow::IsSessionOverlayVissible)
 						[
@@ -354,7 +355,7 @@
 			.AutoWidth()
 			[
 				SNew(SImage)
-				.Image( FAppStyle::GetBrush( TEXT("Profiler.Tab.EventGraph") ) )
+				.Image(FProfilerStyle::Get().GetBrush( TEXT("Profiler.Tab.EventGraph") ) )
 			]
 
 			+SHorizontalBox::Slot()
@@ -629,7 +630,7 @@
 	(*OverlaySettingsSlot)
 	[
 		SNew(SBorder)
-		.BorderImage( FAppStyle::GetBrush("NotificationList.ItemBackground") )
+		.BorderImage(FProfilerStyle::Get().GetBrush("NotificationList.ItemBackground") )
 		.Padding( 8.0f )
 		[
 		 	SNew(SProfilerSettings)
diff --git a/Engine/Source/Developer/ProjectLauncher/Private/Widgets/Profile/SProjectLauncherProfileListRow.h b/Engine/Source/Developer/ProjectLauncher/Private/Widgets/Profile/SProjectLauncherProfileListRow.h
--- a/Engine/Source/Developer/ProjectLauncher/Private/Widgets/Profile/SProjectLauncherProfileListRow.h
+++ b/Engine/Source/Developer/ProjectLauncher/Private/Widgets/Profile/SProjectLauncherProfileListRow.h
@@ -62,7 +62,7 @@
 		STableRow<ILauncherProfilePtr>::ConstructInternal(
 			STableRow::FArguments()
 			.ShowSelection(false)
-			.Style(FAppStyle::Get(), "Launcher.NoHoverTableRow"),
+			.Style(FAppStyle::Get(), "TableView.Row"),
 			InOwnerTableView
 			);
 
diff --git a/Engine/Source/Developer/ProjectLauncher/Private/Widgets/Progress/SProjectLauncherProgress.h b/Engine/Source/Developer/ProjectLauncher/Private/Widgets/Progress/SProjectLauncherProgress.h
--- a/Engine/Source/Developer/ProjectLauncher/Private/Widgets/Progress/SProjectLauncherProgress.h
+++ b/Engine/Source/Developer/ProjectLauncher/Private/Widgets/Progress/SProjectLauncherProgress.h
@@ -99,7 +99,7 @@
 				.AutoHeight()
 				[
 					SNew(STextBlock)
-					.TextStyle(FAppStyle::Get(), "LargeText")
+					.Font(FAppStyle::Get().GetFontStyle("NormalFontBold"))
 					.Text(this, &SProjectLauncherProgress::GetSelectedProfileNameText)
 				]
 
diff --git a/Engine/Source/Developer/ProjectLauncher/Private/Widgets/Project/SProjectLauncherProjectPicker.cpp b/Engine/Source/Developer/ProjectLauncher/Private/Widgets/Project/SProjectLauncherProjectPicker.cpp
--- a/Engine/Source/Developer/ProjectLauncher/Private/Widgets/Project/SProjectLauncherProjectPicker.cpp
+++ b/Engine/Source/Developer/ProjectLauncher/Private/Widgets/Project/SProjectLauncherProjectPicker.cpp
@@ -101,7 +101,6 @@
 				[
 					// project selector
 					SNew(SComboButton)
-					.ComboButtonStyle(FAppStyle::Get(), "Launcher.ComboButton")
 					.ButtonContent()
 					[
 						SNew(STextBlock)
diff --git a/Engine/Source/Developer/ProjectLauncher/Private/Widgets/Shared/SProjectLauncherVariantSelector.h b/Engine/Source/Developer/ProjectLauncher/Private/Widgets/Shared/SProjectLauncherVariantSelector.h
--- a/Engine/Source/Developer/ProjectLauncher/Private/Widgets/Shared/SProjectLauncherVariantSelector.h
+++ b/Engine/Source/Developer/ProjectLauncher/Private/Widgets/Shared/SProjectLauncherVariantSelector.h
@@ -72,7 +72,7 @@
 			// build configuration menu
 			SNew(SComboButton)
 			.VAlign(VAlign_Center)
-			.ComboButtonStyle(FAppStyle::Get(), "Launcher.ComboButton")
+			.ComboButtonStyle(FAppStyle::Get(), "ComboButton")
 			.ButtonContent()
 			[
 				SNew(STextBlock)
diff --git a/Engine/Source/Developer/RealtimeProfiler/Private/SRealtimeProfilerLineGraph.cpp b/Engine/Source/Developer/RealtimeProfiler/Private/SRealtimeProfilerLineGraph.cpp
--- a/Engine/Source/Developer/RealtimeProfiler/Private/SRealtimeProfilerLineGraph.cpp
+++ b/Engine/Source/Developer/RealtimeProfiler/Private/SRealtimeProfilerLineGraph.cpp
@@ -112,8 +112,8 @@
 	// Rendering info
 	bool bEnabled = ShouldBeEnabled( bParentEnabled );
 	ESlateDrawEffect DrawEffects = bEnabled ? ESlateDrawEffect::None : ESlateDrawEffect::DisabledEffect;
-	const FSlateBrush* TimelineAreaBrush = FAppStyle::GetBrush("Profiler.LineGraphArea");
-	const FSlateBrush* WhiteBrush = FAppStyle::GetBrush("WhiteTexture");
+	const FSlateBrush* TimelineAreaBrush = FAppStyle::GetBrush("Brushes.White25");
+	const FSlateBrush* WhiteBrush = FAppStyle::GetBrush("Brushes.White");
 
 
 	// Draw timeline background
diff --git a/Engine/Source/Developer/ScreenShotComparison/Private/Widgets/SScreenShotBrowser.cpp b/Engine/Source/Developer/ScreenShotComparison/Private/Widgets/SScreenShotBrowser.cpp
--- a/Engine/Source/Developer/ScreenShotComparison/Private/Widgets/SScreenShotBrowser.cpp
+++ b/Engine/Source/Developer/ScreenShotComparison/Private/Widgets/SScreenShotBrowser.cpp
@@ -22,6 +22,8 @@
 #include "ISourceControlOperation.h"
 #include "SourceControlOperations.h"
 #include "ISourceControlProvider.h"
+#include "SPositiveActionButton.h"
+#include "SNegativeActionButton.h"
 
 #define LOCTEXT_NAMESPACE "ScreenshotComparison"
 
@@ -176,12 +178,9 @@
 				.AutoWidth()
 				.Padding(2.f, 0.f)
 				[
-					SNew(SButton)
-					.HAlign(HAlign_Center)
+					SNew(SPositiveActionButton)
 					.Text(LOCTEXT("AddAllNewReports", "Add All New Reports"))
 					.ToolTipText(LOCTEXT("AddAllNewReportsTooltip", "Adds all new screenshots contained in the reports."))
-					.ButtonStyle(FAppStyle::Get(), "FlatButton.Success")
-					.TextStyle(FAppStyle::Get(), "FlatButton.DefaultTextStyle")
 					.IsEnabled_Lambda([this]() -> bool
 						{
 							return ComparisonList.Num() > 0 && ISourceControlModule::Get().IsEnabled();
@@ -212,12 +211,10 @@
 				.AutoWidth()
 				.Padding(2.f, 0.f)
 				[
-					SNew(SButton)
-					.HAlign(HAlign_Center)
+					SNew(SNegativeActionButton)
+					.ActionButtonStyle(EActionButtonStyle::Warning)
 					.Text(LOCTEXT("ReplaceAllReports", "Replace All Reports"))
 					.ToolTipText(LOCTEXT("ReplaceAllReportsTooltip", "Replaces all screenshots containing a different result in the reports."))
-					.ButtonStyle(FAppStyle::Get(), "FlatButton.Warning")
-					.TextStyle(FAppStyle::Get(), "FlatButton.DefaultTextStyle")
 					.IsEnabled_Lambda([this]() -> bool
 						{
 							return ComparisonList.Num() > 0 && ISourceControlModule::Get().IsEnabled();
@@ -248,12 +245,9 @@
 				.AutoWidth()
 				.Padding(2.f, 0.f)
 				[
-					SNew(SButton)
-					.HAlign(HAlign_Center)
+					SNew(SNegativeActionButton)
 					.Text(LOCTEXT("DeleteAllReports", "Delete All Reports"))
 					.ToolTipText(LOCTEXT("DeleteAllReportsTooltip", "Deletes all the current reports.  Reports are not removed unless the user resolves them, \nso if you just want to reset the state of the reports, clear them here and then re-run the tests."))
-					.ButtonStyle(FAppStyle::Get(), "FlatButton.Danger")
-					.TextStyle(FAppStyle::Get(), "FlatButton.DefaultTextStyle")
 					.IsEnabled_Lambda([this]() -> bool
 						{
 							return ComparisonList.Num() > 0;
diff --git a/Engine/Source/Developer/ScreenShotComparison/ScreenShotComparison.Build.cs b/Engine/Source/Developer/ScreenShotComparison/ScreenShotComparison.Build.cs
--- a/Engine/Source/Developer/ScreenShotComparison/ScreenShotComparison.Build.cs
+++ b/Engine/Source/Developer/ScreenShotComparison/ScreenShotComparison.Build.cs
@@ -15,7 +15,7 @@
 		PrivateDependencyModuleNames.AddRange(
 			new string[] {
 				"AutomationController",
-				
+				"ToolWidgets",
 				"InputCore",
 				"ScreenShotComparisonTools",
 				"Slate",
diff --git a/Engine/Source/Developer/SessionFrontend/Private/Widgets/Browser/SSessionBrowserTreeGroupRow.h b/Engine/Source/Developer/SessionFrontend/Private/Widgets/Browser/SSessionBrowserTreeGroupRow.h
--- a/Engine/Source/Developer/SessionFrontend/Private/Widgets/Browser/SSessionBrowserTreeGroupRow.h
+++ b/Engine/Source/Developer/SessionFrontend/Private/Widgets/Browser/SSessionBrowserTreeGroupRow.h
@@ -18,6 +18,7 @@
 #include "Styling/AppStyle.h"
 #include "Models/SessionBrowserTreeItems.h"
 #include "Widgets/Images/SImage.h"
+#include "Styling/StyleColors.h"
 
 #define LOCTEXT_NAMESPACE "SSessionBrowserTreeRow"
 
@@ -75,7 +76,8 @@
 							.VAlign(VAlign_Center)
 							[
 								SNew(SImage)
-									.Image(this, &SSessionBrowserTreeGroupRow::HandlePullDownImage)
+								.Image(this, &SSessionBrowserTreeGroupRow::HandlePullDownImage)
+								.ColorAndOpacity(IsHovered() ? FStyleColors::ForegroundHover : FStyleColors::Foreground)
 							]
 					]
 			];
@@ -83,7 +85,7 @@
 		STableRow<TSharedPtr<FSessionBrowserGroupTreeItem>>::ConstructInternal(
 			STableRow<TSharedPtr<FSessionBrowserGroupTreeItem>>::FArguments()
 				.ShowSelection(false)
-				.Style(FAppStyle::Get(), "DetailsView.TreeView.TableRow"),
+				.Style(FAppStyle::Get(), "TableView.Row"),
 			InOwnerTableView
 		);
 	}
@@ -114,15 +116,11 @@
 	{
 		if (IsHovered())
 		{
-			return IsItemExpanded()
-				? FAppStyle::GetBrush("DetailsView.CategoryTop_Hovered")
-				: FAppStyle::GetBrush("DetailsView.CollapsedCategory_Hovered");
+			return FAppStyle::GetBrush("Brushes.Hover");
 		}
 		else
 		{
-			return IsItemExpanded()
-				? FAppStyle::GetBrush("DetailsView.CategoryTop")
-				: FAppStyle::GetBrush("DetailsView.CollapsedCategory");
+			return FAppStyle::GetBrush("Brushes.Header");
 		}
 	}
 
@@ -150,18 +148,9 @@
 	/** Gets the image for the pull-down icon. */
 	const FSlateBrush* HandlePullDownImage() const
 	{
-		if (IsHovered())
-		{
-			return IsItemExpanded()
-				? FAppStyle::GetBrush("DetailsView.PulldownArrow.Up.Hovered")
-				: FAppStyle::GetBrush("DetailsView.PulldownArrow.Down.Hovered");
-		}
-		else
-		{
-			return IsItemExpanded()
-				? FAppStyle::GetBrush("DetailsView.PulldownArrow.Up")
-				: FAppStyle::GetBrush("DetailsView.PulldownArrow.Down");
-		}
+		return IsItemExpanded()
+			? FAppStyle::GetBrush("Icons.ChevronUp")
+			: FAppStyle::GetBrush("Icons.ChevronDown");
 	}
 
 private:
diff --git a/Engine/Source/Developer/SessionFrontend/Private/Widgets/Browser/SSessionBrowserTreeSessionRow.h b/Engine/Source/Developer/SessionFrontend/Private/Widgets/Browser/SSessionBrowserTreeSessionRow.h
--- a/Engine/Source/Developer/SessionFrontend/Private/Widgets/Browser/SSessionBrowserTreeSessionRow.h
+++ b/Engine/Source/Developer/SessionFrontend/Private/Widgets/Browser/SSessionBrowserTreeSessionRow.h
@@ -82,7 +82,7 @@
 		STableRow<TSharedPtr<FSessionBrowserSessionTreeItem>>::ConstructInternal(
 			STableRow<TSharedPtr<FSessionBrowserSessionTreeItem>>::FArguments()
 				.ShowSelection(false)
-				.Style(FAppStyle::Get(), "DetailsView.TreeView.TableRow"),
+				.Style(FAppStyle::Get(), "TableView.Row"),
 			InOwnerTableView
 		);
 	}
@@ -153,15 +153,11 @@
 	{
 		if (IsHovered())
 		{
-			return IsItemExpanded()
-				? FAppStyle::GetBrush("DetailsView.CategoryTop_Hovered")
-				: FAppStyle::GetBrush("DetailsView.CollapsedCategory_Hovered");
+			return FAppStyle::GetBrush("Brushes.Hover");
 		}
 		else
 		{
-			return IsItemExpanded()
-				? FAppStyle::GetBrush("DetailsView.CategoryTop")
-				: FAppStyle::GetBrush("DetailsView.CollapsedCategory");
+			return FAppStyle::GetBrush("Brushes.Header");
 		}
 	}
 
diff --git a/Engine/Source/Developer/TaskGraph/Private/TaskGraphStyle.cpp b/Engine/Source/Developer/TaskGraph/Private/TaskGraphStyle.cpp
--- a/Engine/Source/Developer/TaskGraph/Private/TaskGraphStyle.cpp
+++ b/Engine/Source/Developer/TaskGraph/Private/TaskGraphStyle.cpp
@@ -2,11 +2,8 @@
 
 #include "TaskGraphStyle.h"
 #include "Styling/SlateStyleRegistry.h"
+#include "Styling/SlateStyleMacros.h"
 
-#define IMAGE_BRUSH( RelativePath, ... ) FSlateImageBrush( RootToContentDir( RelativePath, TEXT(".png") ), __VA_ARGS__ )
-#define BOX_BRUSH( RelativePath, ... ) FSlateBoxBrush( RootToContentDir( RelativePath, TEXT(".png") ), __VA_ARGS__ )
-#define BORDER_BRUSH( RelativePath, ... ) FSlateBorderBrush( RootToContentDir( RelativePath, TEXT(".png") ), __VA_ARGS__ )
-
 TSharedPtr< FTaskGraphStyle::FStyle > FTaskGraphStyle::StyleInstance;
 
 FTaskGraphStyle::FStyle::FStyle()
@@ -24,9 +21,9 @@
 	StyleInstance->Set( "StatsHeader", new BOX_BRUSH( "Common/TableViewHeader", 4.f/32.f ) );
 	StyleInstance->Set( "TaskGraph.ContentAreaBrush", new BOX_BRUSH( "/Docking/TabContentArea", FMargin(4/16.0f) ) );
 	StyleInstance->Set( "TaskGraph.Background", new BOX_BRUSH( "Common/ProgressBar_Background", FMargin(5.f/12.f) ) );
-	StyleInstance->Set( "TaskGraph.Normal", new BOX_BRUSH( "Common/TaskGraph_Normal", FMargin(5.f/12.f) ) );
-	StyleInstance->Set( "TaskGraph.Selected", new BOX_BRUSH( "Common/TaskGraph_Selected", FMargin(5.f/12.f) ) );
-	StyleInstance->Set( "TaskGraph.Mono", new BOX_BRUSH( "Common/TaskGraph_Mono", FMargin(5.f/12.f) ) );
+	StyleInstance->Set( "TaskGraph.Normal", new CORE_BOX_BRUSH( "Common/TaskGraph_Normal", FMargin(5.f/12.f) ) );
+	StyleInstance->Set( "TaskGraph.Selected", new CORE_BOX_BRUSH( "Common/TaskGraph_Selected", FMargin(5.f/12.f) ) );
+	StyleInstance->Set( "TaskGraph.Mono", new CORE_BOX_BRUSH( "Common/TaskGraph_Mono", FMargin(5.f/12.f) ) );
 	StyleInstance->Set( "TaskGraph.BorderPadding", FVector2D(1,0) );
 	StyleInstance->Set( "TaskGraph.SortUp", new IMAGE_BRUSH( "Common/SortUpArrow", FVector2D(8,4) ) );
 	StyleInstance->Set( "TaskGraph.SortDown", new IMAGE_BRUSH( "Common/SortDownArrow", FVector2D(8,4) ) );
diff --git a/Engine/Source/Editor/EditorStyle/Private/SlateEditorStyle.cpp b/Engine/Source/Editor/EditorStyle/Private/SlateEditorStyle.cpp
--- a/Engine/Source/Editor/EditorStyle/Private/SlateEditorStyle.cpp
+++ b/Engine/Source/Editor/EditorStyle/Private/SlateEditorStyle.cpp
@@ -1220,47 +1220,7 @@
 	);
 
 
-	// Output Log Window
 #if WITH_EDITOR || (IS_PROGRAM && WITH_UNREAL_DEVELOPER_TOOLS)
-	{
-		const int32 LogFontSize = Settings.IsValid() ? Settings->LogFontSize : 9;
-
-		const FTextBlockStyle NormalLogText = FTextBlockStyle(NormalText)
-			.SetFont( DEFAULT_FONT( "Mono", LogFontSize ) )
-			.SetColorAndOpacity( LogColor_Normal )
-			.SetSelectedBackgroundColor( LogColor_SelectionBackground );
-
-		Set("Log.Normal", NormalLogText );
-
-		Set("Log.Command", FTextBlockStyle(NormalLogText)
-			.SetColorAndOpacity( LogColor_Command )
-			);
-
-		Set("Log.Warning", FTextBlockStyle(NormalLogText)
-			.SetColorAndOpacity( LogColor_Warning )
-			);
-
-		Set("Log.Error", FTextBlockStyle(NormalLogText)
-			.SetColorAndOpacity( LogColor_Error )
-			);
-
-		Set("Log.TabIcon", new IMAGE_BRUSH( "Icons/icon_tab_OutputLog_16x", Icon16x16 ) );
-
-		Set("Log.TextBox", FEditableTextBoxStyle(NormalEditableTextBoxStyle)
-			.SetBackgroundImageNormal( BOX_BRUSH( "Common/WhiteGroupBorder", FMargin(4.0f/16.0f) ) )
-			.SetBackgroundImageHovered( BOX_BRUSH( "Common/WhiteGroupBorder", FMargin(4.0f/16.0f) ) )
-			.SetBackgroundImageFocused( BOX_BRUSH( "Common/WhiteGroupBorder", FMargin(4.0f/16.0f) ) )
-			.SetBackgroundImageReadOnly( BOX_BRUSH( "Common/WhiteGroupBorder", FMargin(4.0f/16.0f) ) )
-			.SetBackgroundColor( LogColor_Background )
-			);
-
-		Set("DebugConsole.Background", new BOX_BRUSH("Old/Menu_Background", FMargin(8.0f / 64.0f)));
-
-		Set("OutputLog.OpenSourceLocation", new IMAGE_BRUSH("Icons/icon_Asset_Open_Source_Location_16x", Icon16x16));
-		Set("OutputLog.OpenInExternalEditor", new IMAGE_BRUSH("Icons/icon_Asset_Open_In_External_Editor_16x", Icon16x16));
-
-	}
-
 	// Debug Tools Window
 	{
 		Set("DebugTools.TabIcon", new IMAGE_BRUSH( "Icons/icon_tab_DebugTools_16x", Icon16x16 ) );
@@ -3773,7 +3733,7 @@
 		Set( "Profiler.Tab", new IMAGE_BRUSH( "Icons/Profiler/Profiler_Tab_16x", Icon16x16 ) );
 		Set( "Profiler.Tab.GraphView", new IMAGE_BRUSH( "Icons/Profiler/Profiler_Graph_View_Tab_16x", Icon16x16 ) );
 		Set( "Profiler.Tab.EventGraph", new IMAGE_BRUSH( "Icons/Profiler/profiler_OpenEventGraph_32x", Icon16x16 ) );
-		Set( "Profiler.Tab.FiltersAndPresets", new IMAGE_BRUSH( "Icons/Profiler/Profiler_Filter_Presets_Tab_16x", Icon16x16 ) );
+		Set( "Profiler.Tab.FiltersAndPresets", new CORE_IMAGE_BRUSH( "Icons/Profiler/Profiler_Filter_Presets_Tab_16x", Icon16x16 ) );
 
 		Set( "ProfilerCommand.ProfilerManager_Load", new IMAGE_BRUSH( "Icons/Profiler/Profiler_Load_Profiler_40x", Icon40x40 ) );
 		Set( "ProfilerCommand.ProfilerManager_Load.Small", new IMAGE_BRUSH( "Icons/Profiler/Profiler_Load_Profiler_40x", Icon20x20 ) );
@@ -3902,7 +3862,6 @@
 		Set( "Profiler.EventGraph.Border.R", new BOX_BRUSH( "Icons/Profiler/Profiler_Border_R_16x",   FMargin(4.0f/16.0f) ) );
 
 		// Misc
-		Set( "Profiler.Misc.WarningSmall", new IMAGE_BRUSH( "ContentBrowser/SCC_NotAtHeadRevision", Icon12x12 ) );
 
 		Set( "Profiler.Misc.SortBy", new IMAGE_BRUSH( "Icons/Profiler/profiler_SortBy_32x", Icon32x32 ) );
 		Set( "Profiler.Misc.SortAscending", new IMAGE_BRUSH( "Icons/Profiler/profiler_SortAscending_32x", Icon32x32 ) );
@@ -3910,7 +3869,6 @@
 
 		Set( "Profiler.Misc.ResetToDefault", new IMAGE_BRUSH( "Icons/Profiler/profiler_ResetToDefault_32x", Icon32x32 ) );
 
-		Set( "Profiler.Misc.Save16", new IMAGE_BRUSH( "Icons/LV_Save", Icon16x16 ) );
 		Set( "Profiler.Misc.Reset16", new IMAGE_BRUSH( "Icons/Profiler/profiler_ResetToDefault_32x", Icon16x16 ) );
 
 		Set( "Profiler.Type.Calls", new IMAGE_BRUSH( "Icons/Profiler/profiler_Calls_32x", Icon16x16 ) );
@@ -7844,151 +7802,7 @@
 {
 	//Automation
 #if WITH_EDITOR || (IS_PROGRAM && WITH_UNREAL_DEVELOPER_TOOLS)
-	{
-		Set( "Automation.Header" , FTextBlockStyle(NormalText)
-			.SetFont( DEFAULT_FONT( "Mono", 12 ) )
-			.SetColorAndOpacity(FLinearColor(FColor(0xffffffff))) );
 
-		Set( "Automation.Normal" , FTextBlockStyle(NormalText)
-			.SetFont( DEFAULT_FONT( "Mono", 9 ) )
-			.SetColorAndOpacity(FLinearColor(FColor(0xffaaaaaa))) );
-
-		Set( "Automation.Warning", FTextBlockStyle(NormalText)
-			.SetFont( DEFAULT_FONT( "Mono", 9 ) )
-			.SetColorAndOpacity(FLinearColor(FColor(0xffbbbb44))) );
-
-		Set( "Automation.Error"  , FTextBlockStyle(NormalText)
-			.SetFont( DEFAULT_FONT( "Mono", 9 ) )
-			.SetColorAndOpacity(FLinearColor(FColor(0xffff0000))) );
-
-		Set( "Automation.ReportHeader" , FTextBlockStyle(NormalText)
-			.SetFont( DEFAULT_FONT( "Mono", 10 ) )
-			.SetColorAndOpacity(FLinearColor(FColor(0xffffffff))) );
-		
-		//state of individual tests
-		Set( "Automation.Success", new IMAGE_BRUSH( "Automation/Success", Icon16x16 ) );
-		Set( "Automation.Warning", new IMAGE_BRUSH( "Automation/Warning", Icon16x16 ) );
-		Set( "Automation.Fail", new IMAGE_BRUSH( "Automation/Fail", Icon16x16 ) );
-		Set( "Automation.InProcess", new IMAGE_BRUSH( "Automation/InProcess", Icon16x16 ) );
-		Set( "Automation.NotRun", new IMAGE_BRUSH( "Automation/NotRun", Icon16x16, FLinearColor(0.0f, 0.0f, 0.0f, 0.4f) ) );
-		Set( "Automation.Skipped", new IMAGE_BRUSH( "Automation/NoSessionWarning", Icon16x16 ) );
-		Set( "Automation.ParticipantsWarning", new IMAGE_BRUSH( "Automation/ParticipantsWarning", Icon16x16 ) );
-		Set( "Automation.Participant", new IMAGE_BRUSH( "Automation/Participant", Icon16x16 ) );
-		
-		//status as a regression test or not
-		Set( "Automation.SmokeTest", new IMAGE_BRUSH( "Automation/SmokeTest", Icon16x16 ) );
-		Set( "Automation.SmokeTestParent", new IMAGE_BRUSH( "Automation/SmokeTestParent", Icon16x16 ) );
-
-		//run icons
-		Set( "AutomationWindow.RunTests", new IMAGE_BRUSH( "Automation/RunTests", Icon40x40) );
-		Set( "AutomationWindow.RefreshTests", new IMAGE_BRUSH( "Automation/RefreshTests", Icon40x40) );
-		Set( "AutomationWindow.FindWorkers", new IMAGE_BRUSH( "Automation/RefreshWorkers", Icon40x40) );
-		Set( "AutomationWindow.StopTests", new IMAGE_BRUSH( "Automation/StopTests", Icon40x40 ) );
-		Set( "AutomationWindow.RunTests.Small", new IMAGE_BRUSH( "Automation/RunTests", Icon20x20) );
-		Set( "AutomationWindow.RefreshTests.Small", new IMAGE_BRUSH( "Automation/RefreshTests", Icon20x20) );
-		Set( "AutomationWindow.FindWorkers.Small", new IMAGE_BRUSH( "Automation/RefreshWorkers", Icon20x20) );
-		Set( "AutomationWindow.StopTests.Small", new IMAGE_BRUSH( "Automation/StopTests", Icon20x20 ) );
-
-		//filter icons
-		Set( "AutomationWindow.ErrorFilter", new IMAGE_BRUSH( "Automation/ErrorFilter", Icon40x40) );
-		Set( "AutomationWindow.WarningFilter", new IMAGE_BRUSH( "Automation/WarningFilter", Icon40x40) );
-		Set( "AutomationWindow.SmokeTestFilter", new IMAGE_BRUSH( "Automation/SmokeTestFilter", Icon40x40) );
-		Set( "AutomationWindow.DeveloperDirectoryContent", new IMAGE_BRUSH( "Automation/DeveloperDirectoryContent", Icon40x40) );
-		Set( "AutomationWindow.ExcludedTestsFilter", new IMAGE_BRUSH("Automation/ExcludedTestsFilter", Icon40x40) );
-		Set( "AutomationWindow.ErrorFilter.Small", new IMAGE_BRUSH( "Automation/ErrorFilter", Icon20x20) );
-		Set( "AutomationWindow.WarningFilter.Small", new IMAGE_BRUSH( "Automation/WarningFilter", Icon20x20) );
-		Set( "AutomationWindow.SmokeTestFilter.Small", new IMAGE_BRUSH( "Automation/SmokeTestFilter", Icon20x20) );
-		Set( "AutomationWindow.DeveloperDirectoryContent.Small", new IMAGE_BRUSH( "Automation/DeveloperDirectoryContent", Icon20x20) );
-		Set( "AutomationWindow.TrackHistory", new IMAGE_BRUSH( "Automation/TrackTestHistory", Icon40x40) );
-
-		//device group settings
-		Set( "AutomationWindow.GroupSettings", new IMAGE_BRUSH( "Automation/Groups", Icon40x40) );
-		Set( "AutomationWindow.GroupSettings.Small", new IMAGE_BRUSH( "Automation/Groups", Icon20x20) );
-
-		//test preset icons
-		Set( "AutomationWindow.PresetNew", new IMAGE_BRUSH( "Icons/icon_add_40x", Icon16x16 ) );
-		Set( "AutomationWindow.PresetSave", new IMAGE_BRUSH( "Icons/icon_file_save_16px", Icon16x16 ) );
-		Set( "AutomationWindow.PresetRemove", new IMAGE_BRUSH( "Icons/icon_Cascade_DeleteLOD_40x", Icon16x16 ) );
-
-		//test backgrounds
-		Set( "AutomationWindow.GameGroupBorder", new BOX_BRUSH( "Automation/GameGroupBorder", FMargin(4.0f/16.0f) ) );
-		Set( "AutomationWindow.EditorGroupBorder", new BOX_BRUSH( "Automation/EditorGroupBorder", FMargin(4.0f/16.0f) ) );
-	}
-
-	// Launcher
-	{
-		Set( "Launcher.Run", new IMAGE_BRUSH("Launcher/Launcher_Run", Icon40x40) );
-		Set( "Launcher.EditSettings", new IMAGE_BRUSH("Launcher/Launcher_EditSettings", Icon40x40) );
-		Set( "Launcher.Back", new IMAGE_BRUSH("Launcher/Launcher_Back", Icon32x32) );
-		Set( "Launcher.Back.Small", new IMAGE_BRUSH("Launcher/Launcher_Back", Icon32x32));
-		Set( "Launcher.Delete", new IMAGE_BRUSH("Launcher/Launcher_Delete", Icon32x32) );
-
-		Set( "Launcher.Instance_Commandlet", new IMAGE_BRUSH( "Launcher/Instance_Commandlet", Icon25x25 ) );
-		Set( "Launcher.Instance_Editor", new IMAGE_BRUSH( "Launcher/Instance_Editor", Icon25x25 ) );
-		Set( "Launcher.Instance_Game", new IMAGE_BRUSH( "Launcher/Instance_Game", Icon25x25 ) );
-		Set( "Launcher.Instance_Other", new IMAGE_BRUSH( "Launcher/Instance_Other", Icon25x25 ) );
-		Set( "Launcher.Instance_Server", new IMAGE_BRUSH( "Launcher/Instance_Server", Icon25x25 ) );
-		Set( "Launcher.Instance_Unknown", new IMAGE_BRUSH( "Launcher/Instance_Unknown", Icon25x25 ) );
-		Set( "LauncherCommand.DeployBuild", new IMAGE_BRUSH( "Launcher/Launcher_Deploy", Icon40x40 ) );
-		Set( "LauncherCommand.QuickLaunch", new IMAGE_BRUSH( "Launcher/Launcher_Launch", Icon40x40 ) );
-		Set( "LauncherCommand.CreateBuild", new IMAGE_BRUSH( "Launcher/Launcher_Build", Icon40x40 ) );
-		Set( "LauncherCommand.AdvancedBuild", new IMAGE_BRUSH( "Launcher/Launcher_Advanced", Icon40x40 ) );
-		Set( "LauncherCommand.AdvancedBuild.Medium", new IMAGE_BRUSH("Launcher/Launcher_Advanced", Icon25x25) );
-		Set( "LauncherCommand.AdvancedBuild.Small", new IMAGE_BRUSH("Launcher/Launcher_Advanced", Icon20x20) );
-
-		Set("Launcher.Filters.Text", FTextBlockStyle(NormalText)
-			.SetFont(DEFAULT_FONT("Bold", 9))
-			.SetColorAndOpacity(FLinearColor(1.0f, 1.0f, 1.0f, 0.9f))
-			.SetShadowOffset(FVector2D(1, 1))
-			.SetShadowColorAndOpacity(FLinearColor(0, 0, 0, 0.9f)));
-
-		Set( "Launcher.Platform.Warning", new IMAGE_BRUSH( "Icons/alert", Icon24x24) );
-
-#if DDPI_HAS_EXTENDED_PLATFORMINFO_DATA
-
-		for (auto Pair : FDataDrivenPlatformInfoRegistry::GetAllPlatformInfos())
-		{
-			const FDataDrivenPlatformInfo& PlatformInfo = Pair.Value;
-
-			// some platforms may specify a "rooted" path in the platform extensions directory, so look for that case here, and use a different path for the brush
-			FString NormalIconPath = PlatformInfo.GetIconPath(EPlatformIconSize::Normal);
-			if (NormalIconPath.StartsWith(TEXT("/Platforms/")))
-			{
-#define PLATFORM_IMAGE_BRUSH( PlatformPath, ... ) FSlateImageBrush( PlatformPath.Replace(TEXT("/Platforms/"), *FPaths::EnginePlatformExtensionsDir()) + TEXT(".png") , __VA_ARGS__ )
-				Set(PlatformInfo.GetIconStyleName(EPlatformIconSize::Normal), new PLATFORM_IMAGE_BRUSH(NormalIconPath, Icon24x24));
-				Set(PlatformInfo.GetIconStyleName(EPlatformIconSize::Large), new PLATFORM_IMAGE_BRUSH(PlatformInfo.GetIconPath(EPlatformIconSize::Large), Icon64x64));
-				Set(PlatformInfo.GetIconStyleName(EPlatformIconSize::XLarge), new PLATFORM_IMAGE_BRUSH(PlatformInfo.GetIconPath(EPlatformIconSize::XLarge), Icon128x128));
-			}
-			else
-			{
-				Set(PlatformInfo.GetIconStyleName(EPlatformIconSize::Normal), new IMAGE_BRUSH(*NormalIconPath, Icon24x24));
-				Set(PlatformInfo.GetIconStyleName(EPlatformIconSize::Large), new IMAGE_BRUSH(*PlatformInfo.GetIconPath(EPlatformIconSize::Large), Icon64x64));
-				Set(PlatformInfo.GetIconStyleName(EPlatformIconSize::XLarge), new IMAGE_BRUSH(*PlatformInfo.GetIconPath(EPlatformIconSize::XLarge), Icon128x128));
-			}
-		}
-
-		for (const FPreviewPlatformMenuItem& Item : FDataDrivenPlatformInfoRegistry::GetAllPreviewPlatformMenuItems())
-		{
-			if(!Item.ActiveIconPath.IsEmpty())
-			{
-				Set(Item.ActiveIconName, new PLATFORM_IMAGE_BRUSH(Item.ActiveIconPath, Icon40x40));
-			}
-			if (!Item.InactiveIconPath.IsEmpty())
-			{
-				Set(Item.InactiveIconName, new PLATFORM_IMAGE_BRUSH(Item.InactiveIconPath, Icon40x40));
-			}
-		}
-
-#endif
-
-		Set("Launcher.NoHoverTableRow", FTableRowStyle(NormalTableRowStyle)
-			.SetEvenRowBackgroundHoveredBrush(FSlateNoResource())
-			.SetOddRowBackgroundHoveredBrush(FSlateNoResource())
-			.SetActiveHoveredBrush(FSlateNoResource())
-			.SetInactiveHoveredBrush(FSlateNoResource())
-			);
-	}
-
 	// Device Manager
 	{
 		Set( "DeviceDetails.Claim", new IMAGE_BRUSH( "Icons/icon_DeviceClaim_40x", Icon40x40 ) );
diff --git a/Engine/Source/Editor/EditorStyle/Private/StarshipStyle.cpp b/Engine/Source/Editor/EditorStyle/Private/StarshipStyle.cpp
--- a/Engine/Source/Editor/EditorStyle/Private/StarshipStyle.cpp
+++ b/Engine/Source/Editor/EditorStyle/Private/StarshipStyle.cpp
@@ -32,18 +32,6 @@ void FStarshipEditorStyle::Initialize()
 	// The core style must be initialized before the editor style
 	FSlateApplication::InitializeCoreStyle();
 
-#if WITH_EDITOR || IS_PROGRAM
-	// Selection highlight
-	USlateThemeManager::Get().SetDefaultColor(EStyleColor::User2, USlateThemeManager::Get().GetColor(EStyleColor::Highlight));
-	USlateThemeManager::Get().SetColorDisplayName(EStyleColor::User2, LOCTEXT("UserColor_OutputLogHighlight", "Log Highlight"));
-	// Normal
-	USlateThemeManager::Get().SetDefaultColor(EStyleColor::User3, USlateThemeManager::Get().GetColor(EStyleColor::Foreground));
-	USlateThemeManager::Get().SetColorDisplayName(EStyleColor::User3, LOCTEXT("UserColor_OutputLogText", "Log Text"));
-	// Command
-	USlateThemeManager::Get().SetDefaultColor(EStyleColor::User4, USlateThemeManager::Get().GetColor(EStyleColor::AccentGreen));
-	USlateThemeManager::Get().SetColorDisplayName(EStyleColor::User4, LOCTEXT("UserColor_OutputLogCommand", "Log Command"));
-#endif // #if WITH_EDITOR || IS_PROGRAM
-
 	const FString ThemesSubDir = TEXT("Slate/Themes");
 
 #if ALLOW_THEMES
@@ -129,10 +117,6 @@ FStarshipEditorStyle::FStyle::FStyle( const TWeakObjectPtr< UEditorStyleSettings
 	, SelectionColor_Subdued( SelectionColor_Subdued_LinearRef )
 	, HighlightColor( HighlightColor_LinearRef )
 	, WindowHighlightColor(WindowHighlightColor_LinearRef)
-	, LogColor_SelectionBackground(EStyleColor::User2)
-	, LogColor_Normal(EStyleColor::User3)
-	, LogColor_Command(EStyleColor::User4)
-
 	, InheritedFromBlueprintTextColor(FLinearColor(0.25f, 0.5f, 1.0f))
 
 	, Settings( InSettings )
@@ -153,7 +137,7 @@ FStarshipEditorStyle::FStyle::~FStyle()
 
 void FStarshipEditorStyle::FStyle::SettingsChanged(FName PropertyName)
 {
-	SyncSettings();
+		SyncSettings();
 }
 
 void FStarshipEditorStyle::FStyle::SyncSettings()
@@ -413,8 +397,7 @@ void FStarshipEditorStyle::FStyle::SetupGeneralStyles()
 		Set("Icons.SelectInViewport", new IMAGE_BRUSH_SVG("Starship/Common/SelectInViewport", Icon16x16));
 		Set("Icons.BrowseContent", new IMAGE_BRUSH_SVG("Starship/Common/BrowseContent", Icon16x16));
 		Set("Icons.Use", new IMAGE_BRUSH_SVG("Starship/Common/use-circle", Icon16x16));
-		Set("Icons.Advanced", new IMAGE_BRUSH_SVG("Starship/Common/Advanced", Icon16x16));
-		Set("Icons.Launch", new IMAGE_BRUSH_SVG("Starship/Common/ProjectLauncher", Icon16x16));
+		Set("Icons.Launch", new CORE_IMAGE_BRUSH_SVG("Starship/Common/ProjectLauncher", Icon16x16));
 		Set("Icons.Next", new IMAGE_BRUSH_SVG("Starship/Common/NextArrow", Icon16x16));
 		Set("Icons.Previous", new IMAGE_BRUSH_SVG("Starship/Common/PreviousArrow", Icon16x16));
 		Set("Icons.Visibility", new IMAGE_BRUSH_SVG("Starship/Common/Visibility", Icon20x20));
@@ -450,7 +433,7 @@ void FStarshipEditorStyle::FStyle::SetupGeneralStyles()
 
 		Set("Icons.Toolbar.Play", new IMAGE_BRUSH_SVG("Starship/Common/play", Icon20x20));
 		Set("Icons.Toolbar.Pause", new IMAGE_BRUSH_SVG("Starship/MainToolbar/pause", Icon20x20));
-		Set("Icons.Toolbar.Stop", new IMAGE_BRUSH_SVG("Starship/MainToolbar/stop", Icon20x20));
+		Set("Icons.Toolbar.Stop", new CORE_IMAGE_BRUSH_SVG("Starship/Common/stop", Icon20x20));
 		Set("Icons.Toolbar.Settings", new CORE_IMAGE_BRUSH_SVG( "Starship/Common/Settings", Icon20x20));
 		Set("Icons.Toolbar.Details", new IMAGE_BRUSH_SVG("Starship/Common/Details", Icon16x16));
 		Set("Icons.Toolbar.Import", new CORE_IMAGE_BRUSH_SVG( "Starship/Common/import", Icon20x20) );
@@ -958,60 +941,7 @@ void FStarshipEditorStyle::FStyle::SetupGeneralStyles()
 
 #endif // WITH_EDITOR || (IS_PROGRAM && WITH_UNREAL_DEVELOPER_TOOLS)
 
-	// Output Log Window
 #if WITH_EDITOR || (IS_PROGRAM && WITH_UNREAL_DEVELOPER_TOOLS)
-	{
-		const int32 LogFontSize = Settings.IsValid() ? Settings->LogFontSize : 9;
-
-		const FTextBlockStyle NormalLogText = FTextBlockStyle(NormalText)
-			.SetFont(DEFAULT_FONT("Mono", LogFontSize))
-			.SetColorAndOpacity(LogColor_Normal)
-			.SetSelectedBackgroundColor(LogColor_SelectionBackground)
-			.SetHighlightColor(FStyleColors::Black);
-
-		Set("Log.Normal", NormalLogText );
-
-		Set("Log.Command", FTextBlockStyle(NormalLogText)
-			.SetColorAndOpacity( LogColor_Command )
-			);
-
-		Set("Log.Warning", FTextBlockStyle(NormalLogText)
-			.SetColorAndOpacity(FStyleColors::Warning)
-			);
-
-		Set("Log.Error", FTextBlockStyle(NormalLogText)
-			.SetColorAndOpacity(FStyleColors::Error)
-			);
-
-		Set("Log.TabIcon", new IMAGE_BRUSH_SVG( "Starship/Common/OutputLog", Icon16x16 ) );
-
-		Set("Log.TextBox", FEditableTextBoxStyle(NormalEditableTextBoxStyle)
-			.SetBackgroundImageNormal( BOX_BRUSH( "Common/WhiteGroupBorder", FMargin(4.0f/16.0f) ) )
-			.SetBackgroundImageHovered( BOX_BRUSH( "Common/WhiteGroupBorder", FMargin(4.0f/16.0f) ) )
-			.SetBackgroundImageFocused( BOX_BRUSH( "Common/WhiteGroupBorder", FMargin(4.0f/16.0f) ) )
-			.SetBackgroundImageReadOnly( BOX_BRUSH( "Common/WhiteGroupBorder", FMargin(4.0f/16.0f) ) )
-			.SetBackgroundColor(FStyleColors::Recessed)
-			);
-
-		Set("DebugConsole.Background", new FSlateNoResource());
-
-		const FButtonStyle DebugConsoleButton = FButtonStyle(FStarshipCoreStyle::GetCoreStyle().GetWidgetStyle<FButtonStyle>("NoBorder"))
-			.SetNormalForeground(FStyleColors::Foreground)
-			.SetNormalPadding(FMargin(2, 2, 2, 2))
-			.SetPressedPadding(FMargin(2, 3, 2, 1));
-
-		const FComboButtonStyle DebugConsoleComboButton = FComboButtonStyle(FStarshipCoreStyle::GetCoreStyle().GetWidgetStyle<FComboButtonStyle>("ComboButton"))
-			.SetDownArrowImage(FSlateNoResource())
-			.SetButtonStyle(DebugConsoleButton);
-
-		Set("DebugConsole.ComboButton", DebugConsoleComboButton);
-
-		Set("DebugConsole.Icon", new IMAGE_BRUSH_SVG("Starship/Common/Console", Icon16x16));
-
-		Set("OutputLog.OpenSourceLocation", new IMAGE_BRUSH("Icons/icon_Asset_Open_Source_Location_16x", Icon16x16));
-		Set("OutputLog.OpenInExternalEditor", new IMAGE_BRUSH("Icons/icon_Asset_Open_In_External_Editor_16x", Icon16x16));
-
-	}
 	// Debugging tools 
 	{
 		Set("PerfTools.TabIcon", new IMAGE_BRUSH( "Icons/icon_tab_PerfTools_16x", Icon16x16 ) );
@@ -1066,12 +996,6 @@ void FStarshipEditorStyle::FStyle::SetupGeneralStyles()
 		Set("SessionFrontEnd.Tabs.Tools", new IMAGE_BRUSH( "/Icons/icon_tab_Tools_16x", Icon16x16 ) );
 	}
 
-	// Launcher Window
-	{
-		Set("Launcher.TabIcon", new IMAGE_BRUSH_SVG( "Starship/Common/ProjectLauncher", Icon16x16 ) );
-		Set("Launcher.Tabs.Tools", new IMAGE_BRUSH( "/Icons/icon_tab_Tools_16x", Icon16x16 ) );
-	}
-
 	// Undo History Window
 	{
 		Set("UndoHistory.TabIcon", new CORE_IMAGE_BRUSH_SVG( "Starship/Common/UndoHistory", Icon16x16 ) );
@@ -1375,11 +1299,6 @@ void FStarshipEditorStyle::FStyle::SetupGeneralStyles()
 	Set( "Editor.SearchBoxFont", DEFAULT_FONT( "Regular", 12) );
 #endif // WITH_EDITOR || (IS_PROGRAM && WITH_UNREAL_DEVELOPER_TOOLS)
 
-	// Console
-	{
-		Set( "DebugConsole.Background", new BOX_BRUSH( "Old/Menu_Background", FMargin(8.0f/64.0f) ) );
-	}
-
 #if WITH_EDITOR || (IS_PROGRAM && WITH_UNREAL_DEVELOPER_TOOLS)
 	// About screen
 	if (IncludeEditorSpecificStyles())
@@ -3072,55 +2991,55 @@ void FStarshipEditorStyle::FStyle::SetupProfilerStyle()
 	// Profiler
 	{
 		// Profiler group brushes
-		Set( "Profiler.Group.16", new BOX_BRUSH( "Icons/Profiler/GroupBorder-16Gray", FMargin(4.0f/16.0f) ) );
+		Set( "Profiler.Group.16", new CORE_BOX_BRUSH( "Icons/Profiler/GroupBorder-16Gray", FMargin(4.0f/16.0f) ) );
 
 		// Profiler toolbar icons
-		Set( "Profiler.Tab", new IMAGE_BRUSH_SVG( "Starship/Common/Visualizer", Icon16x16 ) );
-		Set( "Profiler.Tab.GraphView", new IMAGE_BRUSH( "Icons/Profiler/Profiler_Graph_View_Tab_16x", Icon16x16 ) );
-		Set( "Profiler.Tab.EventGraph", new IMAGE_BRUSH( "Icons/Profiler/profiler_OpenEventGraph_32x", Icon16x16 ) );
-		Set( "Profiler.Tab.FiltersAndPresets", new IMAGE_BRUSH( "Icons/Profiler/Profiler_Filter_Presets_Tab_16x", Icon16x16 ) );
+		Set( "Profiler.Tab", new CORE_IMAGE_BRUSH_SVG( "Starship/Common/Visualizer", Icon16x16 ) );
+		Set( "Profiler.Tab.GraphView", new CORE_IMAGE_BRUSH( "Icons/Profiler/Profiler_Graph_View_Tab_16x", Icon16x16 ) );
+		Set( "Profiler.Tab.EventGraph", new CORE_IMAGE_BRUSH( "Icons/Profiler/profiler_OpenEventGraph_32x", Icon16x16 ) );
+		Set( "Profiler.Tab.FiltersAndPresets", new CORE_IMAGE_BRUSH( "Icons/Profiler/Profiler_Filter_Presets_Tab_16x", Icon16x16 ) );
 
-		Set( "ProfilerCommand.ProfilerManager_Load", new IMAGE_BRUSH( "Icons/Profiler/Profiler_Load_Profiler_40x", Icon40x40 ) );
-		Set( "ProfilerCommand.ProfilerManager_Load.Small", new IMAGE_BRUSH( "Icons/Profiler/Profiler_Load_Profiler_40x", Icon20x20 ) );
+		Set( "ProfilerCommand.ProfilerManager_Load", new CORE_IMAGE_BRUSH( "Icons/Profiler/Profiler_Load_Profiler_40x", Icon40x40 ) );
+		Set( "ProfilerCommand.ProfilerManager_Load.Small", new CORE_IMAGE_BRUSH( "Icons/Profiler/Profiler_Load_Profiler_40x", Icon20x20 ) );
 
-		Set("ProfilerCommand.ProfilerManager_LoadMultiple", new IMAGE_BRUSH("Icons/Profiler/Profiler_LoadMultiple_Profiler_40x", Icon40x40));
-		Set("ProfilerCommand.ProfilerManager_LoadMultiple.Small", new IMAGE_BRUSH("Icons/Profiler/Profiler_LoadMultiple_Profiler_40x", Icon20x20));
+		Set("ProfilerCommand.ProfilerManager_LoadMultiple", new CORE_IMAGE_BRUSH("Icons/Profiler/Profiler_LoadMultiple_Profiler_40x", Icon40x40));
+		Set("ProfilerCommand.ProfilerManager_LoadMultiple.Small", new CORE_IMAGE_BRUSH("Icons/Profiler/Profiler_LoadMultiple_Profiler_40x", Icon20x20));
 
-		Set( "ProfilerCommand.ProfilerManager_Save", new IMAGE_BRUSH( "Icons/LV_Save", Icon40x40 ) );
-		Set( "ProfilerCommand.ProfilerManager_Save.Small", new IMAGE_BRUSH( "Icons/LV_Save", Icon20x20 ) );
+		Set( "ProfilerCommand.ProfilerManager_Save", new IMAGE_BRUSH_SVG("Starship/Common/save", Icon16x16));
+		Set("ProfilerCommand.ProfilerManager_Save.Small", new IMAGE_BRUSH_SVG("Starship/Common/save", Icon16x16));
 		
-		Set( "ProfilerCommand.ProfilerManager_ToggleLivePreview", new IMAGE_BRUSH( "Automation/RefreshTests", Icon40x40) );
-		Set( "ProfilerCommand.ProfilerManager_ToggleLivePreview.Small", new IMAGE_BRUSH( "Automation/RefreshTests", Icon20x20) );
+		Set( "ProfilerCommand.ProfilerManager_ToggleLivePreview", new CORE_IMAGE_BRUSH( "Automation/RefreshTests", Icon40x40) );
+		Set( "ProfilerCommand.ProfilerManager_ToggleLivePreview.Small", new CORE_IMAGE_BRUSH( "Automation/RefreshTests", Icon20x20) );
 
-		Set( "ProfilerCommand.StatsProfiler", new IMAGE_BRUSH( "Icons/Profiler/profiler_stats_40x", Icon40x40 ) );
-		Set( "ProfilerCommand.StatsProfiler.Small", new IMAGE_BRUSH( "Icons/Profiler/profiler_stats_40x", Icon20x20 ) );
+		Set( "ProfilerCommand.StatsProfiler", new CORE_IMAGE_BRUSH( "Icons/Profiler/profiler_stats_40x", Icon40x40 ) );
+		Set( "ProfilerCommand.StatsProfiler.Small", new CORE_IMAGE_BRUSH( "Icons/Profiler/profiler_stats_40x", Icon20x20 ) );
 
-		Set( "ProfilerCommand.MemoryProfiler", new IMAGE_BRUSH( "Icons/Profiler/profiler_mem_40x", Icon40x40 ) );
-		Set( "ProfilerCommand.MemoryProfiler.Small", new IMAGE_BRUSH( "Icons/Profiler/profiler_mem_40x", Icon20x20 ) );
+		Set( "ProfilerCommand.MemoryProfiler", new CORE_IMAGE_BRUSH( "Icons/Profiler/profiler_mem_40x", Icon40x40 ) );
+		Set( "ProfilerCommand.MemoryProfiler.Small", new CORE_IMAGE_BRUSH( "Icons/Profiler/profiler_mem_40x", Icon20x20 ) );
 
-		Set( "ProfilerCommand.FPSChart", new IMAGE_BRUSH( "Icons/Profiler/Profiler_FPS_Chart_40x", Icon40x40 ) );
-		Set( "ProfilerCommand.FPSChart.Small", new IMAGE_BRUSH( "Icons/Profiler/Profiler_FPS_Chart_40x", Icon20x20 ) );
+		Set( "ProfilerCommand.FPSChart", new CORE_IMAGE_BRUSH( "Icons/Profiler/Profiler_FPS_Chart_40x", Icon40x40 ) );
+		Set( "ProfilerCommand.FPSChart.Small", new CORE_IMAGE_BRUSH( "Icons/Profiler/Profiler_FPS_Chart_40x", Icon20x20 ) );
 
-		Set( "ProfilerCommand.OpenSettings", new IMAGE_BRUSH( "Icons/Profiler/Profiler_Settings_40x", Icon40x40 ) );
-		Set( "ProfilerCommand.OpenSettings.Small", new IMAGE_BRUSH( "Icons/Profiler/Profiler_Settings_40x", Icon20x20 ) );
+		Set( "ProfilerCommand.OpenSettings", new CORE_IMAGE_BRUSH( "Icons/Profiler/Profiler_Settings_40x", Icon40x40 ) );
+		Set( "ProfilerCommand.OpenSettings.Small", new CORE_IMAGE_BRUSH( "Icons/Profiler/Profiler_Settings_40x", Icon20x20 ) );
 
-		Set( "ProfilerCommand.ToggleDataPreview", new IMAGE_BRUSH( "Icons/Profiler/profiler_sync_40x", Icon40x40 ) );
-		Set( "ProfilerCommand.ToggleDataPreview.Small", new IMAGE_BRUSH( "Icons/Profiler/profiler_sync_40x", Icon20x20 ) );
+		Set( "ProfilerCommand.ToggleDataPreview", new CORE_IMAGE_BRUSH( "Icons/Profiler/profiler_sync_40x", Icon40x40 ) );
+		Set( "ProfilerCommand.ToggleDataPreview.Small", new CORE_IMAGE_BRUSH( "Icons/Profiler/profiler_sync_40x", Icon20x20 ) );
 
-		Set( "ProfilerCommand.ToggleDataCapture", new IMAGE_BRUSH( "Icons/Profiler/Profiler_Data_Capture_40x", Icon40x40 ) );
-		Set( "ProfilerCommand.ToggleDataCapture.Small", new IMAGE_BRUSH( "Icons/Profiler/Profiler_Data_Capture_40x", Icon20x20 ) );
+		Set( "ProfilerCommand.ToggleDataCapture", new CORE_IMAGE_BRUSH( "Icons/Profiler/Profiler_Data_Capture_40x", Icon40x40 ) );
+		Set( "ProfilerCommand.ToggleDataCapture.Small", new CORE_IMAGE_BRUSH( "Icons/Profiler/Profiler_Data_Capture_40x", Icon20x20 ) );
 
-		Set( "ProfilerCommand.ToggleDataCapture.Checked", new IMAGE_BRUSH( "Icons/icon_stop_40x", Icon40x40 ) );
-		Set( "ProfilerCommand.ToggleDataCapture.Checked.Small", new IMAGE_BRUSH( "Icons/icon_stop_40x", Icon20x20 ) );
+		Set( "ProfilerCommand.ToggleDataCapture.Checked", new CORE_IMAGE_BRUSH( "Icons/icon_stop_40x", Icon40x40 ) );
+		Set( "ProfilerCommand.ToggleDataCapture.Checked.Small", new CORE_IMAGE_BRUSH( "Icons/icon_stop_40x", Icon20x20 ) );
 
-		Set( "ProfilerCommand.ToggleShowDataGraph", new IMAGE_BRUSH( "Icons/Profiler/profiler_ShowGraphData_32x", Icon32x32 ) );
-		Set( "ProfilerCommand.OpenEventGraph", new IMAGE_BRUSH( "Icons/Profiler/profiler_OpenEventGraph_32x", Icon16x16 ) );
+		Set( "ProfilerCommand.ToggleShowDataGraph", new CORE_IMAGE_BRUSH( "Icons/Profiler/profiler_ShowGraphData_32x", Icon32x32 ) );
+		Set( "ProfilerCommand.OpenEventGraph", new CORE_IMAGE_BRUSH( "Icons/Profiler/profiler_OpenEventGraph_32x", Icon16x16 ) );
 
 		// Generic
-		Set( "Profiler.LineGraphArea", new IMAGE_BRUSH( "Old/White", Icon16x16, FLinearColor(1.0f,1.0f,1.0f,0.25f) ) );
+		Set( "Profiler.LineGraphArea", new CORE_IMAGE_BRUSH( "Old/White", Icon16x16, FLinearColor(1.0f,1.0f,1.0f,0.25f) ) );
 		
 		// Tooltip hint icon
-		Set( "Profiler.Tooltip.HintIcon10", new IMAGE_BRUSH( "Icons/Profiler/Profiler_Custom_Tooltip_12x", Icon12x12 ) );
+		Set( "Profiler.Tooltip.HintIcon10", new CORE_IMAGE_BRUSH( "Icons/Profiler/Profiler_Custom_Tooltip_12x", Icon12x12 ) );
 
 		// Text styles
 		Set( "Profiler.CaptionBold", FTextBlockStyle(NormalText)
@@ -3130,13 +3049,6 @@ void FStarshipEditorStyle::FStyle::SetupProfilerStyle()
 			.SetShadowColorAndOpacity( FLinearColor(0.f,0.f,0.f,0.8f) )
 		);
 
-		Set( "Profiler.Caption", FTextBlockStyle(NormalText)
-			.SetFont( DEFAULT_FONT( "Regular", 10 ) )
-			.SetColorAndOpacity( FLinearColor::White )
-			.SetShadowOffset( FVector2D(1.0f, 1.0f) )
-			.SetShadowColorAndOpacity( FLinearColor(0.f,0.f,0.f,0.8f) )
-		);
-
 		Set( "Profiler.TooltipBold", FTextBlockStyle(NormalText)
 			.SetFont( DEFAULT_FONT( "Bold", 8 ) )
 			.SetColorAndOpacity( FLinearColor(0.5f,0.5f,0.5f,1.0f) )
@@ -3152,45 +3064,45 @@ void FStarshipEditorStyle::FStyle::SetupProfilerStyle()
 		);
 
 		// Event graph icons
-		Set( "Profiler.EventGraph.SetRoot", new IMAGE_BRUSH( "Icons/Profiler/profiler_SetRoot_32x", Icon32x32 ) );
-		Set( "Profiler.EventGraph.CullEvents", new IMAGE_BRUSH( "Icons/Profiler/Profiler_Cull_Events_16x", Icon16x16) );
-		Set( "Profiler.EventGraph.FilterEvents", new IMAGE_BRUSH( "Icons/Profiler/Profiler_Filter_Events_16x", Icon16x16) );
+		Set( "Profiler.EventGraph.SetRoot", new CORE_IMAGE_BRUSH( "Icons/Profiler/profiler_SetRoot_32x", Icon32x32 ) );
+		Set( "Profiler.EventGraph.CullEvents", new CORE_IMAGE_BRUSH( "Icons/Profiler/Profiler_Cull_Events_16x", Icon16x16) );
+		Set( "Profiler.EventGraph.FilterEvents", new CORE_IMAGE_BRUSH( "Icons/Profiler/Profiler_Filter_Events_16x", Icon16x16) );
 
-		Set( "Profiler.EventGraph.SelectStack", new IMAGE_BRUSH( "Icons/Profiler/profiler_SelectStack_32x", Icon32x32 ) );
+		Set( "Profiler.EventGraph.SelectStack", new CORE_IMAGE_BRUSH( "Icons/Profiler/profiler_SelectStack_32x", Icon32x32 ) );
 
-		Set( "Profiler.EventGraph.ExpandAll", new IMAGE_BRUSH( "Icons/Profiler/profiler_ExpandAll_32x", Icon32x32 ) );
-		Set( "Profiler.EventGraph.CollapseAll", new IMAGE_BRUSH( "Icons/Profiler/profiler_CollapseAll_32x", Icon32x32 ) );
+		Set( "Profiler.EventGraph.ExpandAll", new CORE_IMAGE_BRUSH( "Icons/Profiler/profiler_ExpandAll_32x", Icon32x32 ) );
+		Set( "Profiler.EventGraph.CollapseAll", new CORE_IMAGE_BRUSH( "Icons/Profiler/profiler_CollapseAll_32x", Icon32x32 ) );
 		
-		Set( "Profiler.EventGraph.ExpandSelection", new IMAGE_BRUSH( "Icons/Profiler/profiler_ExpandSelection_32x", Icon32x32 ) );
-		Set( "Profiler.EventGraph.CollapseSelection", new IMAGE_BRUSH( "Icons/Profiler/profiler_CollapseSelection_32x", Icon32x32 ) );
+		Set( "Profiler.EventGraph.ExpandSelection", new CORE_IMAGE_BRUSH( "Icons/Profiler/profiler_ExpandSelection_32x", Icon32x32 ) );
+		Set( "Profiler.EventGraph.CollapseSelection", new CORE_IMAGE_BRUSH( "Icons/Profiler/profiler_CollapseSelection_32x", Icon32x32 ) );
 
-		Set( "Profiler.EventGraph.ExpandThread", new IMAGE_BRUSH( "Icons/Profiler/profiler_ExpandThread_32x", Icon32x32 ) );
-		Set( "Profiler.EventGraph.CollapseThread", new IMAGE_BRUSH( "Icons/Profiler/profiler_CollapseThread_32x", Icon32x32 ) );
+		Set( "Profiler.EventGraph.ExpandThread", new CORE_IMAGE_BRUSH( "Icons/Profiler/profiler_ExpandThread_32x", Icon32x32 ) );
+		Set( "Profiler.EventGraph.CollapseThread", new CORE_IMAGE_BRUSH( "Icons/Profiler/profiler_CollapseThread_32x", Icon32x32 ) );
 
-		Set( "Profiler.EventGraph.ExpandHotPath", new IMAGE_BRUSH( "Icons/Profiler/profiler_ExpandHotPath_32x", Icon32x32 ) );
-		Set( "Profiler.EventGraph.HotPathSmall", new IMAGE_BRUSH( "Icons/Profiler/profiler_HotPath_32x", Icon12x12 ) );
+		Set( "Profiler.EventGraph.ExpandHotPath", new CORE_IMAGE_BRUSH( "Icons/Profiler/profiler_ExpandHotPath_32x", Icon32x32 ) );
+		Set( "Profiler.EventGraph.HotPathSmall", new CORE_IMAGE_BRUSH( "Icons/Profiler/profiler_HotPath_32x", Icon12x12 ) );
 
-		Set( "Profiler.EventGraph.ExpandHotPath16", new IMAGE_BRUSH( "Icons/Profiler/profiler_HotPath_32x", Icon16x16 ) );
+		Set( "Profiler.EventGraph.ExpandHotPath16", new CORE_IMAGE_BRUSH( "Icons/Profiler/profiler_HotPath_32x", Icon16x16 ) );
 
-		Set( "Profiler.EventGraph.GameThread", new IMAGE_BRUSH( "Icons/Profiler/profiler_GameThread_32x", Icon32x32 ) );
-		Set( "Profiler.EventGraph.RenderThread", new IMAGE_BRUSH( "Icons/Profiler/profiler_RenderThread_32x", Icon32x32 ) );
+		Set( "Profiler.EventGraph.GameThread", new CORE_IMAGE_BRUSH( "Icons/Profiler/profiler_GameThread_32x", Icon32x32 ) );
+		Set( "Profiler.EventGraph.RenderThread", new CORE_IMAGE_BRUSH( "Icons/Profiler/profiler_RenderThread_32x", Icon32x32 ) );
 	
-		Set( "Profiler.EventGraph.ViewColumn", new IMAGE_BRUSH( "Icons/Profiler/profiler_ViewColumn_32x", Icon32x32 ) );
-		Set( "Profiler.EventGraph.ResetColumn", new IMAGE_BRUSH( "Icons/Profiler/profiler_ResetColumn_32x", Icon32x32 ) );
+		Set( "Profiler.EventGraph.ViewColumn", new CORE_IMAGE_BRUSH( "Icons/Profiler/profiler_ViewColumn_32x", Icon32x32 ) );
+		Set( "Profiler.EventGraph.ResetColumn", new CORE_IMAGE_BRUSH( "Icons/Profiler/profiler_ResetColumn_32x", Icon32x32 ) );
 
-		Set( "Profiler.EventGraph.HistoryBack", new IMAGE_BRUSH( "Icons/Profiler/Profiler_History_Back_16x", Icon16x16) );
-		Set( "Profiler.EventGraph.HistoryForward", new IMAGE_BRUSH( "Icons/Profiler/Profiler_History_Fwd_16x", Icon16x16) );
+		Set( "Profiler.EventGraph.HistoryBack", new CORE_IMAGE_BRUSH( "Icons/Profiler/Profiler_History_Back_16x", Icon16x16) );
+		Set( "Profiler.EventGraph.HistoryForward", new CORE_IMAGE_BRUSH( "Icons/Profiler/Profiler_History_Fwd_16x", Icon16x16) );
 
-		Set( "Profiler.EventGraph.MaximumIcon", new IMAGE_BRUSH( "Icons/Profiler/Profiler_Max_Event_Graph_16x", Icon16x16) );
-		Set( "Profiler.EventGraph.AverageIcon", new IMAGE_BRUSH( "Icons/Profiler/Profiler_Average_Event_Graph_16x", Icon16x16) );
+		Set( "Profiler.EventGraph.MaximumIcon", new CORE_IMAGE_BRUSH( "Icons/Profiler/Profiler_Max_Event_Graph_16x", Icon16x16) );
+		Set( "Profiler.EventGraph.AverageIcon", new CORE_IMAGE_BRUSH( "Icons/Profiler/Profiler_Average_Event_Graph_16x", Icon16x16) );
 
-		Set( "Profiler.EventGraph.FlatIcon", new IMAGE_BRUSH( "Icons/Profiler/Profiler_Events_Flat_16x", Icon16x16) );
-		Set( "Profiler.EventGraph.FlatCoalescedIcon", new IMAGE_BRUSH( "Icons/Profiler/Profiler_Events_Flat_Coalesced_16x", Icon16x16) );
-		Set( "Profiler.EventGraph.HierarchicalIcon", new IMAGE_BRUSH( "Icons/Profiler/Profiler_Events_Hierarchial_16x", Icon16x16) );
+		Set( "Profiler.EventGraph.FlatIcon", new CORE_IMAGE_BRUSH( "Icons/Profiler/Profiler_Events_Flat_16x", Icon16x16) );
+		Set( "Profiler.EventGraph.FlatCoalescedIcon", new CORE_IMAGE_BRUSH( "Icons/Profiler/Profiler_Events_Flat_Coalesced_16x", Icon16x16) );
+		Set( "Profiler.EventGraph.HierarchicalIcon", new CORE_IMAGE_BRUSH( "Icons/Profiler/Profiler_Events_Hierarchial_16x", Icon16x16) );
 
-		Set( "Profiler.EventGraph.HasCulledEventsSmall", new IMAGE_BRUSH( "Icons/Profiler/Profiler_Has_Culled_Children_12x", Icon12x12) );
-		Set( "Profiler.EventGraph.CulledEvent", new IMAGE_BRUSH( "Icons/Profiler/Profiler_Culled_12x", Icon12x12) );
-		Set( "Profiler.EventGraph.FilteredEvent", new IMAGE_BRUSH( "Icons/Profiler/Profiler_Filtered_12x", Icon12x12) );
+		Set( "Profiler.EventGraph.HasCulledEventsSmall", new CORE_IMAGE_BRUSH( "Icons/Profiler/Profiler_Has_Culled_Children_12x", Icon12x12) );
+		Set( "Profiler.EventGraph.CulledEvent", new CORE_IMAGE_BRUSH( "Icons/Profiler/Profiler_Culled_12x", Icon12x12) );
+		Set( "Profiler.EventGraph.FilteredEvent", new CORE_IMAGE_BRUSH( "Icons/Profiler/Profiler_Filtered_12x", Icon12x12) );
 
 		Set( "Profiler.EventGraph.DarkText", FTextBlockStyle(NormalText)
 			.SetFont( DEFAULT_FONT( "Regular", 8 ) )
@@ -3199,41 +3111,39 @@ void FStarshipEditorStyle::FStyle::SetupProfilerStyle()
 			);
 
 		// Thread-view
-		Set( "Profiler.ThreadView.SampleBorder", new BOX_BRUSH( "Icons/Profiler/Profiler_ThreadView_SampleBorder_16x", FMargin( 2.0f / 16.0f ) ) );
+		Set( "Profiler.ThreadView.SampleBorder", new CORE_BOX_BRUSH( "Icons/Profiler/Profiler_ThreadView_SampleBorder_16x", FMargin( 2.0f / 16.0f ) ) );
 
 		// Event graph selected event border
-		Set( "Profiler.EventGraph.Border.TB", new BOX_BRUSH( "Icons/Profiler/Profiler_Border_TB_16x", FMargin(4.0f/16.0f) ) );
-		Set( "Profiler.EventGraph.Border.L", new BOX_BRUSH( "Icons/Profiler/Profiler_Border_L_16x",   FMargin(4.0f/16.0f) ) );
-		Set( "Profiler.EventGraph.Border.R", new BOX_BRUSH( "Icons/Profiler/Profiler_Border_R_16x",   FMargin(4.0f/16.0f) ) );
+		Set( "Profiler.EventGraph.Border.TB", new CORE_BOX_BRUSH( "Icons/Profiler/Profiler_Border_TB_16x", FMargin(4.0f/16.0f) ) );
+		Set( "Profiler.EventGraph.Border.L", new CORE_BOX_BRUSH( "Icons/Profiler/Profiler_Border_L_16x",   FMargin(4.0f/16.0f) ) );
+		Set( "Profiler.EventGraph.Border.R", new CORE_BOX_BRUSH( "Icons/Profiler/Profiler_Border_R_16x",   FMargin(4.0f/16.0f) ) );
 
 		// Misc
-		Set( "Profiler.Misc.WarningSmall", new IMAGE_BRUSH( "ContentBrowser/SCC_NotAtHeadRevision", Icon12x12 ) );
 
-		Set( "Profiler.Misc.SortBy", new IMAGE_BRUSH( "Icons/Profiler/profiler_SortBy_32x", Icon32x32 ) );
-		Set( "Profiler.Misc.SortAscending", new IMAGE_BRUSH( "Icons/Profiler/profiler_SortAscending_32x", Icon32x32 ) );
-		Set( "Profiler.Misc.SortDescending", new IMAGE_BRUSH( "Icons/Profiler/profiler_SortDescending_32x", Icon32x32 ) );
+		Set( "Profiler.Misc.SortBy", new CORE_IMAGE_BRUSH( "Icons/Profiler/profiler_SortBy_32x", Icon32x32 ) );
+		Set( "Profiler.Misc.SortAscending", new CORE_IMAGE_BRUSH( "Icons/Profiler/profiler_SortAscending_32x", Icon32x32 ) );
+		Set( "Profiler.Misc.SortDescending", new CORE_IMAGE_BRUSH( "Icons/Profiler/profiler_SortDescending_32x", Icon32x32 ) );
 
-		Set( "Profiler.Misc.ResetToDefault", new IMAGE_BRUSH( "Icons/Profiler/profiler_ResetToDefault_32x", Icon32x32 ) );
+		Set( "Profiler.Misc.ResetToDefault", new CORE_IMAGE_BRUSH( "Icons/Profiler/profiler_ResetToDefault_32x", Icon32x32 ) );
 
-		Set( "Profiler.Misc.Save16", new IMAGE_BRUSH( "Icons/LV_Save", Icon16x16 ) );
-		Set( "Profiler.Misc.Reset16", new IMAGE_BRUSH( "Icons/Profiler/profiler_ResetToDefault_32x", Icon16x16 ) );
+		Set( "Profiler.Misc.Reset16", new CORE_IMAGE_BRUSH( "Icons/Profiler/profiler_ResetToDefault_32x", Icon16x16 ) );
 
-		Set( "Profiler.Type.Calls", new IMAGE_BRUSH( "Icons/Profiler/profiler_Calls_32x", Icon16x16 ) );
-		Set( "Profiler.Type.Event", new IMAGE_BRUSH( "Icons/Profiler/profiler_Event_32x", Icon16x16 ) );
-		Set( "Profiler.Type.Memory", new IMAGE_BRUSH( "Icons/Profiler/profiler_Memory_32x", Icon16x16 ) );
-		Set( "Profiler.Type.Number", new IMAGE_BRUSH( "Icons/Profiler/profiler_Number_32x", Icon16x16 ) );
+		Set( "Profiler.Type.Calls", new CORE_IMAGE_BRUSH( "Icons/Profiler/profiler_Calls_32x", Icon16x16 ) );
+		Set( "Profiler.Type.Event", new CORE_IMAGE_BRUSH( "Icons/Profiler/profiler_Event_32x", Icon16x16 ) );
+		Set( "Profiler.Type.Memory", new CORE_IMAGE_BRUSH( "Icons/Profiler/profiler_Memory_32x", Icon16x16 ) );
+		Set( "Profiler.Type.Number", new CORE_IMAGE_BRUSH( "Icons/Profiler/profiler_Number_32x", Icon16x16 ) );
 
 		// NumberInt, NumberFloat, Memory, Hierarchical
-		Set( "Profiler.Type.NumberInt", new IMAGE_BRUSH( "Icons/Profiler/profiler_Number_32x", Icon16x16 ) );
-		Set( "Profiler.Type.NumberFloat", new IMAGE_BRUSH( "Icons/Profiler/profiler_Number_32x", Icon16x16 ) );
-		Set( "Profiler.Type.Memory", new IMAGE_BRUSH( "Icons/Profiler/profiler_Memory_32x", Icon16x16 ) );
-		Set( "Profiler.Type.Hierarchical", new IMAGE_BRUSH( "Icons/Profiler/profiler_Event_32x", Icon16x16 ) );
-
-		Set( "Profiler.Misc.GenericFilter", new IMAGE_BRUSH( "Icons/Profiler/profiler_GenericFilter_32x", Icon16x16 ) );
-		Set( "Profiler.Misc.GenericGroup", new IMAGE_BRUSH( "Icons/Profiler/profiler_GenericGroup_32x", Icon16x16 ) );
-		Set( "Profiler.Misc.CopyToClipboard", new IMAGE_BRUSH( "Icons/Profiler/profiler_CopyToClipboard_32x", Icon32x32 ) );
+		Set( "Profiler.Type.NumberInt", new CORE_IMAGE_BRUSH( "Icons/Profiler/profiler_Number_32x", Icon16x16 ) );
+		Set( "Profiler.Type.NumberFloat", new CORE_IMAGE_BRUSH( "Icons/Profiler/profiler_Number_32x", Icon16x16 ) );
+		Set( "Profiler.Type.Memory", new CORE_IMAGE_BRUSH( "Icons/Profiler/profiler_Memory_32x", Icon16x16 ) );
+		Set( "Profiler.Type.Hierarchical", new CORE_IMAGE_BRUSH( "Icons/Profiler/profiler_Event_32x", Icon16x16 ) );
+
+		Set( "Profiler.Misc.GenericFilter", new CORE_IMAGE_BRUSH( "Icons/Profiler/profiler_GenericFilter_32x", Icon16x16 ) );
+		Set( "Profiler.Misc.GenericGroup", new CORE_IMAGE_BRUSH( "Icons/Profiler/profiler_GenericGroup_32x", Icon16x16 ) );
+		Set( "Profiler.Misc.CopyToClipboard", new CORE_IMAGE_BRUSH( "Icons/Profiler/profiler_CopyToClipboard_32x", Icon32x32 ) );
 	
-		Set( "Profiler.Misc.Disconnect", new IMAGE_BRUSH( "Icons/Profiler/profiler_Disconnect_32x", Icon32x32 ) );
+		Set( "Profiler.Misc.Disconnect", new CORE_IMAGE_BRUSH( "Icons/Profiler/profiler_Disconnect_32x", Icon32x32 ) );
 
 		//Set( "Profiler.Type.Calls", new IMAGE_BRUSH( "Icons/Profiler/profiler_Calls_32x", Icon40x40) );
 		//Set( "Profiler.Type.Calls.Small", new IMAGE_BRUSH( "Icons/Profiler/profiler_Calls_32x", Icon20x20) );
@@ -4063,30 +3973,30 @@ void FStarshipEditorStyle::FStyle::SetupLevelEditorStyle()
 	// Level editor tool bar icons
 #if WITH_EDITOR
 	{
-		Set("LevelEditor.BrowseDocumentation",      new IMAGE_BRUSH_SVG("Starship/Common/Documentation", Icon16x16));
-		Set("LevelEditor.Tutorials",                new IMAGE_BRUSH_SVG("Starship/Common/Tutorials", Icon16x16));
-		Set("LevelEditor.BrowseViewportControls",   new IMAGE_BRUSH_SVG("Starship/Common/ViewportControls", Icon16x16));
+		Set("LevelEditor.BrowseDocumentation",     new IMAGE_BRUSH_SVG("Starship/Common/Documentation", Icon16x16));
+		Set("LevelEditor.Tutorials",               new IMAGE_BRUSH_SVG("Starship/Common/Tutorials", Icon16x16));
+		Set("LevelEditor.BrowseViewportControls",  new IMAGE_BRUSH_SVG("Starship/Common/ViewportControls", Icon16x16));
 		Set("LevelEditor.PasteHere",				new IMAGE_BRUSH_SVG("Starship/Actors/paste-here", Icon16x16));
 
-		Set("MainFrame.ToggleFullscreen",           new IMAGE_BRUSH_SVG("Starship/Common/EnableFullscreen", Icon16x16));
-		Set("MainFrame.LoadLayout",                 new IMAGE_BRUSH_SVG("Starship/Common/LayoutLoad", Icon16x16));
-		Set("MainFrame.SaveLayout",                 new IMAGE_BRUSH_SVG("Starship/Common/LayoutSave", Icon16x16));
-		Set("MainFrame.RemoveLayout",               new IMAGE_BRUSH_SVG("Starship/Common/LayoutRemove", Icon16x16));
+		Set("MainFrame.ToggleFullscreen",          new IMAGE_BRUSH_SVG("Starship/Common/EnableFullscreen", Icon16x16));
+		Set("MainFrame.LoadLayout",                new IMAGE_BRUSH_SVG("Starship/Common/LayoutLoad", Icon16x16));
+		Set("MainFrame.SaveLayout",                new IMAGE_BRUSH_SVG("Starship/Common/LayoutSave", Icon16x16));
+		Set("MainFrame.RemoveLayout",              new IMAGE_BRUSH_SVG("Starship/Common/LayoutRemove", Icon16x16));
 
-		Set("MainFrame.OpenIssueTracker",           new IMAGE_BRUSH_SVG("Starship/Common/IssueTracker", Icon16x16));
-		Set("MainFrame.ReportABug",                 new IMAGE_BRUSH_SVG("Starship/Common/Bug", Icon16x16));
+		Set("MainFrame.OpenIssueTracker",          new IMAGE_BRUSH_SVG("Starship/Common/IssueTracker", Icon16x16));
+		Set("MainFrame.ReportABug",                new IMAGE_BRUSH_SVG("Starship/Common/Bug", Icon16x16));
 
 		Set("SystemWideCommands.OpenDocumentation", new IMAGE_BRUSH_SVG("Starship/Common/Documentation", Icon16x16));
 		Set("MainFrame.DocumentationHome",	        new IMAGE_BRUSH_SVG("Starship/Common/Documentation", Icon16x16));
 		Set("MainFrame.BrowseAPIReference",         new IMAGE_BRUSH_SVG("Starship/Common/Documentation", Icon16x16));
-		Set("MainFrame.BrowseCVars",                new IMAGE_BRUSH_SVG("Starship/Common/Console", Icon16x16));
+		Set("MainFrame.BrowseCVars",                new CORE_IMAGE_BRUSH_SVG("Starship/Common/Console", Icon16x16));
 		Set("MainFrame.VisitOnlineLearning",		new IMAGE_BRUSH_SVG("Starship/Common/Tutorials", Icon16x16));
 		Set("MainFrame.VisitForums",                new IMAGE_BRUSH_SVG("Starship/Common/Forums", Icon16x16));
 		Set("MainFrame.VisitSearchForAnswersPage",  new IMAGE_BRUSH_SVG("Starship/Common/QuestionAnswer", Icon16x16));
-		Set("MainFrame.VisitSupportWebSite",        new IMAGE_BRUSH_SVG("Starship/Common/Support", Icon16x16));
+		Set("MainFrame.VisitSupportWebSite",       new IMAGE_BRUSH_SVG("Starship/Common/Support", Icon16x16));
 		Set("MainFrame.VisitEpicGamesDotCom",       new IMAGE_BRUSH_SVG("About/EpicGamesLogo", Icon16x16));
 		Set("MainFrame.AboutUnrealEd",              new IMAGE_BRUSH_SVG("About/UnrealLogo", Icon16x16));
-		Set("MainFrame.CreditsUnrealEd",            new IMAGE_BRUSH_SVG("Starship/Common/Credits", Icon16x16));
+		Set("MainFrame.CreditsUnrealEd",           new IMAGE_BRUSH_SVG("Starship/Common/Credits", Icon16x16));
 
 		Set( "EditorViewport.SelectMode", new IMAGE_BRUSH_SVG("Starship/EditorViewport/select", Icon16x16) );
 		Set( "EditorViewport.TranslateMode", new IMAGE_BRUSH_SVG( "Starship/EditorViewport/translate", Icon16x16 ) );
@@ -5405,7 +5315,7 @@ void FStarshipEditorStyle::FStyle::SetupPersonaStyle()
 		Set("BlendSpaceKey.Drag", SelectionColor_Subdued);
 		Set("BlendSpaceKey.Drop", SelectionColor_Inactive);
 		Set("BlendSpaceKey.Invalid", FStyleColors::Warning);
-		Set("BlendSpaceKey.Preview", LogColor_Command);
+		Set("BlendSpaceKey.Preview", FStyleColors::AccentGreen);
 	}
 
 	// Custom menu style for recent commands list
@@ -7222,158 +7132,6 @@ void FStarshipEditorStyle::FStyle::SetupAutomationStyles()
 {
 	//Automation
 #if WITH_EDITOR || (IS_PROGRAM && WITH_UNREAL_DEVELOPER_TOOLS)
-	{
-		Set( "Automation.Header" , FTextBlockStyle(NormalText)
-			.SetFont( DEFAULT_FONT( "Mono", 12 ) )
-			.SetColorAndOpacity(FLinearColor(FColor(0xffffffff))) );
-
-		Set( "Automation.Normal" , FTextBlockStyle(NormalText)
-			.SetFont( DEFAULT_FONT( "Mono", 9 ) )
-			.SetColorAndOpacity(FLinearColor(FColor(0xffaaaaaa))) );
-
-		Set( "Automation.Warning", FTextBlockStyle(NormalText)
-			.SetFont( DEFAULT_FONT( "Mono", 9 ) )
-			.SetColorAndOpacity(FLinearColor(FColor(0xffbbbb44))) );
-
-		Set( "Automation.Error"  , FTextBlockStyle(NormalText)
-			.SetFont( DEFAULT_FONT( "Mono", 9 ) )
-			.SetColorAndOpacity(FLinearColor(FColor(0xffff0000))) );
-
-		Set( "Automation.ReportHeader" , FTextBlockStyle(NormalText)
-			.SetFont( DEFAULT_FONT( "Mono", 10 ) )
-			.SetColorAndOpacity(FLinearColor(FColor(0xffffffff))) );
-		
-		//state of individual tests
-		Set( "Automation.Success", new IMAGE_BRUSH( "Automation/Success", Icon16x16 ) );
-		Set( "Automation.Warning", new IMAGE_BRUSH( "Automation/Warning", Icon16x16 ) );
-		Set( "Automation.Fail", new IMAGE_BRUSH( "Automation/Fail", Icon16x16 ) );
-		Set( "Automation.InProcess", new IMAGE_BRUSH( "Automation/InProcess", Icon16x16 ) );
-		Set( "Automation.NotRun", new IMAGE_BRUSH( "Automation/NotRun", Icon16x16, FLinearColor(0.0f, 0.0f, 0.0f, 0.4f) ) );
-		Set( "Automation.Skipped", new IMAGE_BRUSH( "Automation/NoSessionWarning", Icon16x16 ) );
-		Set( "Automation.ParticipantsWarning", new IMAGE_BRUSH( "Automation/ParticipantsWarning", Icon16x16 ) );
-		Set( "Automation.Participant", new IMAGE_BRUSH( "Automation/Participant", Icon16x16 ) );
-		
-		//status as a regression test or not
-		Set( "Automation.SmokeTest", new IMAGE_BRUSH( "Automation/SmokeTest", Icon16x16 ) );
-		Set( "Automation.SmokeTestParent", new IMAGE_BRUSH( "Automation/SmokeTestParent", Icon16x16 ) );
-
-		//run icons
-		Set( "AutomationWindow.RunTests", new IMAGE_BRUSH( "Automation/RunTests", Icon40x40) );
-		Set( "AutomationWindow.RefreshTests", new IMAGE_BRUSH( "Automation/RefreshTests", Icon40x40) );
-		Set( "AutomationWindow.FindWorkers", new IMAGE_BRUSH( "Automation/RefreshWorkers", Icon40x40) );
-		Set( "AutomationWindow.StopTests", new IMAGE_BRUSH( "Automation/StopTests", Icon40x40 ) );
-		Set( "AutomationWindow.RunTests.Small", new IMAGE_BRUSH( "Automation/RunTests", Icon20x20) );
-		Set( "AutomationWindow.RefreshTests.Small", new IMAGE_BRUSH( "Automation/RefreshTests", Icon20x20) );
-		Set( "AutomationWindow.FindWorkers.Small", new IMAGE_BRUSH( "Automation/RefreshWorkers", Icon20x20) );
-		Set( "AutomationWindow.StopTests.Small", new IMAGE_BRUSH( "Automation/StopTests", Icon20x20 ) );
-
-		//filter icons
-		Set( "AutomationWindow.ErrorFilter", new IMAGE_BRUSH( "Automation/ErrorFilter", Icon40x40) );
-		Set( "AutomationWindow.WarningFilter", new IMAGE_BRUSH( "Automation/WarningFilter", Icon40x40) );
-		Set( "AutomationWindow.SmokeTestFilter", new IMAGE_BRUSH( "Automation/SmokeTestFilter", Icon40x40) );
-		Set( "AutomationWindow.DeveloperDirectoryContent", new IMAGE_BRUSH( "Automation/DeveloperDirectoryContent", Icon40x40) );
-		Set( "AutomationWindow.ExcludedTestsFilter", new IMAGE_BRUSH("Automation/ExcludedTestsFilter", Icon40x40) );
-		Set( "AutomationWindow.ErrorFilter.Small", new IMAGE_BRUSH( "Automation/ErrorFilter", Icon20x20) );
-		Set( "AutomationWindow.WarningFilter.Small", new IMAGE_BRUSH( "Automation/WarningFilter", Icon20x20) );
-		Set( "AutomationWindow.SmokeTestFilter.Small", new IMAGE_BRUSH( "Automation/SmokeTestFilter", Icon20x20) );
-		Set( "AutomationWindow.DeveloperDirectoryContent.Small", new IMAGE_BRUSH( "Automation/DeveloperDirectoryContent", Icon20x20) );
-		Set( "AutomationWindow.TrackHistory", new IMAGE_BRUSH( "Automation/TrackTestHistory", Icon40x40) );
-
-		//device group settings
-		Set( "AutomationWindow.GroupSettings", new IMAGE_BRUSH( "Automation/Groups", Icon40x40) );
-		Set( "AutomationWindow.GroupSettings.Small", new IMAGE_BRUSH( "Automation/Groups", Icon20x20) );
-
-		//test preset icons
-		Set( "AutomationWindow.PresetNew", new IMAGE_BRUSH( "Icons/icon_add_40x", Icon16x16 ) );
-		Set( "AutomationWindow.PresetSave", new IMAGE_BRUSH_SVG("Starship/Common/SaveCurrent", Icon16x16));
-		Set( "AutomationWindow.PresetRemove", new IMAGE_BRUSH( "Icons/icon_Cascade_DeleteLOD_40x", Icon16x16 ) );
-
-		//test backgrounds
-		Set( "AutomationWindow.GameGroupBorder", new BOX_BRUSH( "Automation/GameGroupBorder", FMargin(4.0f/16.0f) ) );
-		Set( "AutomationWindow.EditorGroupBorder", new BOX_BRUSH( "Automation/EditorGroupBorder", FMargin(4.0f/16.0f) ) );
-	}
-
-	// Launcher
-	{
-		Set( "Launcher.Run", new IMAGE_BRUSH("Launcher/Launcher_Run", Icon40x40) );
-		Set( "Launcher.EditSettings", new IMAGE_BRUSH("Launcher/Launcher_EditSettings", Icon40x40) );
-		Set( "Launcher.Back", new IMAGE_BRUSH("Launcher/Launcher_Back", Icon32x32) );
-		Set( "Launcher.Back.Small", new IMAGE_BRUSH("Launcher/Launcher_Back", Icon32x32));
-		Set( "Launcher.Delete", new IMAGE_BRUSH("Launcher/Launcher_Delete", Icon32x32) );
-
-		Set( "Launcher.Instance_Commandlet", new IMAGE_BRUSH( "Launcher/Instance_Commandlet", Icon25x25 ) );
-		Set( "Launcher.Instance_Editor", new IMAGE_BRUSH( "Launcher/Instance_Editor", Icon25x25 ) );
-		Set( "Launcher.Instance_Game", new IMAGE_BRUSH( "Launcher/Instance_Game", Icon25x25 ) );
-		Set( "Launcher.Instance_Other", new IMAGE_BRUSH( "Launcher/Instance_Other", Icon25x25 ) );
-		Set( "Launcher.Instance_Server", new IMAGE_BRUSH( "Launcher/Instance_Server", Icon25x25 ) );
-		Set( "Launcher.Instance_Unknown", new IMAGE_BRUSH( "Launcher/Instance_Unknown", Icon25x25 ) );
-		Set( "LauncherCommand.DeployBuild", new IMAGE_BRUSH( "Launcher/Launcher_Deploy", Icon40x40 ) );
-		Set( "LauncherCommand.QuickLaunch", new IMAGE_BRUSH_SVG( "Starship/Launcher/PaperAirplane", Icon20x20 ) );
-		Set( "LauncherCommand.CreateBuild", new IMAGE_BRUSH( "Launcher/Launcher_Build", Icon40x40 ) );
-		Set( "LauncherCommand.AdvancedBuild", new IMAGE_BRUSH( "Launcher/Launcher_Advanced", Icon40x40 ) );
-		Set( "LauncherCommand.AdvancedBuild.Medium", new IMAGE_BRUSH("Launcher/Launcher_Advanced", Icon25x25) );
-		Set( "LauncherCommand.AdvancedBuild.Small", new IMAGE_BRUSH("Launcher/Launcher_Advanced", Icon20x20) );
-
-		Set("Launcher.Filters.Text", FTextBlockStyle(NormalText)
-			.SetFont(DEFAULT_FONT("Bold", 9))
-			.SetColorAndOpacity(FLinearColor(1.0f, 1.0f, 1.0f, 0.9f))
-			.SetShadowOffset(FVector2D(1, 1))
-			.SetShadowColorAndOpacity(FLinearColor(0, 0, 0, 0.9f)));
-
-		const FComboButtonStyle LauncherComboButton = FComboButtonStyle(FStarshipCoreStyle::GetCoreStyle().GetWidgetStyle<FComboButtonStyle>("ComboButton"))
-			.SetDownArrowPadding(FMargin(67.0f, 2.0f, 2.0f, 2.0f));
-
-		Set("Launcher.ComboButton", LauncherComboButton);
-
-#if DDPI_HAS_EXTENDED_PLATFORMINFO_DATA
-
-		Set( "Launcher.Platform.AllPlatforms", new IMAGE_BRUSH( "Launcher/All_Platforms_24x", Icon24x24) );
-		Set( "Launcher.Platform.AllPlatforms.Large", new IMAGE_BRUSH( "Launcher/All_Platforms_128x", Icon64x64) );
-		Set( "Launcher.Platform.AllPlatforms.XLarge", new IMAGE_BRUSH( "Launcher/All_Platforms_128x", Icon128x128) );
-		for (auto Pair : FDataDrivenPlatformInfoRegistry::GetAllPlatformInfos())
-		{
-			const FDataDrivenPlatformInfo& PlatformInfo = Pair.Value;
-
-			// some platforms may specify a "rooted" path in the platform extensions directory, so look for that case here, and use a different path for the brush
-			FString NormalIconPath = PlatformInfo.GetIconPath(EPlatformIconSize::Normal);
-			if(!NormalIconPath.IsEmpty())
-			{
-				if (NormalIconPath.StartsWith(TEXT("/Platforms/")))
-				{
-#define PLATFORM_IMAGE_BRUSH( PlatformPath, ... ) FSlateImageBrush( PlatformPath.Replace(TEXT("/Platforms/"), *FPaths::EnginePlatformExtensionsDir()) + TEXT(".png") , __VA_ARGS__ )
-					Set(PlatformInfo.GetIconStyleName(EPlatformIconSize::Normal), new PLATFORM_IMAGE_BRUSH(NormalIconPath, Icon24x24));
-					Set(PlatformInfo.GetIconStyleName(EPlatformIconSize::Large), new PLATFORM_IMAGE_BRUSH(PlatformInfo.GetIconPath(EPlatformIconSize::Large), Icon64x64));
-					Set(PlatformInfo.GetIconStyleName(EPlatformIconSize::XLarge), new PLATFORM_IMAGE_BRUSH(PlatformInfo.GetIconPath(EPlatformIconSize::XLarge), Icon128x128));
-				}
-				else
-				{
-					Set(PlatformInfo.GetIconStyleName(EPlatformIconSize::Normal), new IMAGE_BRUSH(*NormalIconPath, Icon24x24));
-					Set(PlatformInfo.GetIconStyleName(EPlatformIconSize::Large), new IMAGE_BRUSH(*PlatformInfo.GetIconPath(EPlatformIconSize::Large), Icon64x64));
-					Set(PlatformInfo.GetIconStyleName(EPlatformIconSize::XLarge), new IMAGE_BRUSH(*PlatformInfo.GetIconPath(EPlatformIconSize::XLarge), Icon128x128));
-				}
-			}
-		}
-
-		for (const FPreviewPlatformMenuItem& Item : FDataDrivenPlatformInfoRegistry::GetAllPreviewPlatformMenuItems())
-		{
-			if (!Item.ActiveIconPath.IsEmpty())
-			{
-				Set(Item.ActiveIconName, new PLATFORM_IMAGE_BRUSH(Item.ActiveIconPath, Icon40x40));
-			}
-			if (!Item.InactiveIconPath.IsEmpty())
-			{
-				Set(Item.InactiveIconName, new PLATFORM_IMAGE_BRUSH(Item.InactiveIconPath, Icon40x40));
-			}
-		}
-#endif
-
-		Set("Launcher.NoHoverTableRow", FTableRowStyle(NormalTableRowStyle)
-			.SetEvenRowBackgroundHoveredBrush(FSlateNoResource())
-			.SetOddRowBackgroundHoveredBrush(FSlateNoResource())
-			.SetActiveHoveredBrush(FSlateNoResource())
-			.SetInactiveHoveredBrush(FSlateNoResource())
-			);
-	}
 
 	// Device Manager
 	{
diff --git a/Engine/Source/Editor/EditorStyle/Private/StarshipStyle.h b/Engine/Source/Editor/EditorStyle/Private/StarshipStyle.h
--- a/Engine/Source/Editor/EditorStyle/Private/StarshipStyle.h
+++ b/Engine/Source/Editor/EditorStyle/Private/StarshipStyle.h
@@ -122,10 +122,6 @@
 		const FSlateColor HighlightColor;
 		const FSlateColor WindowHighlightColor;
 
-		const FSlateColor LogColor_SelectionBackground;
-		const FSlateColor LogColor_Normal;
-		const FSlateColor LogColor_Command;
-
 		// These are common colors used thruout the editor in mutliple style elements
 		const FSlateColor InheritedFromBlueprintTextColor;
 
diff --git a/Engine/Source/Editor/LevelEditor/Private/LevelEditor.cpp b/Engine/Source/Editor/LevelEditor/Private/LevelEditor.cpp
--- a/Engine/Source/Editor/LevelEditor/Private/LevelEditor.cpp
+++ b/Engine/Source/Editor/LevelEditor/Private/LevelEditor.cpp
@@ -332,7 +332,7 @@
 
 	// Clear out some globals that may be referencing this module
 	SetLevelEditorTabManager(nullptr);
-	WorkspaceMenu::GetModule().ResetLevelEditorCategory();
+	WorkspaceMenu::GetMenuStructure().GetLevelEditorCategory()->ClearItems();
 
 	if (FSlateApplication::IsInitialized() && FModuleManager::Get().IsModuleLoaded("SlateReflector"))
 	{
diff --git a/Engine/Source/Editor/StatusBar/Private/StatusBarSubsystem.cpp b/Engine/Source/Editor/StatusBar/Private/StatusBarSubsystem.cpp
--- a/Engine/Source/Editor/StatusBar/Private/StatusBarSubsystem.cpp
+++ b/Engine/Source/Editor/StatusBar/Private/StatusBarSubsystem.cpp
@@ -24,7 +24,7 @@
 #include "OutputLogModule.h"
 #include "WidgetDrawerConfig.h"
 #include "Widgets/Input/SMultiLineEditableTextBox.h"
-#include "Settings/EditorStyleSettings.h"
+#include "OutputLog/Public/OutputLogSettings.h"
 
 #define LOCTEXT_NAMESPACE "StatusBar"
 
@@ -264,7 +264,7 @@
 	// Get the global output log tab if it exists. If it exists and is in the same editor as the status bar we'll focus that instead
 	TSharedPtr<SDockTab> MainOutputLogTab = OutputLogModule.GetOutputLogTab();
 	
-	const bool bCycleToOutputLogDrawer = GetDefault<UEditorStyleSettings>()->bCycleToOutputLogDrawer || bAlwaysToggleDrawer;
+	const bool bCycleToOutputLogDrawer = GetDefault<UOutputLogSettings>()->bCycleToOutputLogDrawer || bAlwaysToggleDrawer;
 
 	for (auto StatusBar : StatusBars)
 	{
diff --git a/Engine/Source/Editor/UnrealEd/Classes/Settings/EditorStyleSettings.h b/Engine/Source/Editor/UnrealEd/Classes/Settings/EditorStyleSettings.h
--- a/Engine/Source/Editor/UnrealEd/Classes/Settings/EditorStyleSettings.h
+++ b/Engine/Source/Editor/UnrealEd/Classes/Settings/EditorStyleSettings.h
@@ -29,22 +29,6 @@
 	LastDockedWindowOrContentBrowser
 };
 
-UENUM()
-enum class ELogCategoryColorizationMode : uint8
-{
-	/** Do not colorize based on log categories */
-	None,
-
-	/** Colorize the entire log line, but not warnings or errors */
-	ColorizeWholeLine,
-
-	/** Colorize only the category name (including on warnings and errors) */
-	ColorizeCategoryOnly,
-
-	/** Colorize the background of the category name (including on warnings and errors) */
-	ColorizeCategoryAsBadge
-};
-
 /**
  * Implements the Editor style settings.
  */
@@ -153,29 +137,6 @@
 	UPROPERTY(config)
 	uint32 bShowHiddenPropertiesWhilePlaying : 1;
 
-	/** The font size used in the output log */
-	UPROPERTY(EditAnywhere, config, Category="Output Log", meta=(DisplayName="Log Font Size", ConfigRestartRequired=true))
-	int32 LogFontSize;
-
-	/** The display mode for timestamps in the output log window */
-	UPROPERTY(EditAnywhere, config, Category="Output Log", meta=(DisplayName = "Output Log Window Timestamp Mode"))
-	TEnumAsByte<ELogTimes::Type> LogTimestampMode;
-
-	/** Should warnings and errors in the Output Log during "Play in Editor" be promoted to the message log? */
-	UPROPERTY(EditAnywhere, config, Category="Output Log")
-	bool bPromoteOutputLogWarningsDuringPIE;
-
-	/** How should categories be colorized in the output log? */
-	UPROPERTY(EditAnywhere, config, Category = "Output Log")
-	ELogCategoryColorizationMode CategoryColorizationMode;
-
-	/**
-	 * If checked pressing the console command shortcut will cycle between focusing the status bar console, opening the output log drawer, and back to the previous focus target. 
-	 * If unchecked, the console command shortcut will only focus the status bar console
-	 */
-	UPROPERTY(EditAnywhere, config, Category = "Output Log", meta = (DisplayName = "Open Output Log Drawer with Console Command Shortcut"))
-	bool bCycleToOutputLogDrawer;
-
 	/** New asset editor tabs will open at the specified location. */
 	UPROPERTY(EditAnywhere, config, Category=UserInterface)
 	EAssetEditorOpenLocation AssetEditorOpenLocation;
diff --git a/Engine/Source/Editor/UnrealEd/Classes/Settings/LevelEditorPlaySettings.h b/Engine/Source/Editor/UnrealEd/Classes/Settings/LevelEditorPlaySettings.h
--- a/Engine/Source/Editor/UnrealEd/Classes/Settings/LevelEditorPlaySettings.h
+++ b/Engine/Source/Editor/UnrealEd/Classes/Settings/LevelEditorPlaySettings.h
@@ -286,6 +286,10 @@
 	UPROPERTY(config, EditAnywhere, Category = PlayInEditor, meta = (DisplayName="Stream Sub-Levels during Play in Editor", ToolTip="Prefer to stream sub-levels from the disk instead of duplicating editor sub-levels"))
 	uint32 bPreferToStreamLevelsInPIE:1;
 
+	/** Should warnings and errors in the Output Log during "Play in Editor" be promoted to the message log? */
+	UPROPERTY(EditAnywhere, config, Category = PlayInEditor)
+	bool bPromoteOutputLogWarningsDuringPIE;
+
 public:
 	/** The width of the new view port window in pixels (0 = use the desktop's screen resolution). */
 	UPROPERTY(config, EditAnywhere, Category=GameViewportSettings, meta=(ClampMin=0))
diff --git a/Engine/Source/Editor/UnrealEd/Private/PlayLevel.cpp b/Engine/Source/Editor/UnrealEd/Private/PlayLevel.cpp
--- a/Engine/Source/Editor/UnrealEd/Private/PlayLevel.cpp
+++ b/Engine/Source/Editor/UnrealEd/Private/PlayLevel.cpp
@@ -30,7 +30,6 @@
 #include "Widgets/SViewport.h"
 #include "Framework/Docking/TabManager.h"
 #include "Styling/AppStyle.h"
-#include "Settings/EditorStyleSettings.h"
 #include "Engine/EngineTypes.h"
 #include "Async/TaskGraphInterfaces.h"
 #include "GameFramework/Actor.h"
@@ -2633,7 +2632,7 @@
 	}
 
 	// Register for log processing so we can promote errors/warnings to the message log
-	if (GetDefault<UEditorStyleSettings>()->bPromoteOutputLogWarningsDuringPIE)
+	if (GetDefault<ULevelEditorPlaySettings>()->bPromoteOutputLogWarningsDuringPIE)
 	{
 		OutputLogErrorsToMessageLogProxyPtr = MakeShareable(new FOutputLogErrorsToMessageLogProxy());
 	}
diff --git a/Engine/Source/Editor/UnrealEd/Private/Settings/EditorStyleClasses.cpp b/Engine/Source/Editor/UnrealEd/Private/Settings/EditorStyleClasses.cpp
--- a/Engine/Source/Editor/UnrealEd/Private/Settings/EditorStyleClasses.cpp
+++ b/Engine/Source/Editor/UnrealEd/Private/Settings/EditorStyleClasses.cpp
@@ -25,8 +25,6 @@
 	
 	bUseGrid = true;
 
-	bCycleToOutputLogDrawer = true;
-
 	RegularColor = FLinearColor(0.035f, 0.035f, 0.035f);
 	RuleColor = FLinearColor(0.008f, 0.008f, 0.008f);
 	CenterColor = FLinearColor::Black;
@@ -35,7 +33,6 @@
 
 	bShowFriendlyNames = true;
 	bShowNativeComponentNames = true;
-	LogTimestampMode = ELogTimes::None;
 }
 
 void UEditorStyleSettings::Init()
diff --git a/Engine/Source/Programs/AutomationTool/Scripts/CopyBuildToStagingDirectory.Automation.cs b/Engine/Source/Programs/AutomationTool/Scripts/CopyBuildToStagingDirectory.Automation.cs
--- a/Engine/Source/Programs/AutomationTool/Scripts/CopyBuildToStagingDirectory.Automation.cs
+++ b/Engine/Source/Programs/AutomationTool/Scripts/CopyBuildToStagingDirectory.Automation.cs
@@ -1150,7 +1150,8 @@ namespace AutomationScripts
 					UnrealTargetPlatform CrashReportPlatform = ThisPlatform.CrashReportPlatform ?? SC.StageTargetPlatform.PlatformType;
 					if (SC.bStageCrashReporter && PlatformSupportsCrashReporter(CrashReportPlatform) && (Params.IterateSharedBuildUsePrecompiledExe == false))
 					{
-						SC.StageCrashReporterFiles(StagedFileType.UFS, DirectoryReference.Combine(SC.EngineRoot, "Content", "Slate"), StageFilesSearch.AllDirectories);
+						SC.StageCrashReporterFiles(StagedFileType.UFS, DirectoryReference.Combine(SC.EngineRoot, "Content", "Slate", "Common"), StageFilesSearch.AllDirectories);
+						SC.StageCrashReporterFiles(StagedFileType.UFS, DirectoryReference.Combine(SC.EngineRoot, "Content", "Slate", "Old"), StageFilesSearch.AllDirectories);
 						SC.StageCrashReporterFiles(StagedFileType.UFS, DirectoryReference.Combine(SC.EngineRoot, "Shaders", "StandaloneRenderer"), StageFilesSearch.AllDirectories);
 						SC.StageCrashReporterFiles(StagedFileType.UFS, DirectoryReference.Combine(SC.EngineRoot, "Content", "Internationalization", InternationalizationPreset), StageFilesSearch.AllDirectories, new StagedDirectoryReference("Engine/Content/Internationalization"));
 
diff --git a/Engine/Source/Programs/CrashReportClient/Private/CrashReportClientStyle.cpp b/Engine/Source/Programs/CrashReportClient/Private/CrashReportClientStyle.cpp
--- a/Engine/Source/Programs/CrashReportClient/Private/CrashReportClientStyle.cpp
+++ b/Engine/Source/Programs/CrashReportClient/Private/CrashReportClientStyle.cpp
@@ -7,6 +7,7 @@
 #include "Math/Color.h"
 #include "Math/Vector2D.h"
 #include "CrashReportClientApp.h"
+#include "Styling/StyleColors.h"
 
 #if !CRASH_REPORT_UNATTENDED_ONLY
 
@@ -116,6 +117,8 @@
 
 	Style.Set( "RichText.Hyperlink", DarkHyperlink );
 
+	Style.Set("ToolPanel.GroupBorder", new FSlateColorBrush(FStyleColors::Panel));
+
 	return StyleRef;
 }
 
diff --git a/Engine/Source/Programs/CrashReportClient/Private/SCrashReportClient.cpp b/Engine/Source/Programs/CrashReportClient/Private/SCrashReportClient.cpp
--- a/Engine/Source/Programs/CrashReportClient/Private/SCrashReportClient.cpp
+++ b/Engine/Source/Programs/CrashReportClient/Private/SCrashReportClient.cpp
@@ -5,7 +5,6 @@
 #if !CRASH_REPORT_UNATTENDED_ONLY
 
 #include "CrashReportClientStyle.h"
-#include "Styling/SlateStyle.h"
 #include "Styling/CoreStyle.h"
 #include "Widgets/Images/SThrobber.h"
 #include "CrashDescription.h"
@@ -85,7 +84,7 @@
 	ChildSlot
 	[
 		SNew(SBorder)
-		.BorderImage(FCoreStyle::Get().GetBrush("ToolPanel.GroupBorder"))
+		.BorderImage(FCrashReportClientStyle::Get().GetBrush("ToolPanel.GroupBorder"))
 		[
 			SNew(SVerticalBox)
 
@@ -119,7 +118,6 @@
 				SNew( SRichTextBlock )
 				.Text(CrashDetailedMessage)
 				.AutoWrapText(true)
-				.DecoratorStyleSet( &FCoreStyle::Get() )
 				+ SRichTextBlock::HyperlinkDecorator( TEXT("browser"), FSlateHyperlinkRun::FOnClick::CreateStatic( &OnBrowserLinkClicked ) )
 			]
 
@@ -343,7 +341,7 @@
 	ChildSlot
 	[
 		SNew(SBorder)
-		.BorderImage(FCoreStyle::Get().GetBrush("ToolPanel.GroupBorder"))
+		.BorderImage(FCrashReportClientStyle::Get().GetBrush("ToolPanel.GroupBorder"))
 		[
 			SNew(SVerticalBox)
 
@@ -369,7 +367,6 @@
 				SNew( SRichTextBlock )
 				.Text(CrashDetailedMessage)
 				.AutoWrapText(true)
-				.DecoratorStyleSet( &FCoreStyle::Get() )
 				+ SRichTextBlock::HyperlinkDecorator( TEXT("browser"), FSlateHyperlinkRun::FOnClick::CreateStatic( &OnBrowserLinkClicked ) )
 			]
 
diff --git a/Engine/Source/Programs/UnrealFrontend/Private/Commands/UserInterfaceCommand.cpp b/Engine/Source/Programs/UnrealFrontend/Private/Commands/UserInterfaceCommand.cpp
--- a/Engine/Source/Programs/UnrealFrontend/Private/Commands/UserInterfaceCommand.cpp
+++ b/Engine/Source/Programs/UnrealFrontend/Private/Commands/UserInterfaceCommand.cpp
@@ -39,7 +39,8 @@
 	FConfigCacheIni::InitializeConfigSystem();
 	GetTargetPlatformManager();
 
-	FCoreStyle::ResetToDefault();
+	// Crank up a normal Slate application using the platform's standalone renderer.
+	FSlateApplication::InitializeAsStandaloneApplication(GetStandardStandaloneRenderer());
 
 	// load required modules
 	FModuleManager::Get().LoadModuleChecked("Messaging");
@@ -112,7 +113,6 @@
 
 void FUserInterfaceCommand::InitializeSlateApplication( const FString& LayoutIni )
 {
-	FSlateApplication::InitializeAsStandaloneApplication(GetStandardStandaloneRenderer());
 	FSlateApplication::InitHighDPI(true);
 
 	FGlobalTabmanager::Get()->SetApplicationTitle(NSLOCTEXT("UnrealFrontend", "AppTitle", "Unreal Frontend"));
diff --git a/Engine/Source/Runtime/SlateCore/Private/Styling/StarshipCoreStyle.cpp b/Engine/Source/Runtime/SlateCore/Private/Styling/StarshipCoreStyle.cpp
--- a/Engine/Source/Runtime/SlateCore/Private/Styling/StarshipCoreStyle.cpp
+++ b/Engine/Source/Runtime/SlateCore/Private/Styling/StarshipCoreStyle.cpp
@@ -16,6 +16,7 @@
 #include "Styling/SegmentedControlStyle.h"
 #include "Styling/StyleColors.h"
 #include "Styling/CoreStyle.h"
+#include "Misc/DataDrivenPlatformInfoRegistry.h"
 
 // This is to fix the issue that SlateStyleMacros like IMAGE_BRUSH look for RootToContentDir but Style->RootToContentDir is how the core style is set up
 #define RootToContentDir Style->RootToContentDir
@@ -391,6 +392,7 @@
 		Style->Set("Icons.EyeDropper", new IMAGE_BRUSH_SVG("Starship/Common/EyeDropper", Icon16x16));
 		Style->Set("Icons.C++", new IMAGE_BRUSH_SVG("Starship/Common/CPP", Icon16x16));
 
+		Style->Set("Icons.Advanced", new IMAGE_BRUSH_SVG("Starship/Common/Advanced", Icon16x16));
 		Style->Set("Icons.Rotate90Clockwise", new IMAGE_BRUSH_SVG("Starship/Common/Rotate90Clockwise", Icon16x16));
 		Style->Set("Icons.Rotate90CounterClockwise", new IMAGE_BRUSH_SVG("Starship/Common/Rotate90CounterClockwise", Icon16x16));
 		Style->Set("Icons.Rotate180", new IMAGE_BRUSH_SVG("Starship/Common/Rotate180", Icon16x16));
@@ -801,6 +803,107 @@
 		Style->Set("SyntaxHighlight.NodeAttributeValue", FTextBlockStyle(SmallMonospacedText).SetColorAndOpacity(FLinearColor(FColor(0xffb46100)))); // orange
 	}
 
+	// Logging defaults
+	{
+		Style->Set("Log.TabIcon", new IMAGE_BRUSH_SVG("Starship/Common/OutputLog", Icon16x16));
+
+		const FTextBlockStyle NormalLogText = FTextBlockStyle(NormalText)
+			.SetFont(DEFAULT_FONT("Mono", 9))
+			.SetColorAndOpacity(FStyleColors::Foreground)
+			.SetSelectedBackgroundColor(FStyleColors::Highlight)
+			.SetHighlightColor(FStyleColors::Black);
+
+		Style->Set("Log.Normal", NormalLogText);
+
+		Style->Set("Log.Warning", FTextBlockStyle(NormalLogText)
+			.SetColorAndOpacity(FStyleColors::Warning)
+		);
+
+		Style->Set("Log.Error", FTextBlockStyle(NormalLogText)
+			.SetColorAndOpacity(FStyleColors::Error)
+		);
+
+		Style->Set("Log.TextBox", FEditableTextBoxStyle(NormalEditableTextBoxStyle)
+			.SetBackgroundImageNormal(BOX_BRUSH("Common/WhiteGroupBorder", FMargin(4.0f / 16.0f)))
+			.SetBackgroundImageHovered(BOX_BRUSH("Common/WhiteGroupBorder", FMargin(4.0f / 16.0f)))
+			.SetBackgroundImageFocused(BOX_BRUSH("Common/WhiteGroupBorder", FMargin(4.0f / 16.0f)))
+			.SetBackgroundImageReadOnly(BOX_BRUSH("Common/WhiteGroupBorder", FMargin(4.0f / 16.0f)))
+			.SetBackgroundColor(FStyleColors::Recessed)
+		);
+	}
+
+	// Project Launcher
+	{
+		Style->Set("Launcher.TabIcon", new IMAGE_BRUSH_SVG("Starship/Common/ProjectLauncher", Icon16x16));
+		Style->Set("Launcher.Tabs.Tools", new IMAGE_BRUSH("/Icons/icon_tab_Tools_16x", Icon16x16));
+		
+
+		Style->Set("Launcher.Run", new IMAGE_BRUSH("Launcher/Launcher_Run", Icon40x40));
+		Style->Set("Launcher.EditSettings", new IMAGE_BRUSH("Launcher/Launcher_EditSettings", Icon40x40));
+		Style->Set("Launcher.Back", new IMAGE_BRUSH("Launcher/Launcher_Back", Icon32x32));
+		Style->Set("Launcher.Back.Small", new IMAGE_BRUSH("Launcher/Launcher_Back", Icon32x32));
+		Style->Set("Launcher.Delete", new IMAGE_BRUSH("Launcher/Launcher_Delete", Icon32x32));
+
+		Style->Set("Launcher.Instance_Commandlet", new IMAGE_BRUSH("Launcher/Instance_Commandlet", Icon25x25));
+		Style->Set("Launcher.Instance_Editor", new IMAGE_BRUSH("Launcher/Instance_Editor", Icon25x25));
+		Style->Set("Launcher.Instance_Game", new IMAGE_BRUSH("Launcher/Instance_Game", Icon25x25));
+		Style->Set("Launcher.Instance_Other", new IMAGE_BRUSH("Launcher/Instance_Other", Icon25x25));
+		Style->Set("Launcher.Instance_Server", new IMAGE_BRUSH("Launcher/Instance_Server", Icon25x25));
+		Style->Set("Launcher.Instance_Unknown", new IMAGE_BRUSH("Launcher/Instance_Unknown", Icon25x25));
+		Style->Set("LauncherCommand.DeployBuild", new IMAGE_BRUSH("Launcher/Launcher_Deploy", Icon40x40));
+		Style->Set("LauncherCommand.QuickLaunch", new IMAGE_BRUSH_SVG("Starship/Launcher/PaperAirplane", Icon20x20));
+		Style->Set("LauncherCommand.CreateBuild", new IMAGE_BRUSH("Launcher/Launcher_Build", Icon40x40));
+		Style->Set("LauncherCommand.AdvancedBuild", new IMAGE_BRUSH("Launcher/Launcher_Advanced", Icon40x40));
+		Style->Set("LauncherCommand.AdvancedBuild.Medium", new IMAGE_BRUSH("Launcher/Launcher_Advanced", Icon25x25));
+		Style->Set("LauncherCommand.AdvancedBuild.Small", new IMAGE_BRUSH("Launcher/Launcher_Advanced", Icon20x20));
+
+#if DDPI_HAS_EXTENDED_PLATFORMINFO_DATA && (WITH_EDITOR || WITH_UNREAL_DEVELOPER_TOOLS)
+
+		Style->Set("Launcher.Platform.AllPlatforms", new IMAGE_BRUSH("Launcher/All_Platforms_24x", Icon24x24));
+		Style->Set("Launcher.Platform.AllPlatforms.Large", new IMAGE_BRUSH("Launcher/All_Platforms_128x", Icon64x64));
+		Style->Set("Launcher.Platform.AllPlatforms.XLarge", new IMAGE_BRUSH("Launcher/All_Platforms_128x", Icon128x128));
+		for (auto Pair : FDataDrivenPlatformInfoRegistry::GetAllPlatformInfos())
+		{
+			const FDataDrivenPlatformInfo& PlatformInfo = Pair.Value;
+
+			// some platforms may specify a "rooted" path in the platform extensions directory, so look for that case here, and use a different path for the brush
+			FString NormalIconPath = PlatformInfo.GetIconPath(EPlatformIconSize::Normal);
+			if (!NormalIconPath.IsEmpty())
+			{
+				if (NormalIconPath.StartsWith(TEXT("/Platforms/")))
+				{
+#define PLATFORM_IMAGE_BRUSH( PlatformPath, ... ) FSlateImageBrush( PlatformPath.Replace(TEXT("/Platforms/"), *FPaths::EnginePlatformExtensionsDir()) + TEXT(".png") , __VA_ARGS__ )
+					Style->Set(PlatformInfo.GetIconStyleName(EPlatformIconSize::Normal), new PLATFORM_IMAGE_BRUSH(NormalIconPath, Icon24x24));
+					Style->Set(PlatformInfo.GetIconStyleName(EPlatformIconSize::Large), new PLATFORM_IMAGE_BRUSH(PlatformInfo.GetIconPath(EPlatformIconSize::Large), Icon64x64));
+					Style->Set(PlatformInfo.GetIconStyleName(EPlatformIconSize::XLarge), new PLATFORM_IMAGE_BRUSH(PlatformInfo.GetIconPath(EPlatformIconSize::XLarge), Icon128x128));
+				}
+				else
+				{
+					const FString PathPrefix = FPaths::EngineContentDir() / TEXT("Editor/Slate");
+#define DEVELOPER_IMAGE_BRUSH( RelativePath, ... ) FSlateImageBrush(PathPrefix / RelativePath + TEXT(".png"), __VA_ARGS__ )
+					Style->Set(PlatformInfo.GetIconStyleName(EPlatformIconSize::Normal), new DEVELOPER_IMAGE_BRUSH(*NormalIconPath, Icon24x24));
+					Style->Set(PlatformInfo.GetIconStyleName(EPlatformIconSize::Large), new DEVELOPER_IMAGE_BRUSH(*PlatformInfo.GetIconPath(EPlatformIconSize::Large), Icon64x64));
+					Style->Set(PlatformInfo.GetIconStyleName(EPlatformIconSize::XLarge), new DEVELOPER_IMAGE_BRUSH(*PlatformInfo.GetIconPath(EPlatformIconSize::XLarge), Icon128x128));
+#undef DEVELOPER_IMAGE_BRUSH
+				}
+			}
+		}
+
+		for (const FPreviewPlatformMenuItem& Item : FDataDrivenPlatformInfoRegistry::GetAllPreviewPlatformMenuItems())
+		{
+			if (!Item.ActiveIconPath.IsEmpty())
+			{
+				Style->Set(Item.ActiveIconName, new PLATFORM_IMAGE_BRUSH(Item.ActiveIconPath, Icon40x40));
+			}
+			if (!Item.InactiveIconPath.IsEmpty())
+			{
+				Style->Set(Item.InactiveIconName, new PLATFORM_IMAGE_BRUSH(Item.InactiveIconPath, Icon40x40));
+			}
+		}
+#endif
+
+	}
+
 	return Style;
 }
 
