Updating developer modules relying on Editor styles

diff --git a/Engine/Source/Developer/Profiler/Private/Widgets/SDataGraph.cpp b/Engine/Source/Developer/Profiler/Private/Widgets/SDataGraph.cpp
--- a/Engine/Source/Developer/Profiler/Private/Widgets/SDataGraph.cpp
+++ b/Engine/Source/Developer/Profiler/Private/Widgets/SDataGraph.cpp
@@ -25,6 +25,7 @@
 #include "ProfilerManager.h"
 #include "Widgets/StatDragDropOp.h"
 #include "../ProfilerStyle.h"
+#include "SSimpleButton.h"
 
 #define LOCTEXT_NAMESPACE "SDataGraph"
 
@@ -68,7 +69,6 @@
 		OnGetMouseFrameIndex = InArgs._OnGetMouseFrameIndex;
 
 		const FSlateColor TextColor( TrackedStat->GraphColor );
-		const FName CloseButtonStyle = TEXT("Docking.MajorTab.CloseButton");
 		FFormatNamedArguments Args;
 		Args.Add( TEXT("StatName"), FText::FromString( TrackedStat->GraphDataSource->GetStatName() ) );
 		const FText ToolTipText = FText::Format( LOCTEXT("DataGraphSummary_CloseButton_TT", "Click to stop tracking '{StatName}' stat"), Args );
@@ -84,15 +84,10 @@
 			.VAlign(VAlign_Center)
 			.Padding( 1.0f )
 			[
-				SNew(SButton)
-				.ButtonStyle( FAppStyle::Get(), CloseButtonStyle )
+				SNew(SSimpleButton)
+				.Icon(FAppStyle::GetBrush("Icons.Stop"))
 				.OnClicked( this, &SDataGraphSummary::CloseButton_OnClicked )
-				.ContentPadding( 0 )
 				.ToolTipText( ToolTipText )
-				[
-					SNew(SSpacer)
-					.Size( FAppStyle::GetBrush(CloseButtonStyle, ".Normal" )->ImageSize )
-				]
 			]
 
 			// Stat group name.
diff --git a/Engine/Source/Developer/Profiler/Profiler.Build.cs b/Engine/Source/Developer/Profiler/Profiler.Build.cs
--- a/Engine/Source/Developer/Profiler/Profiler.Build.cs
+++ b/Engine/Source/Developer/Profiler/Profiler.Build.cs
@@ -23,7 +23,6 @@ public class Profiler : ModuleRules
 				"RHI",
 				"RenderCore",
 				"Slate",
-                
 				"ProfilerClient",
 				"DesktopPlatform",
 			}
@@ -41,6 +40,7 @@ public class Profiler : ModuleRules
 		PrivateDependencyModuleNames.AddRange(
 			new string[] {
 				"SlateCore",
+				"ToolWidgets"
 			}
 		);
 
diff --git a/Engine/Source/Developer/SessionFrontend/Private/Widgets/Browser/SSessionBrowserTreeInstanceRow.h b/Engine/Source/Developer/SessionFrontend/Private/Widgets/Browser/SSessionBrowserTreeInstanceRow.h
--- a/Engine/Source/Developer/SessionFrontend/Private/Widgets/Browser/SSessionBrowserTreeInstanceRow.h
+++ b/Engine/Source/Developer/SessionFrontend/Private/Widgets/Browser/SSessionBrowserTreeInstanceRow.h
@@ -22,6 +22,7 @@
 #include "Models/SessionBrowserTreeItems.h"
 #include "Widgets/Images/SImage.h"
 #include "PlatformInfo.h"
+#include "Styling/StyleColors.h"
 
 #define LOCTEXT_NAMESPACE "SSessionBrowserTreeRow"
 
@@ -154,7 +155,8 @@ public:
 			.VAlign(VAlign_Center)
 			[
 				SNew(SImage)
-				.Image(this, &SSessionBrowserTreeInstanceRow::HandleStatusImage)
+				.Image(FAppStyle::GetBrush("Icons.FilledCircle"))
+				.ColorAndOpacity(this, &SSessionBrowserTreeInstanceRow::HandleStatusImageColor)
 			];
 		}
 		else if (ColumnName == "Type")
@@ -187,7 +189,7 @@ private:
 
 		if (InstanceInfo.IsValid() && !InstanceInfo->IsAuthorized())
 		{
-			return FAppStyle::GetBrush("SessionBrowser.SessionLocked");
+			return FAppStyle::GetBrush("Icons.Lock");
 		}
 
 		return nullptr;
@@ -264,7 +266,7 @@ private:
 	}
 
 	/** Callback for getting the image of the Status icon. */
-	const FSlateBrush* HandleStatusImage() const
+	FSlateColor HandleStatusImageColor() const
 	{
 		TSharedPtr<ISessionInstanceInfo> InstanceInfo = Item->GetInstanceInfo();
 
@@ -272,13 +274,13 @@ private:
 		{
 			if (FDateTime::UtcNow() - InstanceInfo->GetLastUpdateTime() < FTimespan::FromSeconds(10.0))
 			{
-				return FAppStyle::GetBrush("SessionBrowser.StatusRunning");
+				return FStyleColors::AccentGreen;
 			}
 
-			return FAppStyle::GetBrush("SessionBrowser.StatusTimedOut");
+			return FStyleColors::Warning;
 		}
 
-		return nullptr;
+		return FStyleColors::Error;
 	}
 
 	/** Callback for getting the foreground text color. */
diff --git a/Engine/Source/Editor/EditorStyle/Private/SlateEditorStyle.cpp b/Engine/Source/Editor/EditorStyle/Private/SlateEditorStyle.cpp
--- a/Engine/Source/Editor/EditorStyle/Private/SlateEditorStyle.cpp
+++ b/Engine/Source/Editor/EditorStyle/Private/SlateEditorStyle.cpp
@@ -1263,11 +1263,6 @@
 
 	// Session Browser tab
 	{
-		Set("SessionBrowser.SessionLocked", new IMAGE_BRUSH( "Icons/icon_locked_highlight_16px", Icon16x16 ) );
-		Set("SessionBrowser.StatusRunning", new IMAGE_BRUSH( "Icons/icon_status_green_16x", Icon16x16 ) );
-		Set("SessionBrowser.StatusTimedOut", new IMAGE_BRUSH( "Icons/icon_status_grey_16x", Icon16x16 ) );
-		Set("SessionBrowser.Terminate", new IMAGE_BRUSH( "Icons/icon_DevicePowerOff_40x", Icon20x20 ) );
-
 		Set("SessionBrowser.Terminate.Font", FTextBlockStyle(NormalText)
 			.SetFont(DEFAULT_FONT( "Bold", 12))
 			.SetColorAndOpacity(FLinearColor(1.0f, 1.0f, 1.0f))
diff --git a/Engine/Source/Editor/EditorStyle/Private/StarshipStyle.cpp b/Engine/Source/Editor/EditorStyle/Private/StarshipStyle.cpp
--- a/Engine/Source/Editor/EditorStyle/Private/StarshipStyle.cpp
+++ b/Engine/Source/Editor/EditorStyle/Private/StarshipStyle.cpp
@@ -331,8 +331,6 @@ void FStarshipEditorStyle::FStyle::SetupGeneralStyles()
 
 
 
-	Set("PlainBorder", new BORDER_BRUSH("Common/PlainBorder", 2.f / 8.f));
-
 	Set( "WideDash.Horizontal", new CORE_IMAGE_BRUSH("Starship/Common/Dash_Horizontal", FVector2D(10, 1), FLinearColor::White, ESlateBrushTileType::Horizontal));
 	Set( "WideDash.Vertical", new CORE_IMAGE_BRUSH("Starship/Common/Dash_Vertical", FVector2D(1, 10), FLinearColor::White, ESlateBrushTileType::Vertical));
 
@@ -967,11 +965,6 @@ void FStarshipEditorStyle::FStyle::SetupGeneralStyles()
 
 	// Session Browser tab
 	{
-		Set("SessionBrowser.SessionLocked", new IMAGE_BRUSH( "Icons/icon_locked_highlight_16px", Icon16x16 ) );
-		Set("SessionBrowser.StatusRunning", new IMAGE_BRUSH( "Icons/icon_status_green_16x", Icon16x16 ) );
-		Set("SessionBrowser.StatusTimedOut", new IMAGE_BRUSH( "Icons/icon_status_grey_16x", Icon16x16 ) );
-		Set("SessionBrowser.Terminate", new IMAGE_BRUSH( "Icons/icon_DevicePowerOff_40x", Icon20x20 ) );
-
 		Set("SessionBrowser.Terminate.Font", FTextBlockStyle(NormalText)
 			.SetFont(DEFAULT_FONT( "Bold", 12))
 			.SetColorAndOpacity(FLinearColor(1.0f, 1.0f, 1.0f))
@@ -2999,42 +2992,6 @@ void FStarshipEditorStyle::FStyle::SetupProfilerStyle()
 		Set( "Profiler.Tab.EventGraph", new CORE_IMAGE_BRUSH( "Icons/Profiler/profiler_OpenEventGraph_32x", Icon16x16 ) );
 		Set( "Profiler.Tab.FiltersAndPresets", new CORE_IMAGE_BRUSH( "Icons/Profiler/Profiler_Filter_Presets_Tab_16x", Icon16x16 ) );
 
-		Set( "ProfilerCommand.ProfilerManager_Load", new CORE_IMAGE_BRUSH( "Icons/Profiler/Profiler_Load_Profiler_40x", Icon40x40 ) );
-		Set( "ProfilerCommand.ProfilerManager_Load.Small", new CORE_IMAGE_BRUSH( "Icons/Profiler/Profiler_Load_Profiler_40x", Icon20x20 ) );
-
-		Set("ProfilerCommand.ProfilerManager_LoadMultiple", new CORE_IMAGE_BRUSH("Icons/Profiler/Profiler_LoadMultiple_Profiler_40x", Icon40x40));
-		Set("ProfilerCommand.ProfilerManager_LoadMultiple.Small", new CORE_IMAGE_BRUSH("Icons/Profiler/Profiler_LoadMultiple_Profiler_40x", Icon20x20));
-
-		Set( "ProfilerCommand.ProfilerManager_Save", new IMAGE_BRUSH_SVG("Starship/Common/save", Icon16x16));
-		Set("ProfilerCommand.ProfilerManager_Save.Small", new IMAGE_BRUSH_SVG("Starship/Common/save", Icon16x16));
-		
-		Set( "ProfilerCommand.ProfilerManager_ToggleLivePreview", new CORE_IMAGE_BRUSH( "Automation/RefreshTests", Icon40x40) );
-		Set( "ProfilerCommand.ProfilerManager_ToggleLivePreview.Small", new CORE_IMAGE_BRUSH( "Automation/RefreshTests", Icon20x20) );
-
-		Set( "ProfilerCommand.StatsProfiler", new CORE_IMAGE_BRUSH( "Icons/Profiler/profiler_stats_40x", Icon40x40 ) );
-		Set( "ProfilerCommand.StatsProfiler.Small", new CORE_IMAGE_BRUSH( "Icons/Profiler/profiler_stats_40x", Icon20x20 ) );
-
-		Set( "ProfilerCommand.MemoryProfiler", new CORE_IMAGE_BRUSH( "Icons/Profiler/profiler_mem_40x", Icon40x40 ) );
-		Set( "ProfilerCommand.MemoryProfiler.Small", new CORE_IMAGE_BRUSH( "Icons/Profiler/profiler_mem_40x", Icon20x20 ) );
-
-		Set( "ProfilerCommand.FPSChart", new CORE_IMAGE_BRUSH( "Icons/Profiler/Profiler_FPS_Chart_40x", Icon40x40 ) );
-		Set( "ProfilerCommand.FPSChart.Small", new CORE_IMAGE_BRUSH( "Icons/Profiler/Profiler_FPS_Chart_40x", Icon20x20 ) );
-
-		Set( "ProfilerCommand.OpenSettings", new CORE_IMAGE_BRUSH( "Icons/Profiler/Profiler_Settings_40x", Icon40x40 ) );
-		Set( "ProfilerCommand.OpenSettings.Small", new CORE_IMAGE_BRUSH( "Icons/Profiler/Profiler_Settings_40x", Icon20x20 ) );
-
-		Set( "ProfilerCommand.ToggleDataPreview", new CORE_IMAGE_BRUSH( "Icons/Profiler/profiler_sync_40x", Icon40x40 ) );
-		Set( "ProfilerCommand.ToggleDataPreview.Small", new CORE_IMAGE_BRUSH( "Icons/Profiler/profiler_sync_40x", Icon20x20 ) );
-
-		Set( "ProfilerCommand.ToggleDataCapture", new CORE_IMAGE_BRUSH( "Icons/Profiler/Profiler_Data_Capture_40x", Icon40x40 ) );
-		Set( "ProfilerCommand.ToggleDataCapture.Small", new CORE_IMAGE_BRUSH( "Icons/Profiler/Profiler_Data_Capture_40x", Icon20x20 ) );
-
-		Set( "ProfilerCommand.ToggleDataCapture.Checked", new IMAGE_BRUSH( "Icons/icon_stop_40x", Icon40x40 ) );
-		Set( "ProfilerCommand.ToggleDataCapture.Checked.Small", new IMAGE_BRUSH( "Icons/icon_stop_40x", Icon20x20 ) );
-
-		Set( "ProfilerCommand.ToggleShowDataGraph", new CORE_IMAGE_BRUSH( "Icons/Profiler/profiler_ShowGraphData_32x", Icon32x32 ) );
-		Set( "ProfilerCommand.OpenEventGraph", new CORE_IMAGE_BRUSH( "Icons/Profiler/profiler_OpenEventGraph_32x", Icon16x16 ) );
-
 		// Generic
 		Set( "Profiler.LineGraphArea", new CORE_IMAGE_BRUSH( "Old/White", Icon16x16, FLinearColor(1.0f,1.0f,1.0f,0.25f) ) );
 		
diff --git a/Engine/Source/Runtime/SlateCore/Private/Styling/StarshipCoreStyle.cpp b/Engine/Source/Runtime/SlateCore/Private/Styling/StarshipCoreStyle.cpp
--- a/Engine/Source/Runtime/SlateCore/Private/Styling/StarshipCoreStyle.cpp
+++ b/Engine/Source/Runtime/SlateCore/Private/Styling/StarshipCoreStyle.cpp
@@ -412,6 +412,7 @@ TSharedRef<ISlateStyle> FStarshipCoreStyle::Create()
 		Style->Set("ToolPanel.LightGroupBorder", new BOX_BRUSH("Common/LightGroupBorder", FMargin(4.0f / 16.0f)));
 
 		Style->Set("Debug.Border", new BOX_BRUSH("Common/DebugBorder", 4.0f / 16.0f));
+		Style->Set("PlainBorder", new BORDER_BRUSH("Common/PlainBorder", 2.f / 8.f));
 	}
 
 	// Popup text
