Minor fix for missing icon

diff --git a/Engine/Source/Editor/EditorStyle/Private/StarshipStyle.cpp b/Engine/Source/Editor/EditorStyle/Private/StarshipStyle.cpp
--- a/Engine/Source/Editor/EditorStyle/Private/StarshipStyle.cpp
+++ b/Engine/Source/Editor/EditorStyle/Private/StarshipStyle.cpp
@@ -4506,7 +4506,7 @@ void FStarshipEditorStyle::FStyle::SetupLevelEditorStyle()
 		Set("ShowFlagsMenu.SubMenu.Lumen", new IMAGE_BRUSH_SVG("Starship/Common/LightingFeatures", Icon16x16));
 		Set("ShowFlagsMenu.SubMenu.Developer", new IMAGE_BRUSH_SVG("Starship/Common/Developer", Icon16x16));
 		Set("ShowFlagsMenu.SubMenu.Visualize", new IMAGE_BRUSH_SVG("Starship/Common/Visualize", Icon16x16));
-		Set("ShowFlagsMenu.SubMenu.Advanced", new IMAGE_BRUSH_SVG("Starship/Common/Advanced", Icon16x16));
+		Set("ShowFlagsMenu.SubMenu.Advanced", new CORE_IMAGE_BRUSH_SVG("Starship/Common/Advanced", Icon16x16));
 
 		Set("ShowFlagsMenu.SubMenu.Volumes", new IMAGE_BRUSH_SVG("Starship/Common/Volume", Icon16x16));
 		Set("ShowFlagsMenu.SubMenu.Layers", new IMAGE_BRUSH_SVG("Starship/Common/Layers", Icon16x16));
