PR #9353: [Source Control] Add a dedicated set of icons for Plastic SCM (Contributed by [REDACTED_PERSON])

diff --git a/Engine/Plugins/Developer/PlasticSourceControl/Source/PlasticSourceControl/Private/PlasticSourceControlState.cpp b/Engine/Plugins/Developer/PlasticSourceControl/Source/PlasticSourceControl/Private/PlasticSourceControlState.cpp
--- a/Engine/Plugins/Developer/PlasticSourceControl/Source/PlasticSourceControl/Private/PlasticSourceControlState.cpp
+++ b/Engine/Plugins/Developer/PlasticSourceControl/Source/PlasticSourceControl/Private/PlasticSourceControlState.cpp
@@ -102,26 +102,29 @@ FSlateIcon FPlasticSourceControlState::GetIcon() const
 	switch (WorkspaceState)
 	{
 	case EWorkspaceState::CheckedOut:
-	case EWorkspaceState::Replaced: // Merged (waiting for checkin)
-		return FSlateIcon(FAppStyle::GetAppStyleSetName(), "Perforce.CheckedOut");
+	case EWorkspaceState::Replaced: // Merged (waiting for check-in)
+		return FSlateIcon(FAppStyle::GetAppStyleSetName(), "Plastic.CheckedOut");
+	case EWorkspaceState::Changed: // Changed but unchecked-out file custom color icon
+		return FSlateIcon(FAppStyle::GetAppStyleSetName(), "Plastic.Changed"); // custom
 	case EWorkspaceState::Added:
 	case EWorkspaceState::Copied:
-		return FSlateIcon(FAppStyle::GetAppStyleSetName(), "Perforce.OpenForAdd");
+		return FSlateIcon(FAppStyle::GetAppStyleSetName(), "Plastic.OpenForAdd");
 	case EWorkspaceState::Moved:
-		return FSlateIcon(FAppStyle::GetAppStyleSetName(), "Perforce.Branched");
+		return FSlateIcon(FAppStyle::GetAppStyleSetName(), "Plastic.Branched");
 	case EWorkspaceState::Deleted:
+		return FSlateIcon(FAppStyle::GetAppStyleSetName(), "Plastic.MarkedForDelete");
 	case EWorkspaceState::LocallyDeleted:
-		return FSlateIcon(FAppStyle::GetAppStyleSetName(), "Perforce.MarkedForDelete");
+		return FSlateIcon(FAppStyle::GetAppStyleSetName(), "Plastic.LocallyDeleted"); // custom
 	case EWorkspaceState::Conflicted:
-		return FSlateIcon(FAppStyle::GetAppStyleSetName(), "Perforce.NotAtHeadRevision");
+		return FSlateIcon(FAppStyle::GetAppStyleSetName(), "Plastic.Conflicted"); // custom
 	case EWorkspaceState::LockedByOther:
-		return FSlateIcon(FAppStyle::GetAppStyleSetName(), "Perforce.CheckedOutByOtherUser", NAME_None, "SourceControl.LockOverlay");
+		return FSlateIcon(FAppStyle::GetAppStyleSetName(), "Plastic.CheckedOutByOtherUser", NAME_None, "SourceControl.LockOverlay");
 	case EWorkspaceState::Private: // Not controlled
-	case EWorkspaceState::Changed: // Changed but unchecked-out file is in a certain way not controlled
-		return FSlateIcon(FAppStyle::GetAppStyleSetName(), "Perforce.NotInDepot");
-	case EWorkspaceState::Unknown:
+		return FSlateIcon(FAppStyle::GetAppStyleSetName(), "Plastic.NotInDepot");
 	case EWorkspaceState::Ignored:
-	case EWorkspaceState::Controlled: // (Unchanged) same as "Pristine" for Perforce (not checked out) ie no icon
+		return FSlateIcon(FAppStyle::GetAppStyleSetName(), "Plastic.Ignored"); // custom
+	case EWorkspaceState::Unknown:
+	case EWorkspaceState::Controlled: // Unchanged (not checked out) ie no icon
 	default:
 		return FSlateIcon();
 	}
diff --git a/Engine/Source/Editor/EditorStyle/Private/StarshipStyle.cpp b/Engine/Source/Editor/EditorStyle/Private/StarshipStyle.cpp
--- a/Engine/Source/Editor/EditorStyle/Private/StarshipStyle.cpp
+++ b/Engine/Source/Editor/EditorStyle/Private/StarshipStyle.cpp
@@ -7072,8 +7072,8 @@ void FStarshipEditorStyle::FStyle::SetupSourceControlStyles()
 	}
 #endif // WITH_EDITOR || (IS_PROGRAM && WITH_UNREAL_DEVELOPER_TOOLS)
 
-	// Perforce
 #if WITH_EDITOR || (IS_PROGRAM && WITH_UNREAL_DEVELOPER_TOOLS)
+	// Perforce
 	{
 		Set("Perforce.CheckedOut", new CORE_IMAGE_BRUSH_SVG("Starship/SourceControl/SCC_CheckedOut", Icon16x16, FStyleColors::AccentRed));
 		Set("Perforce.OpenForAdd", new CORE_IMAGE_BRUSH_SVG("Starship/SourceControl/SCC_ContentAdd", Icon16x16, FStyleColors::AccentRed));
@@ -7084,7 +7084,22 @@ void FStarshipEditorStyle::FStyle::SetupSourceControlStyles()
 		Set("Perforce.NotAtHeadRevision", new CORE_IMAGE_BRUSH_SVG("Starship/SourceControl/SCC_ModifiedOtherBranch", Icon16x16, FStyleColors::AccentYellow));
 		Set("Perforce.NotInDepot", new CORE_IMAGE_BRUSH_SVG("Starship/SourceControl/SCC_NotInDepot", Icon16x16, FStyleColors::AccentYellow));
 		Set("Perforce.Branched", new CORE_IMAGE_BRUSH_SVG("Starship/SourceControl/SCC_Branched", Icon16x16, FStyleColors::AccentGreen));
-
+	}
+	// Plastic SCM
+	{
+		Set("Plastic.CheckedOut", new CORE_IMAGE_BRUSH_SVG("Starship/SourceControl/SCC_CheckedOut", Icon16x16, FStyleColors::AccentRed));
+		Set("Plastic.Changed", new CORE_IMAGE_BRUSH_SVG("Starship/SourceControl/SCC_CheckedOut", Icon16x16, FStyleColors::AccentYellow)); // custom
+		Set("Plastic.OpenForAdd", new CORE_IMAGE_BRUSH_SVG("Starship/SourceControl/SCC_ContentAdd", Icon16x16, FStyleColors::AccentRed));
+		Set("Plastic.CheckedOutByOtherUser", new CORE_IMAGE_BRUSH_SVG("Starship/SourceControl/SCC_CheckedOut", Icon16x16, FStyleColors::AccentYellow));
+		Set("Plastic.ModifiedOtherBranch", new CORE_IMAGE_BRUSH_SVG("Starship/SourceControl/SCC_ModifiedOtherBranch", Icon16x16, FStyleColors::AccentRed));
+		Set("Plastic.MarkedForDelete", new CORE_IMAGE_BRUSH_SVG("Starship/SourceControl/SCC_MarkedForDelete", Icon16x16, FStyleColors::AccentRed));
+		Set("Plastic.LocallyDeleted", new CORE_IMAGE_BRUSH_SVG("Starship/SourceControl/SCC_MarkedForDelete", Icon16x16, FStyleColors::AccentYellow)); // custom
+		Set("Plastic.NotAtHeadRevision", new CORE_IMAGE_BRUSH_SVG("Starship/SourceControl/SCC_ModifiedOtherBranch", Icon16x16, FStyleColors::AccentYellow));
+		Set("Plastic.Conflicted", new CORE_IMAGE_BRUSH_SVG("Starship/SourceControl/SCC_ModifiedOtherBranch", Icon16x16, FStyleColors::AccentRed)); // custom
+		Set("Plastic.NotInDepot", new CORE_IMAGE_BRUSH_SVG("Starship/SourceControl/SCC_NotInDepot", Icon16x16, FStyleColors::AccentYellow));
+		Set("Plastic.Ignored", new CORE_IMAGE_BRUSH_SVG("Starship/SourceControl/SCC_NotInDepot", Icon16x16, FStyleColors::AccentYellow)); // custom
+		Set("Plastic.Branched", new CORE_IMAGE_BRUSH_SVG("Starship/SourceControl/SCC_Branched", Icon16x16, FStyleColors::AccentGreen));
+		Set("Plastic.LocallyMoved", new CORE_IMAGE_BRUSH_SVG("Starship/SourceControl/SCC_Branched", Icon16x16, FStyleColors::AccentYellow)); // custom
 	}
 	// Subversion
 	{
