Fix missing Output log icon in some places

diff --git a/Engine/Source/Editor/EditorStyle/Private/StarshipStyle.cpp b/Engine/Source/Editor/EditorStyle/Private/StarshipStyle.cpp
--- a/Engine/Source/Editor/EditorStyle/Private/StarshipStyle.cpp
+++ b/Engine/Source/Editor/EditorStyle/Private/StarshipStyle.cpp
@@ -4861,7 +4861,7 @@ void FStarshipEditorStyle::FStyle::SetupPersonaStyle()
 		{
 			Set( "Kismet.Tabs.Variables", new IMAGE_BRUSH_SVG( "Starship/Blueprints/pill", Icon16x16 ) );
 			Set( "Kismet.Tabs.Palette", new IMAGE_BRUSH_SVG( "Starship/Blueprints/Palette", Icon16x16 ) );
-			Set( "Kismet.Tabs.CompilerResults", new IMAGE_BRUSH_SVG("Starship/Common/OutputLog", Icon16x16 ) );
+			Set( "Kismet.Tabs.CompilerResults", new CORE_IMAGE_BRUSH_SVG("Starship/Common/OutputLog", Icon16x16 ) );
 			Set( "Kismet.Tabs.FindResults", new CORE_IMAGE_BRUSH_SVG("Starship/Common/search", Icon16x16 ) );
 			Set( "Kismet.Tabs.Bookmarks", new IMAGE_BRUSH_SVG( "Starship/Common/Bookmarks", Icon16x16 ) );
 			Set( "Kismet.Tabs.Components", new IMAGE_BRUSH_SVG( "Starship/Common/Component", Icon16x16 ) );
