Sequencer: Added missing updated graphics

diff --git a/Engine/Source/Editor/EditorStyle/Private/StarshipStyle.cpp b/Engine/Source/Editor/EditorStyle/Private/StarshipStyle.cpp
--- a/Engine/Source/Editor/EditorStyle/Private/StarshipStyle.cpp
+++ b/Engine/Source/Editor/EditorStyle/Private/StarshipStyle.cpp
@@ -1946,7 +1946,7 @@ void FStarshipEditorStyle::FStyle::SetupSequencerStyles()
 
 		Set("SequencerToolBar", SequencerToolbar);
 
-		Set("Sequencer.IconKeyAuto", new IMAGE_BRUSH_SVG("Sequencer/Keys/Key_CubicAuto", Icon12x12));
+		Set("Sequencer.IconKeyAuto", new IMAGE_BRUSH("Sequencer/IconKeyAuto", Icon12x12));
 		Set("Sequencer.IconKeyBreak", new IMAGE_BRUSH("Sequencer/IconKeyBreak", Icon12x12));
 		Set("Sequencer.IconKeyConstant", new IMAGE_BRUSH("Sequencer/IconKeyConstant", Icon12x12));
 		Set("Sequencer.IconKeyLinear", new IMAGE_BRUSH("Sequencer/IconKeyLinear", Icon12x12));
