PR #8664: Update StarshipStyle spelling in comment (Contributed by [REDACTED_PERSON])

diff --git a/Engine/Source/Editor/EditorStyle/Private/StarshipStyle.h b/Engine/Source/Editor/EditorStyle/Private/StarshipStyle.h
--- a/Engine/Source/Editor/EditorStyle/Private/StarshipStyle.h
+++ b/Engine/Source/Editor/EditorStyle/Private/StarshipStyle.h
@@ -122,7 +122,7 @@ public:
 		const FSlateColor HighlightColor;
 		const FSlateColor WindowHighlightColor;
 
-		// These are common colors used thruout the editor in mutliple style elements
+		// These are common colors used throughout the editor in multiple style elements
 		const FSlateColor InheritedFromBlueprintTextColor;
 
 		// Styles inherited from the parent style
