Sequencer: Icon updates for Sequencer and Cinematic Prestreaming

diff --git a/Engine/Plugins/Experimental/CinematicPrestreaming/Source/CinematicPrestreamingEditor/Private/Styles/CinePrestreamingEditorStyle.cpp b/Engine/Plugins/Experimental/CinematicPrestreaming/Source/CinematicPrestreamingEditor/Private/Styles/CinePrestreamingEditorStyle.cpp
new file mode 100644
--- /dev/null
+++ b/Engine/Plugins/Experimental/CinematicPrestreaming/Source/CinematicPrestreamingEditor/Private/Styles/CinePrestreamingEditorStyle.cpp
@@ -0,0 +1,5 @@
+// Copyright Epic Games, Inc. All Rights Reserved.
+
+#include "Styles/CinePrestreamingEditorStyle.h"
+
+TSharedPtr<FCinePrestreamingEditorStyle> FCinePrestreamingEditorStyle::Singleton;
diff --git a/Engine/Plugins/Experimental/CinematicPrestreaming/Source/CinematicPrestreamingEditor/Private/Styles/CinePrestreamingEditorStyle.h b/Engine/Plugins/Experimental/CinematicPrestreaming/Source/CinematicPrestreamingEditor/Private/Styles/CinePrestreamingEditorStyle.h
new file mode 100644
--- /dev/null
+++ b/Engine/Plugins/Experimental/CinematicPrestreaming/Source/CinematicPrestreamingEditor/Private/Styles/CinePrestreamingEditorStyle.h
@@ -0,0 +1,59 @@
+// Copyright Epic Games, Inc. All Rights Reserved.
+
+#pragma once
+
+#include "CoreMinimal.h"
+#include "Layout/Margin.h"
+#include "Fonts/SlateFontInfo.h"
+#include "Misc/Paths.h"
+#include "Brushes/SlateBorderBrush.h"
+#include "Brushes/SlateBoxBrush.h"
+#include "Brushes/SlateImageBrush.h"
+#include "Styling/SlateStyleRegistry.h"
+#include "Styling/SlateStyle.h"
+
+#define IMAGE_BRUSH_SVG(RelativePath, ...) FSlateVectorImageBrush(RootToContentDir(RelativePath, TEXT(".svg")), __VA_ARGS__)
+
+/**
+ * Implements the visual style of the cinematic prestreaming sequence editor.
+ */
+class FCinePrestreamingEditorStyle final : public FSlateStyleSet
+{
+public:
+	
+	FCinePrestreamingEditorStyle()
+		: FSlateStyleSet("CinePrestreamingEditorStyle")
+	{
+		const FVector2D Icon16x16(16.0f, 16.0f);
+		const FVector2D Icon24x24(24.0f, 24.0f);
+		const FVector2D Icon48x48(48.0f, 48.0f);
+		const FVector2D Icon64x64(64.0f, 64.0f);
+
+		SetContentRoot(FPaths::EnginePluginsDir() / TEXT("Experimental/CinematicPrestreaming/Content"));
+
+		Set("Sequencer.Tracks.CinePrestreaming_16", new IMAGE_BRUSH_SVG("CinePrestreaming_16", Icon16x16));
+		Set("Sequencer.Tracks.CinePrestreaming_64", new IMAGE_BRUSH_SVG("CinePrestreaming_64", Icon64x64));
+
+		FSlateStyleRegistry::RegisterSlateStyle(*this);
+	}
+
+	virtual ~FCinePrestreamingEditorStyle()
+	{
+		FSlateStyleRegistry::UnRegisterSlateStyle(*this);
+	}
+
+	static TSharedRef<FCinePrestreamingEditorStyle> Get()
+	{
+		if (!Singleton.IsValid())
+		{
+			Singleton = MakeShareable(new FCinePrestreamingEditorStyle);
+		}
+		return Singleton.ToSharedRef();
+	}
+
+private:
+
+	static TSharedPtr<FCinePrestreamingEditorStyle> Singleton;
+};
+
+#undef IMAGE_BRUSH_SVG
diff --git a/Engine/Plugins/Experimental/CinematicPrestreaming/Source/CinematicPrestreamingEditor/Private/TrackEditors/CinePrestreamingTrackEditor.cpp b/Engine/Plugins/Experimental/CinematicPrestreaming/Source/CinematicPrestreamingEditor/Private/TrackEditors/CinePrestreamingTrackEditor.cpp
--- a/Engine/Plugins/Experimental/CinematicPrestreaming/Source/CinematicPrestreamingEditor/Private/TrackEditors/CinePrestreamingTrackEditor.cpp
+++ b/Engine/Plugins/Experimental/CinematicPrestreaming/Source/CinematicPrestreamingEditor/Private/TrackEditors/CinePrestreamingTrackEditor.cpp
@@ -3,6 +3,7 @@
 #include "TrackEditors/CinePrestreamingTrackEditor.h"
 
 #include "CinePrestreamingData.h"
+#include "Styles/CinePrestreamingEditorStyle.h"
 #include "Styling/AppStyle.h"
 #include "Framework/MultiBox/MultiBoxBuilder.h"
 #include "ISequencerSection.h"
@@ -162,7 +163,7 @@ bool FCinePrestreamingTrackEditor::SupportsType(TSubclassOf<UMovieSceneTrack> Ty
 
 const FSlateBrush* FCinePrestreamingTrackEditor::GetIconBrush() const
 {
-	return FAppStyle::Get().GetBrush("Sequencer.Tracks.DataLayer");
+	return FCinePrestreamingEditorStyle::Get()->GetBrush("Sequencer.Tracks.CinePrestreaming_16");
 }
 
 TSharedRef<ISequencerSection> FCinePrestreamingTrackEditor::MakeSectionInterface(UMovieSceneSection& SectionObject, UMovieSceneTrack& Track, FGuid ObjectBinding)
@@ -178,7 +179,7 @@ void FCinePrestreamingTrackEditor::BuildAddTrackMenu(FMenuBuilder& MenuBuilder)
 	MenuBuilder.AddMenuEntry(
 		LOCTEXT("AddTrack", "Cinematic Prestreaming"),
 		LOCTEXT("AddTrackToolTip", "Adds a new track that can load streaming data used by the renderer to request things in advanced of needing them."),
-		FSlateIcon(FAppStyle::GetAppStyleSetName(), "Sequencer.Tracks.DataLayer"),
+		FSlateIcon(FCinePrestreamingEditorStyle::Get()->GetStyleSetName(), "Sequencer.Tracks.CinePrestreaming_16"),
 		FUIAction(FExecuteAction::CreateRaw(this, &FCinePrestreamingTrackEditor::HandleAddTrack)));
 }
 
diff --git a/Engine/Plugins/MovieScene/LevelSequenceEditor/Source/LevelSequenceEditor/Private/Styles/LevelSequenceEditorStyle.h b/Engine/Plugins/MovieScene/LevelSequenceEditor/Source/LevelSequenceEditor/Private/Styles/LevelSequenceEditorStyle.h
--- a/Engine/Plugins/MovieScene/LevelSequenceEditor/Source/LevelSequenceEditor/Private/Styles/LevelSequenceEditorStyle.h
+++ b/Engine/Plugins/MovieScene/LevelSequenceEditor/Source/LevelSequenceEditor/Private/Styles/LevelSequenceEditorStyle.h
@@ -35,8 +35,8 @@ public:
 		// tab icons
 		Set("LevelSequenceEditor.Tabs.Sequencer", new IMAGE_BRUSH("icon_tab_sequencer_16x", Icon16x16));
 
-		Set("LevelSequenceEditor.PossessNewActor", new IMAGE_BRUSH("Icon_Actor_To_Sequencer_16x", Icon16x16));
-		Set("LevelSequenceEditor.PossessNewActor.Small", new IMAGE_BRUSH("Icon_Actor_To_Sequencer_16x", Icon16x16));
+		Set("LevelSequenceEditor.PossessNewActor", new IMAGE_BRUSH_SVG("ActorToSequencer", Icon16x16));
+		Set("LevelSequenceEditor.PossessNewActor.Small", new IMAGE_BRUSH_SVG("ActorToSequencer", Icon16x16));
 
 		Set("LevelSequenceEditor.CreateNewLevelSequenceInLevel", new IMAGE_BRUSH_SVG("LevelSequence", Icon16x16));
 		Set("LevelSequenceEditor.CreateNewLevelSequenceInLevel.Small", new IMAGE_BRUSH_SVG("LevelSequence", Icon16x16));
diff --git a/Engine/Source/Editor/EditorStyle/Private/StarshipStyle.cpp b/Engine/Source/Editor/EditorStyle/Private/StarshipStyle.cpp
--- a/Engine/Source/Editor/EditorStyle/Private/StarshipStyle.cpp
+++ b/Engine/Source/Editor/EditorStyle/Private/StarshipStyle.cpp
@@ -2115,13 +2115,13 @@ void FStarshipEditorStyle::FStyle::SetupSequencerStyles()
 
 		Set("Sequencer.Tracks.Media", new IMAGE_BRUSH("Sequencer/Dropdown_Icons/Icon_Media_Track_16x", Icon16x16));
 		Set("Sequencer.Tracks.Audio", new IMAGE_BRUSH("Sequencer/Dropdown_Icons/Icon_Audio_Track_16x", Icon16x16));
-		Set("Sequencer.Tracks.Event", new IMAGE_BRUSH("Sequencer/Dropdown_Icons/Icon_Event_Track_16x", Icon16x16));
-		Set("Sequencer.Tracks.Fade", new IMAGE_BRUSH("Sequencer/Dropdown_Icons/Icon_Fade_Track_16x", Icon16x16));
-		Set("Sequencer.Tracks.CameraCut", new IMAGE_BRUSH("Sequencer/Dropdown_Icons/Icon_Camera_Cut_Track_16x", Icon16x16));
-		Set("Sequencer.Tracks.CinematicShot", new IMAGE_BRUSH("Sequencer/Dropdown_Icons/Icon_Shot_Track_16x", Icon16x16));
-		Set("Sequencer.Tracks.Slomo", new IMAGE_BRUSH("Sequencer/Dropdown_Icons/Icon_Play_Rate_Track_16x", Icon16x16));
-		Set("Sequencer.Tracks.Sub", new IMAGE_BRUSH("Sequencer/Dropdown_Icons/Icon_Sub_Track_16x", Icon16x16));
-		Set("Sequencer.Tracks.LevelVisibility", new IMAGE_BRUSH("Sequencer/Dropdown_Icons/Icon_Level_Visibility_Track_16x", Icon16x16));
+		Set("Sequencer.Tracks.Event", new IMAGE_BRUSH_SVG("Starship/Sequencer/EventTrack", Icon16x16));
+		Set("Sequencer.Tracks.Fade", new IMAGE_BRUSH_SVG("Starship/Sequencer/FadeTrack", Icon16x16));
+		Set("Sequencer.Tracks.CameraCut", new IMAGE_BRUSH_SVG("Starship/Sequencer/CameraCutTrack", Icon16x16));
+		Set("Sequencer.Tracks.CinematicShot", new IMAGE_BRUSH_SVG("Starship/Sequencer/ShotTrack", Icon16x16));
+		Set("Sequencer.Tracks.Slomo", new IMAGE_BRUSH_SVG("Starship/Sequencer/SlomoTrack", Icon16x16));
+		Set("Sequencer.Tracks.Sub", new IMAGE_BRUSH_SVG("Starship/Sequencer/SubTrack", Icon16x16));
+		Set("Sequencer.Tracks.LevelVisibility", new IMAGE_BRUSH_SVG("Starship/Sequencer/LevelVisibilityTrack", Icon16x16));
 		Set("Sequencer.Tracks.DataLayer", new IMAGE_BRUSH_SVG("Starship/Common/DataLayers", Icon16x16));
 		Set("Sequencer.Tracks.CVar", new CORE_IMAGE_BRUSH_SVG("Starship/Common/Console", Icon16x16));
 
