Sequencer: More icon updates

diff --git a/Engine/Source/Editor/EditorStyle/Private/StarshipStyle.cpp b/Engine/Source/Editor/EditorStyle/Private/StarshipStyle.cpp
--- a/Engine/Source/Editor/EditorStyle/Private/StarshipStyle.cpp
+++ b/Engine/Source/Editor/EditorStyle/Private/StarshipStyle.cpp
@@ -2113,8 +2113,8 @@ void FStarshipEditorStyle::FStyle::SetupSequencerStyles()
 
 		Set("Sequencer.TrackArea.LaneColor", FLinearColor(0.3f, 0.3f, 0.3f, 0.3f));
 
-		Set("Sequencer.Tracks.Media", new IMAGE_BRUSH("Sequencer/Dropdown_Icons/Icon_Media_Track_16x", Icon16x16));
-		Set("Sequencer.Tracks.Audio", new IMAGE_BRUSH("Sequencer/Dropdown_Icons/Icon_Audio_Track_16x", Icon16x16));
+		Set("Sequencer.Tracks.Media", new IMAGE_BRUSH_SVG("Starship/AssetIcons/MediaPlayer_16", Icon16x16));
+		Set("Sequencer.Tracks.Audio", new IMAGE_BRUSH_SVG("Starship/AssetIcons/AmbientSound_16", Icon16x16));
 		Set("Sequencer.Tracks.Event", new IMAGE_BRUSH_SVG("Starship/Sequencer/EventTrack", Icon16x16));
 		Set("Sequencer.Tracks.Fade", new IMAGE_BRUSH_SVG("Starship/Sequencer/FadeTrack", Icon16x16));
 		Set("Sequencer.Tracks.CameraCut", new IMAGE_BRUSH_SVG("Starship/Sequencer/CameraCutTrack", Icon16x16));
