SBasicFilterBar: Move styles references in Editor only configs to StarshipCoreStyle

diff --git a/Engine/Source/Developer/ToolWidgets/Public/Filters/SBasicFilterBar.h b/Engine/Source/Developer/ToolWidgets/Public/Filters/SBasicFilterBar.h
--- a/Engine/Source/Developer/ToolWidgets/Public/Filters/SBasicFilterBar.h
+++ b/Engine/Source/Developer/ToolWidgets/Public/Filters/SBasicFilterBar.h
@@ -470,7 +470,7 @@ protected:
 			{
 			case EFilterPillStyle::Basic:
 				ContentWidget = SAssignNew( ToggleButtonPtr, SFilterCheckBox )
-					.Style(FAppStyle::Get(), "DetailsView.SectionButton")
+					.Style(FAppStyle::Get(), "FilterBar.BasicFilterButton")
 					.ToolTipText(FilterToolTip)
 					.IsChecked(this, &SFilter::IsChecked)
 					.OnCheckStateChanged(this, &SFilter::FilterToggled)
@@ -487,7 +487,7 @@ protected:
 			default:
 				ContentWidget = SNew(SBorder)
 				 .Padding(1.0f)
-				 .BorderImage(FAppStyle::Get().GetBrush("ContentBrowser.FilterBackground"))
+				 .BorderImage(FAppStyle::Get().GetBrush("FilterBar.FilterBackground"))
 				 [
 					 SNew(SHorizontalBox)
 					 +SHorizontalBox::Slot()
@@ -495,7 +495,7 @@ protected:
 					 .AutoWidth()
 					 [
 						 SNew(SImage)
-						 .Image(FAppStyle::Get().GetBrush("ContentBrowser.FilterImage"))
+						 .Image(FAppStyle::Get().GetBrush("FilterBar.FilterImage"))
 						 .ColorAndOpacity(this, &SFilter::GetFilterImageColorAndOpacity)
 					 ]
 					 +SHorizontalBox::Slot()
@@ -503,7 +503,7 @@ protected:
 					 .VAlign(VAlign_Center)
 					 [
 						 SAssignNew( ToggleButtonPtr, SFilterCheckBox )
-						.Style(FAppStyle::Get(), "ContentBrowser.FilterButton")
+						.Style(FAppStyle::Get(), "FilterBar.FilterButton")
 						.ToolTipText(FilterToolTip)
 						.IsChecked(this, &SFilter::IsChecked)
 						.OnCheckStateChanged(this, &SFilter::FilterToggled)
diff --git a/Engine/Source/Editor/EditorStyle/Private/StarshipStyle.cpp b/Engine/Source/Editor/EditorStyle/Private/StarshipStyle.cpp
--- a/Engine/Source/Editor/EditorStyle/Private/StarshipStyle.cpp
+++ b/Engine/Source/Editor/EditorStyle/Private/StarshipStyle.cpp
@@ -2917,15 +2917,7 @@ void FStarshipEditorStyle::FStyle::SetupPropertyEditorStyles()
 		);
 
 		Set( "DetailsView.GridLine", new FSlateColorBrush(FStyleColors::Recessed) );
-		Set( "DetailsView.SectionButton", FCheckBoxStyle( FStarshipCoreStyle::GetCoreStyle().GetWidgetStyle<FCheckBoxStyle>("ToggleButtonCheckbox"))
-			.SetUncheckedImage(FSlateRoundedBoxBrush(FStyleColors::Header, 4.0f, FStyleColors::Input, 1.0f))
-			.SetUncheckedHoveredImage(FSlateRoundedBoxBrush(FStyleColors::Hover, 4.0f, FStyleColors::Input, 1.0f))
-			.SetUncheckedPressedImage(FSlateRoundedBoxBrush(FStyleColors::Hover, 4.0f, FStyleColors::Input, 1.0f))
-			.SetCheckedImage(FSlateRoundedBoxBrush(FStyleColors::Primary, 4.0f, FStyleColors::Input, 1.0f))
-			.SetCheckedHoveredImage(FSlateRoundedBoxBrush(FStyleColors::PrimaryHover, 4.0f, FStyleColors::Input, 1.0f))
-			.SetCheckedPressedImage(FSlateRoundedBoxBrush(FStyleColors::PrimaryHover, 4.0f, FStyleColors::Input, 1.0f))
-			.SetPadding(FMargin(16, 4))
-		);
+		Set( "DetailsView.SectionButton", FCheckBoxStyle( FStarshipCoreStyle::GetCoreStyle().GetWidgetStyle<FCheckBoxStyle>("FilterBar.BasicFilterButton")));
 
 		Set( "DetailsView.ChannelToggleButton", FCheckBoxStyle( FStarshipCoreStyle::GetCoreStyle().GetWidgetStyle<FCheckBoxStyle>("ToggleButtonCheckbox"))
 			.SetUncheckedImage(FSlateRoundedBoxBrush(FStyleColors::Input, 4.0f, FStyleColors::DropdownOutline, 1.0f))
@@ -5994,43 +5986,12 @@ void FStarshipEditorStyle::FStyle::SetupContentBrowserStyle()
 		Set( "ContentBrowser.SortUp", new IMAGE_BRUSH( "Common/SortUpArrow", Icon8x4 ) );
 		Set( "ContentBrowser.SortDown", new IMAGE_BRUSH( "Common/SortDownArrow", Icon8x4 ) );
 
-		// Filter list
-		/* Set images for various SCheckBox states associated with "ContentBrowser.FilterButton" ... */
-		const FCheckBoxStyle ContentBrowserFilterButtonCheckBoxStyle = FCheckBoxStyle()
-			.SetUncheckedImage(FSlateNoResource())
-			.SetUncheckedHoveredImage(FSlateNoResource())
-			.SetUncheckedPressedImage(FSlateNoResource())
-			.SetCheckedImage(FSlateNoResource())
-			.SetCheckedHoveredImage(FSlateNoResource())
-			.SetCheckedPressedImage(FSlateNoResource())
-			.SetForegroundColor(FStyleColors::Foreground)
-			.SetHoveredForegroundColor(FStyleColors::ForegroundHover)
-			.SetCheckedForegroundColor(FStyleColors::Foreground)
-			.SetCheckedHoveredForegroundColor(FStyleColors::ForegroundHover)
-			.SetCheckedPressedForegroundColor(FStyleColors::ForegroundHover)
-			.SetPressedForegroundColor(FStyleColors::ForegroundHover);
-
-		Set("ContentBrowser.FilterImage", new IMAGE_BRUSH_SVG("Starship/ContentBrowser/AssetFilterColorSegment", FVector2D(8, 22)));
+		// Filter List - These are aliases for SBasicFilterBar styles in StarshipCoreStyle for backwards compatibility
+		Set("ContentBrowser.FilterImage", new CORE_IMAGE_BRUSH_SVG("Starship/CoreWidgets/FilterBar/FilterColorSegment", FVector2D(8, 22)));
 		Set("ContentBrowser.FilterBackground", new FSlateRoundedBoxBrush(FStyleColors::Secondary, 3.0f));
 
-		/* ... and add the new style */
-		Set("ContentBrowser.FilterButton", ContentBrowserFilterButtonCheckBoxStyle );
-
-		// Filter list implemented as toolbar
-		FToolBarStyle FilterToolBar = FToolBarStyle(FStarshipCoreStyle::GetCoreStyle().GetWidgetStyle<FToolBarStyle>("SlimToolBar"))
-			.SetIconSize(Icon16x16)
-			.SetBackground(FSlateNoResource())
-			.SetLabelPadding(FMargin(0))
-			.SetComboButtonPadding(FMargin(0))
-			.SetBlockPadding(FMargin(3, 0))
-			.SetIndentedBlockPadding(FMargin(0))
-			.SetBackgroundPadding(FMargin(0))
-			.SetButtonPadding(FMargin(0))
-			.SetCheckBoxPadding(FMargin(0))
-			.SetSeparatorBrush(FSlateNoResource())
-			.SetSeparatorPadding(FMargin(0));
-
-		Set("ContentBrowser.FilterToolBar", FilterToolBar);
+		Set("ContentBrowser.FilterButton", FStarshipCoreStyle::GetCoreStyle().GetWidgetStyle<FCheckBoxStyle>("FilterBar.FilterButton"));
+		Set("ContentBrowser.FilterToolBar", FStarshipCoreStyle::GetCoreStyle().GetWidgetStyle<FToolBarStyle>("FilterBar.FilterToolBar"));
 
 		// Sources view
 		Set("ContentBrowser.Sources.Paths", new IMAGE_BRUSH("ContentBrowser/Sources_Paths_16x", Icon16x16));
diff --git a/Engine/Source/Runtime/SlateCore/Private/Styling/StarshipCoreStyle.cpp b/Engine/Source/Runtime/SlateCore/Private/Styling/StarshipCoreStyle.cpp
--- a/Engine/Source/Runtime/SlateCore/Private/Styling/StarshipCoreStyle.cpp
+++ b/Engine/Source/Runtime/SlateCore/Private/Styling/StarshipCoreStyle.cpp
@@ -852,6 +852,52 @@ TSharedRef<ISlateStyle> FStarshipCoreStyle::Create()
 		FilterSearchBoxStyle.SetTextBoxStyle(FilterTextBoxStyle);
 
 		Style->Set("FilterBar.SearchBox", FilterSearchBoxStyle);
+		
+		const FCheckBoxStyle FilterButtonCheckBoxStyle = FCheckBoxStyle()
+			.SetUncheckedImage(FSlateNoResource())
+			.SetUncheckedHoveredImage(FSlateNoResource())
+			.SetUncheckedPressedImage(FSlateNoResource())
+			.SetCheckedImage(FSlateNoResource())
+			.SetCheckedHoveredImage(FSlateNoResource())
+			.SetCheckedPressedImage(FSlateNoResource())
+			.SetForegroundColor(FStyleColors::Foreground)
+			.SetHoveredForegroundColor(FStyleColors::ForegroundHover)
+			.SetCheckedForegroundColor(FStyleColors::Foreground)
+			.SetCheckedHoveredForegroundColor(FStyleColors::ForegroundHover)
+			.SetCheckedPressedForegroundColor(FStyleColors::ForegroundHover)
+			.SetPressedForegroundColor(FStyleColors::ForegroundHover);
+
+		Style->Set("FilterBar.FilterImage", new IMAGE_BRUSH_SVG("Starship/CoreWidgets/FilterBar/FilterColorSegment", FVector2D(8, 22)));
+		Style->Set("FilterBar.FilterBackground", new FSlateRoundedBoxBrush(FStyleColors::Secondary, 3.0f));
+
+		/* ... and add the new style */
+		Style->Set("FilterBar.FilterButton", FilterButtonCheckBoxStyle );
+		
+		Style->Set( "FilterBar.BasicFilterButton", FCheckBoxStyle( Style->GetWidgetStyle<FCheckBoxStyle>("ToggleButtonCheckbox"))
+			.SetUncheckedImage(FSlateRoundedBoxBrush(FStyleColors::Header, 4.0f, FStyleColors::Input, 1.0f))
+			.SetUncheckedHoveredImage(FSlateRoundedBoxBrush(FStyleColors::Hover, 4.0f, FStyleColors::Input, 1.0f))
+			.SetUncheckedPressedImage(FSlateRoundedBoxBrush(FStyleColors::Hover, 4.0f, FStyleColors::Input, 1.0f))
+			.SetCheckedImage(FSlateRoundedBoxBrush(FStyleColors::Primary, 4.0f, FStyleColors::Input, 1.0f))
+			.SetCheckedHoveredImage(FSlateRoundedBoxBrush(FStyleColors::PrimaryHover, 4.0f, FStyleColors::Input, 1.0f))
+			.SetCheckedPressedImage(FSlateRoundedBoxBrush(FStyleColors::PrimaryHover, 4.0f, FStyleColors::Input, 1.0f))
+			.SetPadding(FMargin(16, 4))
+		);
+
+		// FilteBar implemented as toolbar
+		FToolBarStyle FilterToolBar = FToolBarStyle(Style->GetWidgetStyle<FToolBarStyle>("SlimToolBar"))
+			.SetIconSize(Icon16x16)
+			.SetBackground(FSlateNoResource())
+			.SetLabelPadding(FMargin(0))
+			.SetComboButtonPadding(FMargin(0))
+			.SetBlockPadding(FMargin(3, 0))
+			.SetIndentedBlockPadding(FMargin(0))
+			.SetBackgroundPadding(FMargin(0))
+			.SetButtonPadding(FMargin(0))
+			.SetCheckBoxPadding(FMargin(0))
+			.SetSeparatorBrush(FSlateNoResource())
+			.SetSeparatorPadding(FMargin(0));
+
+		Style->Set("FilterBar.FilterToolBar", FilterToolBar);
 	}
 	
 	// Project Launcher
