Update pose watch icon
* Deprecated UPoseWatch::IconName as only pose watch (node watch) components can have icons

diff --git a/Engine/Source/Editor/EditorStyle/Private/StarshipStyle.cpp b/Engine/Source/Editor/EditorStyle/Private/StarshipStyle.cpp
--- a/Engine/Source/Editor/EditorStyle/Private/StarshipStyle.cpp
+++ b/Engine/Source/Editor/EditorStyle/Private/StarshipStyle.cpp
@@ -5557,6 +5557,8 @@ void FStarshipEditorStyle::FStyle::SetupPersonaStyle()
 		Set("AnimGraph.Attribute.InertialBlending.Icon", new IMAGE_BRUSH_SVG("Starship/AnimationBlueprintEditor/AnimationGraphInertialBlending", Icon16x16));
 		Set("AnimGraph.Attribute.InertialBlending.Color", FStyleColors::AccentOrange);
 
+		Set("AnimGraph.PoseWatch.Icon", new IMAGE_BRUSH_SVG("Starship/AnimationBlueprintEditor/AnimationGraphPoseWatch", Icon16x16));
+
 		Set("AnimGraph.AnimNodeReference.Subtitle", FTextBlockStyle(NormalText)
 			.SetFont(DEFAULT_FONT("Italic", 8))
 			.SetColorAndOpacity(FLinearColor(218.0f/255.0f,218.0f/255.0f,96.0f/255.0f, 0.5f))
diff --git a/Engine/Source/Editor/Persona/Private/PoseWatchManagerPoseWatchTreeItem.cpp b/Engine/Source/Editor/Persona/Private/PoseWatchManagerPoseWatchTreeItem.cpp
--- a/Engine/Source/Editor/Persona/Private/PoseWatchManagerPoseWatchTreeItem.cpp
+++ b/Engine/Source/Editor/Persona/Private/PoseWatchManagerPoseWatchTreeItem.cpp
@@ -95,7 +95,7 @@ private:
 
 	const FSlateBrush* GetIcon() const
 	{
-		return FAppStyle::Get().GetBrush(TEXT("ClassIcon.PoseAsset"));
+		return FAppStyle::Get().GetBrush("AnimGraph.PoseWatch.Icon");
 	}
 
 	const FSlateBrush* GetIconOverlay() const
diff --git a/Engine/Source/Editor/Persona/Private/TabSpawners.cpp b/Engine/Source/Editor/Persona/Private/TabSpawners.cpp
--- a/Engine/Source/Editor/Persona/Private/TabSpawners.cpp
+++ b/Engine/Source/Editor/Persona/Private/TabSpawners.cpp
@@ -583,7 +583,7 @@ FPoseWatchManagerSummoner::FPoseWatchManagerSummoner(TSharedPtr<class FBlueprint
 	, BlueprintEditor(InBlueprintEditor)
 {
 	TabLabel = LOCTEXT("PoseWatchManager", "Pose Watch Manager");
-	TabIcon = FSlateIcon(FAppStyle::GetAppStyleSetName(), "ClassIcon.PoseAsset");
+	TabIcon = FSlateIcon(FAppStyle::GetAppStyleSetName(), "AnimGraph.PoseWatch.Icon");
 	bIsSingleton = true;
 }
 
diff --git a/Engine/Source/Editor/UnrealEd/Private/AnimationEditorUtils.cpp b/Engine/Source/Editor/UnrealEd/Private/AnimationEditorUtils.cpp
--- a/Engine/Source/Editor/UnrealEd/Private/AnimationEditorUtils.cpp
+++ b/Engine/Source/Editor/UnrealEd/Private/AnimationEditorUtils.cpp
@@ -974,7 +974,7 @@ namespace AnimationEditorUtils
 		UPoseWatch* NewPoseWatch = NewObject<UPoseWatch>(AnimBlueprint);
 		NewPoseWatch->Node = Node;
 		NewPoseWatch->SetUniqueDefaultLabel();
-		NewPoseWatch->AddElement<UPoseWatchPoseElement>(LOCTEXT("PoseWatchElementLabel_PoseWatch", "Pose Watch"), TEXT("ClassIcon.PoseAsset"));
+		NewPoseWatch->AddElement<UPoseWatchPoseElement>(LOCTEXT("PoseWatchElementLabel_PoseWatch", "Pose Watch"), TEXT("AnimGraph.PoseWatch.Icon"));
 		AnimBlueprint->PoseWatches.Add(NewPoseWatch);
 		SetPoseWatch(NewPoseWatch, AnimBlueprint);
 		return NewPoseWatch;
diff --git a/Engine/Source/Runtime/Engine/Classes/Engine/PoseWatch.h b/Engine/Source/Runtime/Engine/Classes/Engine/PoseWatch.h
--- a/Engine/Source/Runtime/Engine/Classes/Engine/PoseWatch.h
+++ b/Engine/Source/Runtime/Engine/Classes/Engine/PoseWatch.h
@@ -428,7 +428,7 @@ protected:
 	FText Label;
 
 	UPROPERTY()
-	FName IconName;
+	FName IconName_DEPRECATED;
 
 	UPROPERTY()
 	TWeakObjectPtr<UPoseWatchFolder> Parent;
diff --git a/Engine/Source/Runtime/Engine/Private/Engine/PoseWatch.cpp b/Engine/Source/Runtime/Engine/Private/Engine/PoseWatch.cpp
--- a/Engine/Source/Runtime/Engine/Private/Engine/PoseWatch.cpp
+++ b/Engine/Source/Runtime/Engine/Private/Engine/PoseWatch.cpp
@@ -477,7 +477,7 @@ UPoseWatchPoseElement::UPoseWatchPoseElement(const FObjectInitializer& ObjectIni
 	ViewportOffset = FVector3d::ZeroVector;
 
 	Label = GetDefaultLabel();
-	SetIconName("ClassIcon.PoseAsset");
+	SetIconName("AnimGraph.PoseWatch.Icon");
 #endif
 }
 
@@ -490,7 +490,7 @@ bool UPoseWatchPoseElement::GetIsEnabled() const
 
 FSlateIcon UPoseWatchPoseElement::StaticGetIcon()
 {
-	static FSlateIcon Icon(FAppStyle::GetAppStyleSetName(), "ClassIcon.PoseAsset");
+	static FSlateIcon Icon(FAppStyle::GetAppStyleSetName(), "AnimGraph.PoseWatch.Icon");
 	return Icon;
 }
 
