Content Browser folder icon fixes.
- C++ folders in the folder tree now use SVGs like the other icons
- Large version of the C++ icon has been added and is now used in the asset view, same as we do for developer folders
- Small version of the open developer folder icon added and now used when developer folders are open

diff --git a/Engine/Source/Editor/ContentBrowser/Private/AssetViewWidgets.cpp b/Engine/Source/Editor/ContentBrowser/Private/AssetViewWidgets.cpp
--- a/Engine/Source/Editor/ContentBrowser/Private/AssetViewWidgets.cpp
+++ b/Engine/Source/Editor/ContentBrowser/Private/AssetViewWidgets.cpp
@@ -165,6 +165,7 @@ TSharedRef<SWidget> FAssetViewItemHelper::CreateListTileItemContents(T* const In
 		// TODO: Allow items to customize their widget
 
 		const bool bDeveloperFolder = ContentBrowserUtils::IsItemDeveloperContent(InTileOrListItem->AssetItem->GetItem());
+		const bool bCodeFolder = EnumHasAnyFlags(InTileOrListItem->AssetItem->GetItem().GetItemCategory(), EContentBrowserItemFlags::Category_Class);
 
 		const bool bCollectionFolder = EnumHasAnyFlags(InTileOrListItem->AssetItem->GetItem().GetItemCategory(), EContentBrowserItemFlags::Category_Collection);
 		ECollectionShareType::Type CollectionFolderShareType = ECollectionShareType::CST_All;
@@ -175,7 +176,9 @@ TSharedRef<SWidget> FAssetViewItemHelper::CreateListTileItemContents(T* const In
 
 		const FSlateBrush* FolderBaseImage = bDeveloperFolder
 			? FAppStyle::GetBrush("ContentBrowser.ListViewDeveloperFolderIcon") 
-			: FAppStyle::GetBrush("ContentBrowser.ListViewFolderIcon");
+			: bCodeFolder
+				? FAppStyle::GetBrush("ContentBrowser.ListViewCodeFolderIcon")
+				: FAppStyle::GetBrush("ContentBrowser.ListViewFolderIcon");
 
 		// Folder base
 		ItemContentsOverlay->AddSlot()
diff --git a/Engine/Source/Editor/ContentBrowser/Private/SourcesViewWidgets.cpp b/Engine/Source/Editor/ContentBrowser/Private/SourcesViewWidgets.cpp
--- a/Engine/Source/Editor/ContentBrowser/Private/SourcesViewWidgets.cpp
+++ b/Engine/Source/Editor/ContentBrowser/Private/SourcesViewWidgets.cpp
@@ -56,8 +56,9 @@ void SAssetTreeItem::Construct( const FArguments& InArgs )
 	FolderClosedBrush = FAppStyle::GetBrush("ContentBrowser.AssetTreeFolderClosed");
 	FolderOpenCodeBrush = FAppStyle::GetBrush("ContentBrowser.AssetTreeFolderOpenCode");
 	FolderClosedCodeBrush = FAppStyle::GetBrush("ContentBrowser.AssetTreeFolderClosedCode");
-	FolderDeveloperBrush = FAppStyle::GetBrush("ContentBrowser.AssetTreeFolderDeveloper");
-	
+	FolderOpenDeveloperBrush = FAppStyle::GetBrush("ContentBrowser.AssetTreeFolderOpenDeveloper");
+	FolderClosedDeveloperBrush = FAppStyle::GetBrush("ContentBrowser.AssetTreeFolderClosedDeveloper");
+
 	FolderType = EFolderType::Normal;
 	if (ContentBrowserUtils::IsItemDeveloperContent(InArgs._TreeItem->GetItem()))
 	{
@@ -215,7 +216,7 @@ const FSlateBrush* SAssetTreeItem::GetFolderIcon() const
 		return ( IsItemExpanded.Get() ) ? FolderOpenCodeBrush : FolderClosedCodeBrush;
 
 	case EFolderType::Developer:
-		return FolderDeveloperBrush;
+		return (IsItemExpanded.Get()) ? FolderOpenDeveloperBrush : FolderClosedDeveloperBrush;
 
 	default:
 		return ( IsItemExpanded.Get() ) ? FolderOpenBrush : FolderClosedBrush;
diff --git a/Engine/Source/Editor/ContentBrowser/Private/SourcesViewWidgets.h b/Engine/Source/Editor/ContentBrowser/Private/SourcesViewWidgets.h
--- a/Engine/Source/Editor/ContentBrowser/Private/SourcesViewWidgets.h
+++ b/Engine/Source/Editor/ContentBrowser/Private/SourcesViewWidgets.h
@@ -127,7 +127,8 @@ private:
 	const FSlateBrush* FolderClosedBrush;
 	const FSlateBrush* FolderOpenCodeBrush;
 	const FSlateBrush* FolderClosedCodeBrush;
-	const FSlateBrush* FolderDeveloperBrush;
+	const FSlateBrush* FolderOpenDeveloperBrush;
+	const FSlateBrush* FolderClosedDeveloperBrush;
 
 	/** Delegate called to get the selection state of an asset path */
 	FIsSelected IsSelected;
diff --git a/Engine/Source/Editor/EditorStyle/Private/StarshipStyle.cpp b/Engine/Source/Editor/EditorStyle/Private/StarshipStyle.cpp
--- a/Engine/Source/Editor/EditorStyle/Private/StarshipStyle.cpp
+++ b/Engine/Source/Editor/EditorStyle/Private/StarshipStyle.cpp
@@ -5921,9 +5921,10 @@ void FStarshipEditorStyle::FStyle::SetupContentBrowserStyle()
 		Set("ContentBrowser.AssetTreeFolderClosed", new CORE_IMAGE_BRUSH_SVG("Starship/Common/folder-closed", Icon16x16));
 		Set("ContentBrowser.AssetTreeFolderOpen", new CORE_IMAGE_BRUSH_SVG("Starship/Common/folder-open", Icon16x16));
 
-		Set( "ContentBrowser.AssetTreeFolderDeveloper", new IMAGE_BRUSH_SVG( "Starship/ContentBrowser/folder-developer", Icon16x16));
-		Set( "ContentBrowser.AssetTreeFolderOpenCode", new IMAGE_BRUSH( "Icons/FolderOpen_Code", FVector2D(18, 16) ) );
-		Set( "ContentBrowser.AssetTreeFolderClosedCode", new IMAGE_BRUSH( "Icons/FolderClosed_Code", FVector2D(18, 16) ) );
+		Set( "ContentBrowser.AssetTreeFolderOpenDeveloper", new IMAGE_BRUSH_SVG("Starship/ContentBrowser/folder-developer-open", Icon16x16));
+		Set( "ContentBrowser.AssetTreeFolderClosedDeveloper", new IMAGE_BRUSH_SVG( "Starship/ContentBrowser/folder-developer", Icon16x16));
+		Set( "ContentBrowser.AssetTreeFolderOpenCode", new IMAGE_BRUSH_SVG( "Starship/ContentBrowser/folder-code-open", Icon16x16) );
+		Set( "ContentBrowser.AssetTreeFolderClosedCode", new IMAGE_BRUSH_SVG( "Starship/ContentBrowser/folder-code", Icon16x16) );
 
 		Set("ContentBrowser.DefaultFolderColor", FStyleColors::AccentFolder);
 
@@ -6104,6 +6105,7 @@ void FStarshipEditorStyle::FStyle::SetupContentBrowserStyle()
 
 		Set( "ContentBrowser.ListViewFolderIcon", new IMAGE_BRUSH_SVG( "Starship/ContentBrowser/folder", FVector2D(64, 64) ) );
 		Set( "ContentBrowser.ListViewDeveloperFolderIcon", new IMAGE_BRUSH_SVG("Starship/ContentBrowser/FolderLargeDeveloper", FVector2D(64, 64)));
+		Set( "ContentBrowser.ListViewCodeFolderIcon", new IMAGE_BRUSH_SVG("Starship/ContentBrowser/FolderLargeCode", FVector2D(64, 64)));
 
 		Set("ContentBrowser.AssetTileItem.FolderAreaHoveredBackground", new FSlateRoundedBoxBrush(FStyleColors::Secondary, 4.0f));
 		Set("ContentBrowser.AssetTileItem.FolderAreaSelectedBackground", new FSlateRoundedBoxBrush(FStyleColors::Secondary, 4.0f));
