Curve Table Editor: Fix an issue causing misalignment between the rows

diff --git a/Engine/Source/Editor/EditorStyle/Private/StarshipStyle.cpp b/Engine/Source/Editor/EditorStyle/Private/StarshipStyle.cpp
--- a/Engine/Source/Editor/EditorStyle/Private/StarshipStyle.cpp
+++ b/Engine/Source/Editor/EditorStyle/Private/StarshipStyle.cpp
@@ -6420,7 +6420,7 @@ void FStarshipEditorStyle::FStyle::SetupToolkitStyles()
 			.SetBackgroundImageHovered(FSlateRoundedBoxBrush(FStyleColors::Input, 0.0, FStyleColors::Hover, 1.0))
 			.SetBackgroundImageFocused(FSlateRoundedBoxBrush(FStyleColors::Input, 0.0, FStyleColors::Primary, 1.0))
 			.SetBackgroundImageReadOnly(FSlateNoResource())
-			.SetPadding(FMargin(4.f, 5.0f, 4.f, 4.0f)) // The padding should be 4 top, 5 bottom
+			.SetPadding(FMargin(4.0f))
 			.SetForegroundColor(FStyleColors::Foreground)
 			.SetBackgroundColor(FStyleColors::White)
 			.SetReadOnlyForegroundColor(FSlateColor::UseForeground())
