Removed tint from Source Control Status Bar buttons icons to match UX spec

diff --git a/Engine/Source/Editor/EditorStyle/Private/StarshipStyle.cpp b/Engine/Source/Editor/EditorStyle/Private/StarshipStyle.cpp
--- a/Engine/Source/Editor/EditorStyle/Private/StarshipStyle.cpp
+++ b/Engine/Source/Editor/EditorStyle/Private/StarshipStyle.cpp
@@ -7051,10 +7051,10 @@ void FStarshipEditorStyle::FStyle::SetupSourceControlStyles()
 		Set("SourceControl.StatusIcon.Error", new CORE_IMAGE_BRUSH_SVG("Starship/Common/alert-circle", Icon16x16, FStyleColors::AccentYellow));
 		Set("SourceControl.StatusIcon.Off", new CORE_IMAGE_BRUSH_SVG("Starship/Common/reject", Icon16x16, FStyleColors::Foreground));
 		Set("SourceControl.StatusIcon.Unknown", new CORE_IMAGE_BRUSH_SVG("Starship/Common/help", Icon16x16, FStyleColors::AccentYellow));
-		Set("SourceControl.StatusBar.AtLatestRevision", new CORE_IMAGE_BRUSH_SVG("Starship/Common/SCC_Cloud_AtLatest", Icon16x16, FStyleColors::AccentGreen));
-		Set("SourceControl.StatusBar.NotAtLatestRevision", new CORE_IMAGE_BRUSH_SVG("Starship/Common/SCC_Cloud_NotAtLatest", Icon16x16, FStyleColors::AccentYellow));
-		Set("SourceControl.StatusBar.NoLocalChanges", new CORE_IMAGE_BRUSH_SVG("Starship/Common/SCC_Cloud_NoLocalChanges", Icon16x16, FStyleColors::AccentGreen));
-		Set("SourceControl.StatusBar.HasLocalChanges", new CORE_IMAGE_BRUSH_SVG("Starship/Common/SCC_Cloud_HasLocalChanges", Icon16x16, FStyleColors::AccentYellow));
+		Set("SourceControl.StatusBar.AtLatestRevision", new CORE_IMAGE_BRUSH_SVG("Starship/Common/SCC_Cloud_AtLatest", Icon16x16));
+		Set("SourceControl.StatusBar.NotAtLatestRevision", new CORE_IMAGE_BRUSH_SVG("Starship/Common/SCC_Cloud_NotAtLatest", Icon16x16));
+		Set("SourceControl.StatusBar.NoLocalChanges", new CORE_IMAGE_BRUSH_SVG("Starship/Common/SCC_Cloud_NoLocalChanges", Icon16x16));
+		Set("SourceControl.StatusBar.HasLocalChanges", new CORE_IMAGE_BRUSH_SVG("Starship/Common/SCC_Cloud_HasLocalChanges", Icon16x16));
 
 		Set("SourceControl.ChangelistsTab", new CORE_IMAGE_BRUSH_SVG("Starship/Common/check-circle", Icon16x16));
 		Set("SourceControl.Changelist", new CORE_IMAGE_BRUSH_SVG("Starship/SourceControl/SCC_Changelist", Icon16x16, FStyleColors::AccentRed));
