Add EULA style markdown to support RichText in slate

diff --git a/Engine/Source/Editor/EditorStyle/Private/StarshipStyle.cpp b/Engine/Source/Editor/EditorStyle/Private/StarshipStyle.cpp
--- a/Engine/Source/Editor/EditorStyle/Private/StarshipStyle.cpp
+++ b/Engine/Source/Editor/EditorStyle/Private/StarshipStyle.cpp
@@ -325,6 +325,30 @@ void FStarshipEditorStyle::FStyle::SetupGeneralStyles()
 			.SetShadowColorAndOpacity(FLinearColor(0, 0, 0, 0.9f)));
 	}
 
+	// EULA RichText
+	{
+		Set("EULA.Header", FTextBlockStyle(FAppStyle::Get().GetWidgetStyle<FTextBlockStyle>("Log.Normal"))
+			.SetFont(FCoreStyle::GetDefaultFontStyle("Bold", 10)));
+
+		Set("EULA.HighlightItalic", FTextBlockStyle(NormalText)
+			.SetFont(DEFAULT_FONT("Italic", FStarshipCoreStyle::RegularTextSize))
+			.SetColorAndOpacity(FLinearColor(1.0f,1.0f, 1.0f)));
+
+		const FButtonStyle EulaHyperlinkButton = FButtonStyle()
+			.SetNormal(BORDER_BRUSH("Old/HyperlinkDotted", FMargin(0, 0, 0, 3 / 16.0f), FLinearColor(0.25f, 0.5f, 1.0f)))
+			.SetPressed(FSlateNoResource())
+			.SetHovered(BORDER_BRUSH("Old/HyperlinkUnderline", FMargin(0, 0, 0, 3 / 16.0f), FLinearColor(0.25f, 0.5f, 1.0f)));
+
+		const FTextBlockStyle EulaHyperlinkText = FTextBlockStyle(NormalText)
+			.SetColorAndOpacity(FLinearColor(0.25f, 0.5f, 1.0f));
+
+		Set("EULA.Hyperlink", FHyperlinkStyle()
+			.SetUnderlineStyle(EulaHyperlinkButton)
+			.SetTextStyle(EulaHyperlinkText)
+			.SetPadding(FMargin(0.0f)));
+
+	}
+
 	// Rendering resources that never change
 	{
 		Set( "None", new FSlateNoResource() );
