Add alternating colors for rows in the Outliner and Content Browser
STableRow style: Add a default style with alternating row colors

diff --git a/Engine/Source/Editor/EditorStyle/Private/StarshipStyle.cpp b/Engine/Source/Editor/EditorStyle/Private/StarshipStyle.cpp
--- a/Engine/Source/Editor/EditorStyle/Private/StarshipStyle.cpp
+++ b/Engine/Source/Editor/EditorStyle/Private/StarshipStyle.cpp
@@ -1161,7 +1161,9 @@ void FStarshipEditorStyle::FStyle::SetupGeneralStyles()
 		Set( "SceneOutliner.World", 		new CORE_IMAGE_BRUSH_SVG("Starship/Common/world", Icon16x16 ) );
 		Set( "SceneOutliner.ChangedItemHighlight", new FSlateRoundedBoxBrush( FStyleColors::White, 1.0) );
 
-		Set( "SceneOutliner.TableViewRow", FTableRowStyle(NormalTableRowStyle));
+		const FTableRowStyle AlternatingTableRowStyle = GetWidgetStyle<FTableRowStyle>("TableView.AlternatingRow");
+		
+		Set( "SceneOutliner.TableViewRow", AlternatingTableRowStyle);
 	}
 
 	// Socket chooser
@@ -5988,10 +5990,10 @@ void FStarshipEditorStyle::FStyle::SetupContentBrowserStyle()
 		Set( "ContentBrowser.AssetTileViewNameFontDirty", FStyleFonts::Get().SmallBold);
 
 		Set("ContentBrowser.AssetListView.ColumnListTableRow", FTableRowStyle()
-			.SetEvenRowBackgroundBrush( FSlateNoResource() )
-			.SetEvenRowBackgroundHoveredBrush( IMAGE_BRUSH( "Common/Selection", Icon8x8, FLinearColor(1.0f, 1.0f, 1.0f, 0.1f) ) )
-			.SetOddRowBackgroundBrush( FSlateNoResource() )
-			.SetOddRowBackgroundHoveredBrush( IMAGE_BRUSH( "Common/Selection", Icon8x8, FLinearColor(1.0f, 1.0f, 1.0f, 0.1f) ) )
+			.SetEvenRowBackgroundBrush(FSlateColorBrush(FStyleColors::Recessed))
+			.SetEvenRowBackgroundHoveredBrush(FSlateColorBrush(FStyleColors::SelectHover))
+			.SetOddRowBackgroundBrush(FSlateColorBrush(FStyleColors::Background))
+			.SetOddRowBackgroundHoveredBrush(FSlateColorBrush(FStyleColors::SelectHover))
 			.SetSelectorFocusedBrush( BORDER_BRUSH( "Common/Selector", FMargin(4.f/16.f), FStyleColors::Select ) )
 			.SetActiveBrush( IMAGE_BRUSH( "Common/Selection", Icon8x8, FStyleColors::Select ) )
 			.SetActiveHoveredBrush( IMAGE_BRUSH( "Common/Selection", Icon8x8, FStyleColors::Select ) )
diff --git a/Engine/Source/Runtime/SlateCore/Private/Styling/StarshipCoreStyle.cpp b/Engine/Source/Runtime/SlateCore/Private/Styling/StarshipCoreStyle.cpp
--- a/Engine/Source/Runtime/SlateCore/Private/Styling/StarshipCoreStyle.cpp
+++ b/Engine/Source/Runtime/SlateCore/Private/Styling/StarshipCoreStyle.cpp
@@ -2023,6 +2023,11 @@ void FStarshipCoreStyle::SetupTableViewStyles(TSharedRef<FStyle>& Style)
 		.SetOddRowBackgroundBrush(IMAGE_BRUSH("Common/Selection", Icon8x8, FLinearColor(0.0f, 0.0f, 0.0f, 0.1f)));
 	Style->Set("TableView.DarkRow", DarkTableRowStyle);
 
+	const FTableRowStyle AlternatingTableRowStyle = FTableRowStyle(DefaultTableRowStyle)
+		.SetEvenRowBackgroundBrush(FSlateColorBrush(FStyleColors::Recessed))
+		.SetOddRowBackgroundBrush(FSlateColorBrush(FStyleColors::Background));
+	Style->Set("TableView.AlternatingRow", AlternatingTableRowStyle);
+
 	Style->Set("TreeArrow_Collapsed", new IMAGE_BRUSH_SVG("Starship/Common/caret-right", Icon16x16, FStyleColors::Foreground));
 	Style->Set("TreeArrow_Collapsed_Hovered", new IMAGE_BRUSH_SVG("Starship/Common/caret-right", Icon16x16, FStyleColors::ForegroundHover));
 	Style->Set("TreeArrow_Expanded", new IMAGE_BRUSH_SVG("Starship/Common/caret-down", Icon16x16, FStyleColors::Foreground));
