Content Browser: Add new border for the unsupported asset text in the tooltip for items

diff --git a/Engine/Source/Editor/EditorStyle/Private/StarshipStyle.cpp b/Engine/Source/Editor/EditorStyle/Private/StarshipStyle.cpp
--- a/Engine/Source/Editor/EditorStyle/Private/StarshipStyle.cpp
+++ b/Engine/Source/Editor/EditorStyle/Private/StarshipStyle.cpp
@@ -6108,7 +6108,8 @@ void FStarshipEditorStyle::FStyle::SetupContentBrowserStyle()
 		Set( "ContentBrowser.TileViewTooltip.ToolTipBorder", new FSlateColorBrush( FLinearColor::Black ) );
 		Set( "ContentBrowser.TileViewTooltip.NonContentBorder", new BOX_BRUSH( "/Docking/TabContentArea", FMargin(4/16.0f) ) );
 		Set( "ContentBrowser.TileViewTooltip.ContentBorder", new FSlateColorBrush( FStyleColors::Panel));
-		Set("ContentBrowser.TileViewTooltip.PillBorder", new FSlateRoundedBoxBrush(FStyleColors::Transparent, 12.0f, FStyleColors::White, 0.5f));
+		Set( "ContentBrowser.TileViewTooltip.PillBorder", new FSlateRoundedBoxBrush(FStyleColors::Transparent, 10.0f, FStyleColors::White, 1.0f));
+		Set( "ContentBrowser.TileViewTooltip.UnsupportedAssetPillBorder", new FSlateRoundedBoxBrush(FStyleColors::Transparent, 10.0f, FStyleColors::Warning, 1.0f));
 		Set( "ContentBrowser.TileViewTooltip.NameFont", DEFAULT_FONT( "Regular", 12 ) );
 		Set( "ContentBrowser.TileViewTooltip.AssetUserDescriptionFont", DEFAULT_FONT("Regular", 12 ) );
 
