Moved common style for blueprint/code references from StarshipStyle.cpp to StarshipCoreStyle.cpp
UnrealFrontend - Automated test names are not readable

diff --git a/Engine/Source/Editor/EditorStyle/Private/StarshipStyle.cpp b/Engine/Source/Editor/EditorStyle/Private/StarshipStyle.cpp
--- a/Engine/Source/Editor/EditorStyle/Private/StarshipStyle.cpp
+++ b/Engine/Source/Editor/EditorStyle/Private/StarshipStyle.cpp
@@ -1858,25 +1858,6 @@ void FStarshipEditorStyle::FStyle::SetupGeneralStyles()
 		Set("Scalability.GroupFont", DEFAULT_FONT( "Bold", 10 ) );
 	}
 
-	// Common styles for blueprint/code references that also need to be exposed to external tools
-	{
-		FTextBlockStyle InheritedFromNativeTextStyle = FTextBlockStyle(NormalText)
-			.SetFont(DEFAULT_FONT("Regular", 10));
-
-		Set("Common.InheritedFromNativeTextStyle", InheritedFromNativeTextStyle);
-
-		// Go to native class hyperlink
-		FButtonStyle EditNativeHyperlinkButton = FButtonStyle()
-			.SetNormal(BORDER_BRUSH("Old/HyperlinkDotted", FMargin(0, 0, 0, 3 / 16.0f)))
-			.SetPressed(FSlateNoResource())
-			.SetHovered(BORDER_BRUSH("Old/HyperlinkUnderline", FMargin(0, 0, 0, 3 / 16.0f)));
-		FHyperlinkStyle EditNativeHyperlinkStyle = FHyperlinkStyle()
-			.SetUnderlineStyle(EditNativeHyperlinkButton)
-			.SetTextStyle(InheritedFromNativeTextStyle)
-			.SetPadding(FMargin(0.0f));
-
-		Set("Common.GotoNativeCodeHyperlink", EditNativeHyperlinkStyle);
-	}
 #endif // WITH_EDITOR || (IS_PROGRAM && WITH_UNREAL_DEVELOPER_TOOLS)
 
 #if WITH_EDITOR || (IS_PROGRAM && WITH_UNREAL_DEVELOPER_TOOLS)
diff --git a/Engine/Source/Runtime/SlateCore/Private/Styling/StarshipCoreStyle.cpp b/Engine/Source/Runtime/SlateCore/Private/Styling/StarshipCoreStyle.cpp
--- a/Engine/Source/Runtime/SlateCore/Private/Styling/StarshipCoreStyle.cpp
+++ b/Engine/Source/Runtime/SlateCore/Private/Styling/StarshipCoreStyle.cpp
@@ -515,6 +515,26 @@ TSharedRef<ISlateStyle> FStarshipCoreStyle::Create()
 		Style->Set("Hyperlink", Hyperlink);
 	}
 
+	// Common styles for blueprint/code references that also need to be exposed to external tools
+	{
+		FTextBlockStyle InheritedFromNativeTextStyle = FTextBlockStyle(NormalText)
+			.SetFont(DEFAULT_FONT("Regular", 10));
+
+		Style->Set("Common.InheritedFromNativeTextStyle", InheritedFromNativeTextStyle);
+
+		// Go to native class hyperlink
+		FButtonStyle EditNativeHyperlinkButton = FButtonStyle()
+			.SetNormal(BORDER_BRUSH("Old/HyperlinkDotted", FMargin(0, 0, 0, 3 / 16.0f)))
+			.SetPressed(FSlateNoResource())
+			.SetHovered(BORDER_BRUSH("Old/HyperlinkUnderline", FMargin(0, 0, 0, 3 / 16.0f)));
+		FHyperlinkStyle EditNativeHyperlinkStyle = FHyperlinkStyle()
+			.SetUnderlineStyle(EditNativeHyperlinkButton)
+			.SetTextStyle(InheritedFromNativeTextStyle)
+			.SetPadding(FMargin(0.0f));
+
+		Style->Set("Common.GotoNativeCodeHyperlink", EditNativeHyperlinkStyle);
+	}
+
 	// SProgressBar defaults...
 	{
 		Style->Set("ProgressBar", FProgressBarStyle()
