Sequencer: Set disabled styles for transport control buttons

diff --git a/Engine/Source/Editor/EditorStyle/Private/StarshipStyle.cpp b/Engine/Source/Editor/EditorStyle/Private/StarshipStyle.cpp
--- a/Engine/Source/Editor/EditorStyle/Private/StarshipStyle.cpp
+++ b/Engine/Source/Editor/EditorStyle/Private/StarshipStyle.cpp
@@ -878,12 +878,14 @@ void FStarshipEditorStyle::FStyle::SetupGeneralStyles()
 			.SetNormal(IMAGE_BRUSH( "/Sequencer/Transport_Bar/Step_Forward_24x_OFF", Icon24x24 ))
 			.SetHovered(IMAGE_BRUSH( "/Sequencer/Transport_Bar/Step_Forward_24x_OFF", Icon24x24 ))
 			.SetPressed(IMAGE_BRUSH( "/Sequencer/Transport_Bar/Step_Forward_24x", Icon24x24 ))
+			.SetDisabled(IMAGE_BRUSH( "/Sequencer/Transport_Bar/Step_Forward_24x_OFF", Icon24x24, FLinearColor(0.4f, 0.4f, 0.4f, 0.4f) ))
 		);
 
 		Set( "Animation.Forward_End", FButtonStyle( Button )
 			.SetNormal(IMAGE_BRUSH( "/Sequencer/Transport_Bar/Go_To_End_24x_OFF", Icon24x24 ))
 			.SetHovered(IMAGE_BRUSH( "/Sequencer/Transport_Bar/Go_To_End_24x_OFF", Icon24x24 ))
 			.SetPressed(IMAGE_BRUSH( "/Sequencer/Transport_Bar/Go_To_End_24x", Icon24x24 ))
+			.SetDisabled(IMAGE_BRUSH( "/Sequencer/Transport_Bar/Go_To_End_24x_OFF", Icon24x24, FLinearColor(0.4f, 0.4f, 0.4f, 0.4f) ))
 		);
 
 		Set( "Animation.Backward", FButtonStyle( Button )
@@ -896,12 +898,14 @@ void FStarshipEditorStyle::FStyle::SetupGeneralStyles()
 			.SetNormal(IMAGE_BRUSH( "/Sequencer/Transport_Bar/Step_Backwards_24x_OFF", Icon24x24 ))
 			.SetHovered(IMAGE_BRUSH( "/Sequencer/Transport_Bar/Step_Backwards_24x_OFF", Icon24x24 ))
 			.SetPressed(IMAGE_BRUSH( "/Sequencer/Transport_Bar/Step_Backwards_24x", Icon24x24 ))
+			.SetDisabled(IMAGE_BRUSH( "/Sequencer/Transport_Bar/Step_Backwards_24x_OFF", Icon24x24, FLinearColor(0.4f, 0.4f, 0.4f, 0.4f) ))
 		);
 
 		Set( "Animation.Backward_End", FButtonStyle( Button )
 			.SetNormal(IMAGE_BRUSH( "/Sequencer/Transport_Bar/Go_To_Front_24x_OFF", Icon24x24 ))
 			.SetHovered(IMAGE_BRUSH( "/Sequencer/Transport_Bar/Go_To_Front_24x_OFF", Icon24x24 ))
 			.SetPressed(IMAGE_BRUSH( "/Sequencer/Transport_Bar/Go_To_Front_24x", Icon24x24 ))
+			.SetDisabled(IMAGE_BRUSH( "/Sequencer/Transport_Bar/Step_Backwards_24x_OFF", Icon24x24, FLinearColor(0.4f, 0.4f, 0.4f, 0.4f) ))
 		);
 
 		Set( "Animation.Loop.Enabled", FButtonStyle( Button )
