Fix missing Search icon in Material Editor

diff --git a/Engine/Source/Editor/EditorStyle/Private/StarshipStyle.cpp b/Engine/Source/Editor/EditorStyle/Private/StarshipStyle.cpp
--- a/Engine/Source/Editor/EditorStyle/Private/StarshipStyle.cpp
+++ b/Engine/Source/Editor/EditorStyle/Private/StarshipStyle.cpp
@@ -6635,7 +6635,7 @@ void FStarshipEditorStyle::FStyle::SetupToolkitStyles()
 		Set( "MaterialEditor.TogglePlatformStats", new IMAGE_BRUSH_SVG( "Starship/MaterialEditor/PlatformStats", Icon20x20 ) );
 		Set("MaterialEditor.TogglePlatformStats.Tab", new IMAGE_BRUSH_SVG("Starship/MaterialEditor/PlatformStats", Icon16x16));
 		Set( "MaterialEditor.CameraHome", new IMAGE_BRUSH_SVG( "Starship/Common/Home", Icon20x20 ) );
-		Set( "MaterialEditor.FindInMaterial", new IMAGE_BRUSH_SVG( "Starship/Common/Search", Icon20x20 ) );
+		Set( "MaterialEditor.FindInMaterial", new CORE_IMAGE_BRUSH_SVG( "Starship/Common/Search", Icon20x20 ) );
 
 
 		Set("MaterialEditor.CellListViewRow", FTableRowStyle(NormalTableRowStyle)
@@ -6672,7 +6672,7 @@ void FStarshipEditorStyle::FStyle::SetupToolkitStyles()
 	{
 		// Tab icons
 		{
-			Set( "FontEditor.Tabs.Preview", new IMAGE_BRUSH_SVG( "Starship/Common/Search", Icon16x16 ) );
+			Set( "FontEditor.Tabs.Preview", new CORE_IMAGE_BRUSH_SVG( "Starship/Common/Search", Icon16x16 ) );
 			Set( "FontEditor.Tabs.Properties", new IMAGE_BRUSH( "/Icons/icon_tab_SelectionDetails_16x", Icon16x16 ) );
 			Set( "FontEditor.Tabs.PageProperties", new IMAGE_BRUSH( "/Icons/properties_16x", Icon16x16 ) );
 		}
