Updating append key icon in the curve table editor.

diff --git a/Engine/Source/Editor/CurveTableEditor/Private/CurveTableEditor.cpp b/Engine/Source/Editor/CurveTableEditor/Private/CurveTableEditor.cpp
--- a/Engine/Source/Editor/CurveTableEditor/Private/CurveTableEditor.cpp
+++ b/Engine/Source/Editor/CurveTableEditor/Private/CurveTableEditor.cpp
@@ -544,7 +544,7 @@ void FCurveTableEditor::ExtendToolbar()
 						NAME_None, 
 						FText::GetEmpty(),
 						TAttribute<FText>(), 
-						FSlateIcon(FAppStyle::Get().GetStyleSetName(), "Sequencer.KeyTriangle")
+						FSlateIcon(FAppStyle::Get().GetStyleSetName(), "Sequencer.KeyTriangle20")
 					);
 				}
 			}
diff --git a/Engine/Source/Editor/EditorStyle/Private/StarshipStyle.cpp b/Engine/Source/Editor/EditorStyle/Private/StarshipStyle.cpp
--- a/Engine/Source/Editor/EditorStyle/Private/StarshipStyle.cpp
+++ b/Engine/Source/Editor/EditorStyle/Private/StarshipStyle.cpp
@@ -2062,6 +2062,7 @@ void FStarshipEditorStyle::FStyle::SetupSequencerStyles()
 		Set("Sequencer.KeyDiamondBorder", new IMAGE_BRUSH("Sequencer/KeyDiamondBorder", Icon12x12));
 		Set("Sequencer.KeySquare", new IMAGE_BRUSH("Sequencer/KeySquare", Icon12x12));
 		Set("Sequencer.KeyTriangle", new IMAGE_BRUSH("Sequencer/KeyTriangle", Icon12x12));
+		Set("Sequencer.KeyTriangle20", new CORE_IMAGE_BRUSH_SVG("Starship/Common/curve-editor-append-key-20", Icon20x20));
 		Set("Sequencer.KeyLeft", new IMAGE_BRUSH("Sequencer/KeyLeft", Icon12x12));
 		Set("Sequencer.KeyRight", new IMAGE_BRUSH("Sequencer/KeyRight", Icon12x12));
 		Set("Sequencer.PartialKey", new IMAGE_BRUSH("Sequencer/PartialKey", FVector2D(11.f, 11.f)));
