Adding asset icons for editor utility widgets and editor utility blueprints

diff --git a/Engine/Source/Editor/EditorStyle/Private/StarshipStyle.cpp b/Engine/Source/Editor/EditorStyle/Private/StarshipStyle.cpp
--- a/Engine/Source/Editor/EditorStyle/Private/StarshipStyle.cpp
+++ b/Engine/Source/Editor/EditorStyle/Private/StarshipStyle.cpp
@@ -5883,6 +5883,8 @@ void FStarshipEditorStyle::FStyle::SetupClassIconsAndThumbnails()
 			{TEXT("DialogueWave")},
 			{TEXT("DirectionalLight")},
 			{TEXT("DocumentationActor")},
+			{TEXT("EditorUtilityBlueprint")},	
+			{TEXT("EditorUtilityWidgetBlueprint")},
 			{TEXT("EnvQuery")},
 			{TEXT("Emitter")},
 			{TEXT("EmptyActor")},
