Sequencer: add icon overlays for dynamic bindings

diff --git a/Engine/Source/Editor/EditorStyle/Private/SlateEditorStyle.cpp b/Engine/Source/Editor/EditorStyle/Private/SlateEditorStyle.cpp
--- a/Engine/Source/Editor/EditorStyle/Private/SlateEditorStyle.cpp
+++ b/Engine/Source/Editor/EditorStyle/Private/SlateEditorStyle.cpp
@@ -2498,6 +2498,8 @@ void FSlateEditorStyle::FStyle::SetupSequencerStyles()
 		Set("Sequencer.SpawnableIconOverlay", new IMAGE_BRUSH(TEXT("Sequencer/SpawnableIconOverlay"), FVector2D(13, 13)));
 		Set("Sequencer.InvalidSpawnableIcon", new IMAGE_BRUSH(TEXT("Sequencer/InvalidSpawnableIcon"), Icon16x16));
 		Set("Sequencer.MultipleIconOverlay", new IMAGE_BRUSH(TEXT("Sequencer/MultipleIconOverlay"), FVector2D(13, 13)));
+		Set("Sequencer.DynamicBindingIconOverlay", new IMAGE_BRUSH(TEXT("Sequencer/DynamicBindingIconOverlay"), Icon16x16));
+		Set("Sequencer.SpawnableDynamicBindingIconOverlay", new IMAGE_BRUSH(TEXT("Sequencer/SpawnableDynamicBindingIconOverlay"), Icon16x16));
 		Set("Sequencer.LockSequence", new IMAGE_BRUSH("Sequencer/Main_Icons/Icon_Sequencer_Locked_16x", Icon16x16));
 		Set("Sequencer.UnlockSequence", new IMAGE_BRUSH("Sequencer/Main_Icons/Icon_Sequencer_Unlocked_16x", Icon16x16));
 
diff --git a/Engine/Source/Editor/EditorStyle/Private/StarshipStyle.cpp b/Engine/Source/Editor/EditorStyle/Private/StarshipStyle.cpp
--- a/Engine/Source/Editor/EditorStyle/Private/StarshipStyle.cpp
+++ b/Engine/Source/Editor/EditorStyle/Private/StarshipStyle.cpp
@@ -2220,6 +2220,8 @@ void FStarshipEditorStyle::FStyle::SetupSequencerStyles()
 		Set("Sequencer.TrackHoverHighlight_Bottom", new IMAGE_BRUSH(TEXT("Sequencer/TrackHoverHighlight_Bottom"), FVector2D(4, 4)));
 		Set("Sequencer.SpawnableIconOverlay", new IMAGE_BRUSH(TEXT("Sequencer/SpawnableIconOverlay"), FVector2D(13, 13)));
 		Set("Sequencer.MultipleIconOverlay", new IMAGE_BRUSH(TEXT("Sequencer/MultipleIconOverlay"), FVector2D(13, 13)));
+		Set("Sequencer.DynamicBindingIconOverlay", new IMAGE_BRUSH(TEXT("Sequencer/DynamicBindingIconOverlay"), Icon16x16));
+		Set("Sequencer.SpawnableDynamicBindingIconOverlay", new IMAGE_BRUSH(TEXT("Sequencer/SpawnableDynamicBindingIconOverlay"), Icon16x16));
 		Set("Sequencer.LockSequence", new IMAGE_BRUSH("Sequencer/Main_Icons/Icon_Sequencer_Locked_16x", Icon16x16));
 		Set("Sequencer.UnlockSequence", new IMAGE_BRUSH("Sequencer/Main_Icons/Icon_Sequencer_Unlocked_16x", Icon16x16));
 
diff --git a/Engine/Source/Editor/Sequencer/Private/MVVM/ViewModels/PossessableModel.cpp b/Engine/Source/Editor/Sequencer/Private/MVVM/ViewModels/PossessableModel.cpp
--- a/Engine/Source/Editor/Sequencer/Private/MVVM/ViewModels/PossessableModel.cpp
+++ b/Engine/Source/Editor/Sequencer/Private/MVVM/ViewModels/PossessableModel.cpp
@@ -73,6 +73,13 @@ FText FPossessableModel::GetIconToolTipText() const
 
 const FSlateBrush* FPossessableModel::GetIconOverlayBrush() const
 {
+	UMovieScene*            MovieScene  = OwnerModel ? OwnerModel->GetMovieScene() : nullptr;
+	FMovieScenePossessable* Possessable = MovieScene ? MovieScene->FindPossessable(ObjectBindingID) : nullptr;
+	if (Possessable && Possessable->DynamicBinding.WeakEndpoint.IsValid())
+	{
+		return FAppStyle::GetBrush("Sequencer.DynamicBindingIconOverlay");
+	}
+
 	TSharedPtr<ISequencer> Sequencer = GetEditor()->GetSequencer();
 	if (Sequencer)
 	{
diff --git a/Engine/Source/Editor/Sequencer/Private/MVVM/ViewModels/SpawnableModel.cpp b/Engine/Source/Editor/Sequencer/Private/MVVM/ViewModels/SpawnableModel.cpp
--- a/Engine/Source/Editor/Sequencer/Private/MVVM/ViewModels/SpawnableModel.cpp
+++ b/Engine/Source/Editor/Sequencer/Private/MVVM/ViewModels/SpawnableModel.cpp
@@ -86,6 +86,13 @@ FText FSpawnableModel::GetIconToolTipText() const
 
 const FSlateBrush* FSpawnableModel::GetIconOverlayBrush() const
 {
+	UMovieScene*          MovieScene  = OwnerModel ? OwnerModel->GetMovieScene() : nullptr;
+	FMovieSceneSpawnable* Spawnable  = MovieScene ? MovieScene->FindSpawnable(ObjectBindingID) : nullptr;
+	if (Spawnable && Spawnable->DynamicBinding.WeakEndpoint.IsValid())
+	{
+		return FAppStyle::GetBrush("Sequencer.SpawnableDynamicBindingIconOverlay");
+	}
+
 	return FAppStyle::GetBrush("Sequencer.SpawnableIconOverlay");
 }
 
