Fixed comment bubble text input background scaling incorrectly with multiple lines of text.

diff --git a/Engine/Source/Editor/EditorStyle/Private/StarshipStyle.cpp b/Engine/Source/Editor/EditorStyle/Private/StarshipStyle.cpp
--- a/Engine/Source/Editor/EditorStyle/Private/StarshipStyle.cpp
+++ b/Engine/Source/Editor/EditorStyle/Private/StarshipStyle.cpp
@@ -3580,10 +3580,10 @@ void FStarshipEditorStyle::FStyle::SetupGraphEditorStyles()
 			FEditableTextBoxStyle CommentEditableText = FEditableTextBoxStyle(NormalEditableTextBoxStyle)
 				.SetFont(NormalText.Font)
 				.SetForegroundColor(FStyleColors::Foreground)
-				.SetBackgroundImageNormal(FSlateRoundedBoxBrush(FStyleColors::Foreground, FStyleColors::Secondary, 1.0f))
-				.SetBackgroundImageHovered(FSlateRoundedBoxBrush(FStyleColors::Foreground, FStyleColors::Hover, 1.0f))
-				.SetBackgroundImageFocused(FSlateRoundedBoxBrush(FStyleColors::Foreground, FStyleColors::Primary, 1.0f))
-				.SetBackgroundImageReadOnly(FSlateRoundedBoxBrush(FStyleColors::Header, FStyleColors::InputOutline, 1.0f))
+				.SetBackgroundImageNormal(FSlateRoundedBoxBrush(FStyleColors::Foreground, CoreStyleConstants::InputFocusRadius, FStyleColors::Secondary, CoreStyleConstants::InputFocusThickness))
+				.SetBackgroundImageHovered(FSlateRoundedBoxBrush(FStyleColors::Foreground, CoreStyleConstants::InputFocusRadius, FStyleColors::Hover, CoreStyleConstants::InputFocusThickness))
+				.SetBackgroundImageFocused(FSlateRoundedBoxBrush(FStyleColors::Foreground, CoreStyleConstants::InputFocusRadius, FStyleColors::Primary, CoreStyleConstants::InputFocusThickness))
+				.SetBackgroundImageReadOnly(FSlateRoundedBoxBrush(FStyleColors::Header, CoreStyleConstants::InputFocusRadius, FStyleColors::InputOutline, CoreStyleConstants::InputFocusThickness))
 				.SetForegroundColor(FStyleColors::Background)
 				.SetBackgroundColor(FStyleColors::White)
 				.SetReadOnlyForegroundColor(FSlateColor::UseForeground())
