PR #10366: Fix the comment bubble borders in blueprints.

diff --git a/Engine/Source/Editor/EditorStyle/Private/StarshipStyle.cpp b/Engine/Source/Editor/EditorStyle/Private/StarshipStyle.cpp
--- a/Engine/Source/Editor/EditorStyle/Private/StarshipStyle.cpp
+++ b/Engine/Source/Editor/EditorStyle/Private/StarshipStyle.cpp
@@ -3611,7 +3611,11 @@ void FStarshipEditorStyle::FStyle::SetupGraphEditorStyles()
 
 			FEditableTextBoxStyle GraphCommentBlockTitleEditableText = FEditableTextBoxStyle(GraphNodeTitleEditableText)
 				.SetFont(GraphCommentBlockTitle.Font)
-				.SetScrollBarStyle( ScrollBar );
+				.SetScrollBarStyle( ScrollBar )
+				.SetBackgroundImageNormal(FSlateRoundedBoxBrush(FStyleColors::Foreground, 0.0f, FStyleColors::Transparent, 0.0f))
+				.SetBackgroundImageHovered(FSlateRoundedBoxBrush(FStyleColors::Foreground, CoreStyleConstants::InputFocusRadius, FStyleColors::Hover, CoreStyleConstants::InputFocusThickness))
+				.SetBackgroundImageFocused(FSlateRoundedBoxBrush(FStyleColors::Foreground, CoreStyleConstants::InputFocusRadius, FStyleColors::Primary, CoreStyleConstants::InputFocusThickness))
+				.SetBackgroundImageReadOnly(FSlateRoundedBoxBrush(FStyleColors::Header, CoreStyleConstants::InputFocusRadius, FStyleColors::InputOutline, CoreStyleConstants::InputFocusThickness));
 			Set( "Graph.CommentBlock.TitleEditableText", GraphCommentBlockTitleEditableText );
 
 			Set( "Graph.CommentBlock.TitleInlineEditableText", FInlineEditableTextBlockStyle()
