Fix the "My Blueprint" context menu delete icon. It was still using the old UE4 style logo with the little red x, now it's the normal UE5 trash can logo.

#rnx

diff --git a/Engine/Source/Editor/EditorStyle/Private/StarshipStyle.cpp b/Engine/Source/Editor/EditorStyle/Private/StarshipStyle.cpp
--- a/Engine/Source/Editor/EditorStyle/Private/StarshipStyle.cpp
+++ b/Engine/Source/Editor/EditorStyle/Private/StarshipStyle.cpp
@@ -7617,7 +7617,7 @@ void FStarshipEditorStyle::FStyle::SetupLocalizationDashboardStyles()
 
 void FStarshipEditorStyle::FStyle::SetupMyBlueprintStyles()
 {
-	Set("MyBlueprint.DeleteEntry", new IMAGE_BRUSH("Icons/Edit/icon_Edit_Delete_40x", Icon16x16));
+	Set( "MyBlueprint.DeleteEntry", new IMAGE_BRUSH("Icons/GeneralTools/Delete_40x", Icon16x16));
 }
 
 void FStarshipEditorStyle::FStyle::SetupStatusBarStyle()
