StarshipStyle: Fix a bug where changing the Editor Window Background Color in Editor Preferences was not working because the stylesheet was referring to a pointer to the style settings that had been partially removed

diff --git a/Engine/Source/Editor/EditorStyle/Private/StarshipStyle.cpp b/Engine/Source/Editor/EditorStyle/Private/StarshipStyle.cpp
--- a/Engine/Source/Editor/EditorStyle/Private/StarshipStyle.cpp
+++ b/Engine/Source/Editor/EditorStyle/Private/StarshipStyle.cpp
@@ -27,7 +27,6 @@ FName FStarshipEditorStyle::StyleSetName = TEXT("EditorStyle");
 void FStarshipEditorStyle::Initialize()
 {
 	LLM_SCOPE_BYNAME(TEXT("FStarshipEditorStyle"));
-	Settings = NULL;
 
 	// The core style must be initialized before the editor style
 	FSlateApplication::InitializeCoreStyle();
@@ -38,7 +37,7 @@ void FStarshipEditorStyle::Initialize()
 	USlateThemeManager::Get().ApplyTheme(USlateThemeManager::Get().GetCurrentTheme().Id);
 	//UStyleColorTable::Get().SaveCurrentThemeAs(UStyleColorTable::Get().GetCurrentTheme().Filename);
 #endif
-	StyleInstance = Create(Settings);
+	StyleInstance = Create();
 	SetStyle(StyleInstance.ToSharedRef());
 }
 
@@ -58,7 +57,6 @@ BEGIN_SLATE_FUNCTION_BUILD_OPTIMIZATION
  *****************************************************************************/
 
 TSharedPtr< FStarshipEditorStyle::FStyle > FStarshipEditorStyle::StyleInstance = NULL;
-TWeakObjectPtr< UEditorStyleSettings > FStarshipEditorStyle::Settings = NULL;
 
 void FStarshipEditorStyle::FStyle::SetColor(const TSharedRef< FLinearColor >& Source, const FLinearColor& Value)
 {
@@ -80,7 +78,7 @@ bool FStarshipEditorStyle::FStyle::IncludeEditorSpecificStyles()
 /* FStarshipEditorStyle interface
  *****************************************************************************/
 
-FStarshipEditorStyle::FStyle::FStyle( const TWeakObjectPtr< UEditorStyleSettings >& InSettings )
+FStarshipEditorStyle::FStyle::FStyle(  )
 	: FSlateStyleSet(FStarshipEditorStyle::StyleSetName)
 
 	// Note, these sizes are in Slate Units.
@@ -118,8 +116,6 @@ FStarshipEditorStyle::FStyle::FStyle( const TWeakObjectPtr< UEditorStyleSettings
 	, HighlightColor( HighlightColor_LinearRef )
 	, WindowHighlightColor(WindowHighlightColor_LinearRef)
 	, InheritedFromBlueprintTextColor(FLinearColor(0.25f, 0.5f, 1.0f))
-
-	, Settings( InSettings )
 {
 }
 
@@ -127,7 +123,8 @@ FStarshipEditorStyle::FStyle::~FStyle()
 {
 
 #ifdef WITH_EDITOR
-	if (Settings.IsValid())
+	UEditorStyleSettings* Settings = GetMutableDefault<UEditorStyleSettings>();
+	if (Settings)
 	{
 		Settings->OnSettingChanged().Remove(SettingChangedHandler);
 	}
@@ -142,7 +139,8 @@ void FStarshipEditorStyle::FStyle::SettingsChanged(FName PropertyName)
 
 void FStarshipEditorStyle::FStyle::SyncSettings()
 {
-	if (Settings.IsValid())
+	UEditorStyleSettings* Settings = GetMutableDefault<UEditorStyleSettings>();
+	if (Settings)
 	{
 		// The subdued selection color is derived from the selection color
 		auto SubduedSelectionColor = Settings->GetSubduedSelectionColor();
@@ -259,7 +257,8 @@ void FStarshipEditorStyle::FStyle::Initialize()
 	SyncSettings();
 
 #ifdef WITH_EDITOR
-	if (Settings.IsValid())
+	UEditorStyleSettings* Settings = GetMutableDefault<UEditorStyleSettings>();
+	if (Settings)
 	{
 		SettingChangedHandler = Settings->OnSettingChanged().AddRaw(this, &FStarshipEditorStyle::FStyle::SettingsChanged);
 	}
diff --git a/Engine/Source/Editor/EditorStyle/Private/StarshipStyle.h b/Engine/Source/Editor/EditorStyle/Private/StarshipStyle.h
--- a/Engine/Source/Editor/EditorStyle/Private/StarshipStyle.h
+++ b/Engine/Source/Editor/EditorStyle/Private/StarshipStyle.h
@@ -40,7 +40,8 @@ public:
 	class FStyle : public FSlateStyleSet
 	{
 	public:
-		FStyle(const TWeakObjectPtr< UEditorStyleSettings >& InSettings);
+
+		FStyle();
 		~FStyle();
 
 		void Initialize();
@@ -138,21 +139,18 @@ public:
 		FSlateFontInfo NormalFont;
 
 		FSlateBrush* WindowTitleOverride;
-
-		TWeakObjectPtr< UEditorStyleSettings > Settings;
-
+		
 		FDelegateHandle SettingChangedHandler;
 	};
 
-	static TSharedRef<class FStarshipEditorStyle::FStyle> Create(const TWeakObjectPtr< UEditorStyleSettings >& InCustomization)
+	static TSharedRef<class FStarshipEditorStyle::FStyle> Create()
 	{
-		TSharedRef<class FStarshipEditorStyle::FStyle> NewStyle = MakeShareable(new FStarshipEditorStyle::FStyle(InCustomization));
+		TSharedRef<class FStarshipEditorStyle::FStyle> NewStyle = MakeShareable(new FStarshipEditorStyle::FStyle());
 		NewStyle->Initialize();
 
 		return NewStyle;
 	}
 
 	static TSharedPtr<FStarshipEditorStyle::FStyle> StyleInstance;
-	static TWeakObjectPtr<UEditorStyleSettings> Settings;
 	static FName StyleSetName;
 };
