Fix to Warning in Slate Starship Style. Change to Old/White icon.

diff --git a/Engine/Source/Editor/EditorStyle/Private/StarshipStyle.cpp b/Engine/Source/Editor/EditorStyle/Private/StarshipStyle.cpp
--- a/Engine/Source/Editor/EditorStyle/Private/StarshipStyle.cpp
+++ b/Engine/Source/Editor/EditorStyle/Private/StarshipStyle.cpp
@@ -7639,7 +7639,7 @@ void FStarshipEditorStyle::FStyle::SetupSourceCodeStyles()
 	FTextBlockStyle NormalSourceCodeText = FTextBlockStyle(NormalText)
 		.SetFont(SourceCodeFont);
 	const FTextBlockStyle SourceCodeErrorText = FTextBlockStyle(NormalSourceCodeText)
-		.SetUnderlineBrush(IMAGE_BRUSH("Icons/White", Icon8x8, FLinearColor::Red, ESlateBrushTileType::Both))
+		.SetUnderlineBrush(IMAGE_BRUSH("Old/White", Icon8x8, FLinearColor::Red, ESlateBrushTileType::Both))
 		.SetColorAndOpacity(FLinearColor::Red);
 	
 	Set("SyntaxHighlight.SourceCode.Normal", FTextBlockStyle(NormalSourceCodeText).SetColorAndOpacity(FLinearColor(FColor(189, 183, 107))));
