Changing the color of the "Find in Blueprints" button icon on the "Find Results" blueprint menu to be white instead of black. This makes it easier to see and will now adhere to the editor's default styling for buttons

#rnx

diff --git a/Engine/Source/Editor/EditorStyle/Private/SlateEditorStyle.cpp b/Engine/Source/Editor/EditorStyle/Private/SlateEditorStyle.cpp
--- a/Engine/Source/Editor/EditorStyle/Private/SlateEditorStyle.cpp
+++ b/Engine/Source/Editor/EditorStyle/Private/SlateEditorStyle.cpp
@@ -5830,7 +5830,7 @@ void FSlateEditorStyle::FStyle::SetupPersonaStyle()
 	{
 		Set("FindResults.FindInBlueprints", FTextBlockStyle(NormalText)
 			.SetFont(ICON_FONT(10))
-			.SetColorAndOpacity(FLinearColor(0.f, 0.f, 0.f))
+			.SetColorAndOpacity(FLinearColor::White)
 		);
 
 		Set("FindResults.LockButton_Locked", new IMAGE_BRUSH("Icons/padlock_locked_16x", Icon16x16));
diff --git a/Engine/Source/Editor/EditorStyle/Private/StarshipStyle.cpp b/Engine/Source/Editor/EditorStyle/Private/StarshipStyle.cpp
--- a/Engine/Source/Editor/EditorStyle/Private/StarshipStyle.cpp
+++ b/Engine/Source/Editor/EditorStyle/Private/StarshipStyle.cpp
@@ -5229,7 +5229,7 @@ void FStarshipEditorStyle::FStyle::SetupPersonaStyle()
 	{
 		Set("FindResults.FindInBlueprints", FTextBlockStyle(NormalText)
 			.SetFont(ICON_FONT(10))
-			.SetColorAndOpacity(FLinearColor(0.f, 0.f, 0.f))
+			.SetColorAndOpacity(FLinearColor::White)
 		);
 	}
 
