Allow using EditorStyle in targets that compile against editor

diff --git a/Engine/Source/Editor/EditorStyle/EditorStyle.Build.cs b/Engine/Source/Editor/EditorStyle/EditorStyle.Build.cs
--- a/Engine/Source/Editor/EditorStyle/EditorStyle.Build.cs
+++ b/Engine/Source/Editor/EditorStyle/EditorStyle.Build.cs
@@ -6,7 +6,7 @@ public class EditorStyle : ModuleRules
 {
     public EditorStyle(ReadOnlyTargetRules Target) : base(Target)
 	{
-		if (Target.Type != TargetType.Editor)
+		if (!Target.bCompileAgainstEditor)
 		{
 			throw new BuildException("Unable to instantiate EditorStyle module for non-editor targets.");
 		}
