Details Panel - Overrides widget support phase 1 - UI creation

* added an overrides combobutton builder
* added an overrides menu
* Updated the details single row and category renderers to use the overrides button and menu when an entity is selected, if the CVar UI.EnableOverrides is true
* Added some infrastructure in the display managers to start some property tracking, as the Category rows need to know if any of it's properties are currently overridable\resetable

diff --git a/Engine/Source/Editor/EditorStyle/Private/StarshipStyle.cpp b/Engine/Source/Editor/EditorStyle/Private/StarshipStyle.cpp
--- a/Engine/Source/Editor/EditorStyle/Private/StarshipStyle.cpp
+++ b/Engine/Source/Editor/EditorStyle/Private/StarshipStyle.cpp
@@ -3029,7 +3029,7 @@ void FStarshipEditorStyle::FStyle::SetupPropertyEditorStyles()
 		
 		/*******************************************************************************/
 
-		const FButtonStyle CategoryButton = FButtonStyle()
+		const FButtonStyle CategoryRowButton = FButtonStyle()
 			.SetNormal(FSlateRoundedBoxBrush(FStyleColors::Header, 0.f))
 			.SetHovered(FSlateRoundedBoxBrush(FStyleColors::Input, 0.f))
 			.SetPressed(FSlateRoundedBoxBrush(FStyleColors::Header, 0.f))
@@ -3038,12 +3038,33 @@ void FStarshipEditorStyle::FStyle::SetupPropertyEditorStyles()
 			.SetPressedPadding(FMargin(2.f, 0.f, 2.f, 0.f));
 
 		FComboButtonStyle CategoryComboButton = FComboButtonStyle(FStarshipCoreStyle::GetCoreStyle().GetWidgetStyle<FComboButtonStyle>("ComboButton"))
-		.SetButtonStyle(CategoryButton)
+		.SetButtonStyle(CategoryRowButton)
 		.SetDownArrowPadding(FMargin(2.f, 5.f, 3.f, 5.f))
 		.SetDownArrowImage(CORE_IMAGE_BRUSH_SVG("Starship/Common/ellipsis-vertical-narrow", FVector2D(6, 15)));
-		CategoryComboButton.ButtonStyle = CategoryButton;
+		CategoryComboButton.ButtonStyle = CategoryRowButton;
 		Set( "DetailsView.CategoryComboButton", CategoryComboButton);
 
+		Set( "DetailsView.OverrideAdded", new CORE_IMAGE_BRUSH_SVG("Starship/Common/OverrideAdded", Icon16x16));
+		Set( "DetailsView.OverrideInside", new CORE_IMAGE_BRUSH_SVG("Starship/Common/OverrideInside", Icon16x16));
+		Set( "DetailsView.OverrideHereInside", new CORE_IMAGE_BRUSH_SVG("Starship/Common/OverrideHereInside", Icon16x16));
+		Set( "DetailsView.OverrideHere", new CORE_IMAGE_BRUSH_SVG("Starship/Common/OverrideHere", Icon16x16));	
+		Set( "DetailsView.OverrideOptions", new CORE_IMAGE_BRUSH_SVG("Starship/Common/OverrideOptions", Icon16x16));
+		Set( "DetailsView.OverrideRemoved", new CORE_IMAGE_BRUSH_SVG("Starship/Common/OverrideRemoved", Icon16x16));
+
+		const FButtonStyle OverridesButton = FButtonStyle()
+			.SetNormalForeground(FStyleColors::Primary)
+			.SetNormal(FSlateRoundedBoxBrush(FStyleColors::Primary, 0.f))
+			.SetHovered(FSlateRoundedBoxBrush(FStyleColors::Primary, 0.f))
+			.SetPressed(FSlateRoundedBoxBrush(FStyleColors::Primary, 0.f))
+			.SetDisabled(FSlateRoundedBoxBrush(FStyleColors::Primary, 0.f))
+			.SetNormalPadding(FMargin(2.f, 0.f, 2.f, 0.f))
+			.SetPressedPadding(FMargin(2.f, 0.f, 2.f, 0.f));
+		
+		FComboButtonStyle OverridesComboButton = FComboButtonStyle(FStarshipCoreStyle::GetCoreStyle().GetWidgetStyle<FComboButtonStyle>("ComboButton"))
+		.SetButtonStyle(OverridesButton)
+		.SetDownArrowPadding(FMargin(2.f, 5.f, 3.f, 5.f));
+		OverridesComboButton.ButtonStyle = CategoryRowButton;
+		Set( "DetailsView.OverridesComboButton", OverridesComboButton);
 		Set( "DetailsView.CategoryTop_Hovered",       new FSlateColorBrush(FStyleColors::Hover));
 		Set( "DetailsView.CategoryBottom",            new FSlateColorBrush(FStyleColors::Recessed));
 		
diff --git a/Engine/Source/Editor/PropertyEditor/Private/DetailCategoryBuilderImpl.cpp b/Engine/Source/Editor/PropertyEditor/Private/DetailCategoryBuilderImpl.cpp
--- a/Engine/Source/Editor/PropertyEditor/Private/DetailCategoryBuilderImpl.cpp
+++ b/Engine/Source/Editor/PropertyEditor/Private/DetailCategoryBuilderImpl.cpp
@@ -1126,6 +1126,11 @@ void FDetailCategoryImpl::GetChildren(FDetailNodeList& OutChildren, const bool&
 	GetGeneratedChildren(OutChildren, bInIgnoreVisibility, false);
 }
 
+FName FDetailCategoryImpl::GetObjectName() const
+{
+	return ObjectName;
+}
+
 void FDetailCategoryImpl::GetGeneratedChildren(FDetailNodeList& OutChildren, bool bIgnoreVisibility, bool bIgnoreAdvanced)
 {
 	for (TSharedRef<FDetailTreeNode>& Child : SimpleChildNodes)
diff --git a/Engine/Source/Editor/PropertyEditor/Private/DetailCategoryBuilderImpl.h b/Engine/Source/Editor/PropertyEditor/Private/DetailCategoryBuilderImpl.h
--- a/Engine/Source/Editor/PropertyEditor/Private/DetailCategoryBuilderImpl.h
+++ b/Engine/Source/Editor/PropertyEditor/Private/DetailCategoryBuilderImpl.h
@@ -213,6 +213,11 @@ public:
 	virtual bool ShouldShowOnlyChildren() const override { return bShowOnlyChildren; }
 	virtual FName GetNodeName() const override { return GetCategoryName(); }
 
+	/**
+	 * Returns the name of the object which populates this Category
+	 */
+	FName GetObjectName() const;
+
 	/**
 	 * Gets all generated children with options for ignoring current child visibility or advanced dropdowns
 	 */
diff --git a/Engine/Source/Editor/PropertyEditor/Private/DetailsDisplayManager.cpp b/Engine/Source/Editor/PropertyEditor/Private/DetailsDisplayManager.cpp
--- a/Engine/Source/Editor/PropertyEditor/Private/DetailsDisplayManager.cpp
+++ b/Engine/Source/Editor/PropertyEditor/Private/DetailsDisplayManager.cpp
@@ -73,6 +73,38 @@ FMargin FDetailsDisplayManager::GetTablePadding() const
 	return Style ? Style->GetTablePadding(bIsScrollBarNeeded) : 0;
 }
 
+void FDetailsDisplayManager::UpdatePropertyForCategory(FName InCategoryObjectName, FProperty* Property, bool bAddProperty)
+{
+	TSet<FProperty*> Set;
+	const TSet<FProperty*>* PropertySetPtr =  CategoryNameToUpdatePropertySetMap.Find(InCategoryObjectName);
+
+	if ( PropertySetPtr )
+	{
+		Set = *PropertySetPtr;
+	}
+
+	if ( bAddProperty )
+	{
+		Set.Add(Property );
+	}
+	else
+	{
+		Set.Remove( Property );
+	}
+	
+	CategoryNameToUpdatePropertySetMap.Emplace(InCategoryObjectName, Set);
+}
+
+bool FDetailsDisplayManager::GetCategoryHasAnyUpdatedProperties(FName InCategoryObjectName) const
+{
+	if ( const TSet<FProperty*>* PropertySetPtr =  CategoryNameToUpdatePropertySetMap.Find(InCategoryObjectName) )
+	{
+		TSet<FProperty*> Set = *PropertySetPtr;
+		return Set.Num() > 0;
+	}
+	return false;
+}
+
 bool FDetailsDisplayManager::ShowEmptyCategoryIfRootUObjectHasNoPropertyData(UObject* InNode) const
 {
 	return false;
@@ -118,6 +150,12 @@ bool FDetailsDisplayManager::AddEmptyCategoryToDetailLayoutIfNeeded(TSharedRef<F
 	return false;
 }
 
+TSharedPtr<FPropertyUpdatedWidgetBuilder> FDetailsDisplayManager::GetPropertyUpdatedWidget(FResetToDefault ResetToDefault,
+	bool bIsCategoryUpdateWidget)
+{
+	return nullptr;
+}
+
 bool FDetailsDisplayManager::GetIsScrollBarNeeded() const
 {
 	return bIsScrollBarNeeded;
diff --git a/Engine/Source/Editor/PropertyEditor/Private/SDetailCategoryTableRow.cpp b/Engine/Source/Editor/PropertyEditor/Private/SDetailCategoryTableRow.cpp
--- a/Engine/Source/Editor/PropertyEditor/Private/SDetailCategoryTableRow.cpp
+++ b/Engine/Source/Editor/PropertyEditor/Private/SDetailCategoryTableRow.cpp
@@ -101,6 +101,15 @@ void SDetailCategoryTableRow::Construct(const FArguments& InArgs, TSharedRef<FDe
 	}
 
 	OwnerTableViewWeak = InOwnerTableView;
+	PropertyUpdatedWidgetBuilder = DisplayManager->GetPropertyUpdatedWidget(FExecuteAction::CreateLambda( [this]
+		{
+			UE_LOG(LogTemp, Warning, TEXT("DisplayManager->GetPropertyUpdatedWidget reset"));
+		}), true);
+	if (PropertyUpdatedWidgetBuilder.IsValid())
+	{
+		TAttribute<bool> IsHovered = TAttribute<bool>::CreateSP( this, &SDetailCategoryTableRow::IsHovered);
+		PropertyUpdatedWidgetBuilder->Bind_IsRowHovered(IsHovered);
+	}
 	
 	this->ChildSlot
 	[
@@ -137,18 +146,42 @@ void SDetailCategoryTableRow::Construct(const FArguments& InArgs, TSharedRef<FDe
 				return DetailsViewStyle ? DetailsViewStyle->GetCategoryButtonsMargin() : 0;
 			})
 			[
+
+			SNew(SHorizontalBox)
+			+SHorizontalBox::Slot()
+			.HAlign(HAlign_Left)
+			.VAlign(VAlign_Center)
+			.AutoWidth()
+			[
 			DisplayManager.IsValid() ?
-				*FCategoryMenuComboButtonBuilder( DisplayManager.ToSharedRef() )
-					.Set_OnGetContent(FOnGetContent::CreateLambda([this]
-					{
-						const TSharedPtr<SWidget> Menu = DisplayManager->GetCategoryMenu(ObjectName);
-						return Menu.IsValid() ? Menu.ToSharedRef() : SNullWidget::NullWidget;
-					}))
-					.Bind_IsVisible(TAttribute<EVisibility>::CreateLambda([this]
-					{
-						return IsHovered() ? EVisibility::Visible : EVisibility::Collapsed;
-					})) :
-					SNullWidget::NullWidget
+			*FCategoryMenuComboButtonBuilder( DisplayManager.ToSharedRef() )
+				.Set_OnGetContent(FOnGetContent::CreateLambda([this]
+				{
+					const TSharedPtr<SWidget> Menu = DisplayManager->GetCategoryMenu(ObjectName);
+					return Menu.IsValid() ? Menu.ToSharedRef() : SNullWidget::NullWidget;
+				}))
+				.Bind_IsVisible(TAttribute<EVisibility>::CreateLambda([this]
+				{
+					return IsHovered() ? EVisibility::Visible : EVisibility::Collapsed;
+				})) :
+				SNullWidget::NullWidget
+			]
+			+SHorizontalBox::Slot()
+			.HAlign(HAlign_Left)
+			.VAlign(VAlign_Center)
+			.AutoWidth()
+			[
+				PropertyUpdatedWidgetBuilder.IsValid() ?
+						PropertyUpdatedWidgetBuilder->Bind_IsVisible(TAttribute<EVisibility>::CreateLambda([this]
+							{
+								if (!DisplayManager->GetCategoryHasAnyUpdatedProperties(ObjectName))
+								{
+									return EVisibility::Collapsed;
+								}
+								return EVisibility::Visible;
+							})).GenerateWidget().ToSharedRef() :
+							SNullWidget::NullWidget
+			]
 			]
 			] 
 		]
diff --git a/Engine/Source/Editor/PropertyEditor/Private/SDetailSingleItemRow.cpp b/Engine/Source/Editor/PropertyEditor/Private/SDetailSingleItemRow.cpp
--- a/Engine/Source/Editor/PropertyEditor/Private/SDetailSingleItemRow.cpp
+++ b/Engine/Source/Editor/PropertyEditor/Private/SDetailSingleItemRow.cpp
@@ -629,6 +629,25 @@ void SDetailSingleItemRow::Construct( const FArguments& InArgs, FDetailLayoutCus
 				ToolbarBuilder.AddToolBarButton(Extension.UIAction, NAME_None, Extension.Label, Extension.ToolTip, Extension.Icon);
 			}
 
+			FProperty* Property;
+			
+			if (GetPropertyNode().IsValid() &&
+				DetailsView &&
+				DetailsView->GetDisplayManager().IsValid())
+			{
+					
+				Property = GetPropertyNode()->GetProperty();
+				DisplayManager = DetailsView->GetDisplayManager();
+				PropertyUpdatedWidgetBuilder = DisplayManager->GetPropertyUpdatedWidget(
+					FExecuteAction::CreateSP(this, &SDetailSingleItemRow::OnResetToDefaultClicked));
+
+				if (PropertyUpdatedWidgetBuilder.IsValid())
+				{
+					TAttribute<bool> IsHovered = TAttribute<bool>::CreateSP( this, &SDetailSingleItemRow::IsHovered);
+					PropertyUpdatedWidgetBuilder->Bind_IsRowHovered(IsHovered);
+				}
+			}
+
 			Splitter->AddSlot()
 				.Value(ColumnSizeData.GetRightColumnWidth())
 				.OnSlotResized(ColumnSizeData.GetOnRightColumnResized())
@@ -641,6 +660,20 @@ void SDetailSingleItemRow::Construct( const FArguments& InArgs, FDetailLayoutCus
 				.VAlign(VAlign_Center)
 				.Padding(0)
 				[
+					PropertyUpdatedWidgetBuilder.IsValid() ?
+					PropertyUpdatedWidgetBuilder
+						->Bind_IsVisible(
+							TAttribute<EVisibility>::CreateLambda([this, Category, Property ]
+							{
+								const bool bIsResetVisible = IsResetToDefaultVisible();
+								const EVisibility Visibility = bIsResetVisible ? EVisibility::Visible : EVisibility::Collapsed;
+								if (DisplayManager.IsValid() && Category.IsValid() && Property)
+								{
+									DisplayManager->UpdatePropertyForCategory(Category->GetObjectName(), Property, bIsResetVisible);
+								}
+								return Visibility;
+							}))
+						.GenerateWidget().ToSharedRef()  :
 					ToolbarBuilder.MakeWidget()
 				]
 			];
diff --git a/Engine/Source/Editor/PropertyEditor/Private/SDetailTableRowBase.h b/Engine/Source/Editor/PropertyEditor/Private/SDetailTableRowBase.h
--- a/Engine/Source/Editor/PropertyEditor/Private/SDetailTableRowBase.h
+++ b/Engine/Source/Editor/PropertyEditor/Private/SDetailTableRowBase.h
@@ -65,6 +65,13 @@ public:
 	static const float ScrollBarPadding;
 
 protected:
+
+	/**
+	 * The FPropertyUpdatedWidgetBuilder which builds a Property Updated widget which can be used in place of the
+	 * reset to default button
+	 */
+	TSharedPtr<FPropertyUpdatedWidgetBuilder> PropertyUpdatedWidgetBuilder;
+	
 	/**
 	 * Called when the user opens the context menu on this row
 	 *
diff --git a/Engine/Source/Editor/PropertyEditor/Private/UserInterface/Widgets/OverridesComboButtonBuilder.cpp b/Engine/Source/Editor/PropertyEditor/Private/UserInterface/Widgets/OverridesComboButtonBuilder.cpp
new file mode 100644
--- /dev/null
+++ b/Engine/Source/Editor/PropertyEditor/Private/UserInterface/Widgets/OverridesComboButtonBuilder.cpp
@@ -0,0 +1,82 @@
+//  Copyright Epic Games, Inc. All Rights Reserved.
+
+#include "UserInterface/Widgets/OverridesComboButtonBuilder.h"
+#include "Styling/SlateBrush.h"
+
+FOverridesComboButtonBuilder::FOverridesComboButtonBuilder( TSharedRef<FDetailsDisplayManager> InDetailsDisplayManager):
+	FPropertyUpdatedWidgetBuilder(),
+	DisplayManager(InDetailsDisplayManager)
+{
+}
+
+FOverridesComboButtonBuilder& FOverridesComboButtonBuilder::Set_OnGetContent(FOnGetContent InOnGetContent)
+{
+	OnGetContent.Unbind();
+	OnGetContent = InOnGetContent;
+	return *this;
+}
+
+TSharedPtr<SWidget> FOverridesComboButtonBuilder::GenerateWidget()
+{
+	static FComboButtonStyle NormalStyle =  FAppStyle::Get().GetWidgetStyle<FComboButtonStyle>("DetailsView.OverridesComboButton");
+	static FComboButtonStyle HoverStyle =  FAppStyle::Get().GetWidgetStyle<FComboButtonStyle>("DetailsView.OverridesComboButton");
+	static const FSlateBrush* OverrideOptionsBrush = FAppStyle::GetBrush("DetailsView.OverrideOptions");
+	static const FSlateBrush* OverrideHereBrush = FAppStyle::GetBrush("DetailsView.OverrideHere");
+
+	NormalStyle.SetDownArrowImage( *OverrideHereBrush );
+	HoverStyle.SetDownArrowImage( *OverrideOptionsBrush );
+ 
+	bool bShouldShowOverrides = true;
+	return 	bShouldShowOverrides  ?
+				SNew(SHorizontalBox)
+				.Visibility(IsVisible)
+				+SHorizontalBox::Slot()
+					
+				.HAlign(HAlign_Right)
+				.VAlign(VAlign_Center)
+				.AutoWidth()
+				[
+					SNew( SComboButton )
+					.ComboButtonStyle( &NormalStyle )
+					.Visibility_Lambda([this]
+					{
+						if ( IsRowHoveredAttr.IsBound() && IsRowHoveredAttr.Get())
+						{
+							return EVisibility::Collapsed;
+						}
+						return EVisibility::Visible;
+					})
+					.HasDownArrow(true)
+					.OnGetMenuContent(OnGetContent)
+				]
+				+SHorizontalBox::Slot()
+				.HAlign(HAlign_Right)
+				.VAlign(VAlign_Center)
+				.AutoWidth()
+				[
+					SNew( SComboButton )
+					.ComboButtonStyle( &HoverStyle )
+					.Visibility_Lambda([this]
+					{
+						if (IsRowHoveredAttr.IsBound() && IsRowHoveredAttr.Get())
+						{
+							return EVisibility::Visible;
+						}
+						return EVisibility::Collapsed;
+					})
+					.HasDownArrow(true)
+					.OnGetMenuContent(OnGetContent)
+				]
+	:
+				SNullWidget::NullWidget;
+}
+
+TSharedRef<SWidget> FOverridesComboButtonBuilder::operator*()
+{
+	return GenerateWidget().ToSharedRef();
+}
+
+FOverridesComboButtonBuilder::~FOverridesComboButtonBuilder()
+{
+	OnGetContent.Unbind();
+}
\ No newline at end of file
diff --git a/Engine/Source/Editor/PropertyEditor/Private/UserInterface/Widgets/PropertyUpdatedWidgetBuilder.cpp b/Engine/Source/Editor/PropertyEditor/Private/UserInterface/Widgets/PropertyUpdatedWidgetBuilder.cpp
new file mode 100644
--- /dev/null
+++ b/Engine/Source/Editor/PropertyEditor/Private/UserInterface/Widgets/PropertyUpdatedWidgetBuilder.cpp
@@ -0,0 +1,30 @@
+// Copyright Epic Games, Inc. All Rights Reserved.
+
+#include "UserInterface/Widgets/PropertyUpdatedWidgetBuilder.h"
+
+FPropertyUpdatedWidgetBuilder::FPropertyUpdatedWidgetBuilder() : FToolElementRegistrationArgs("FPropertyUpdatedWidget")
+{
+}
+
+FPropertyUpdatedWidgetBuilder& FPropertyUpdatedWidgetBuilder::Bind_IsVisible(TAttribute<EVisibility> InIsVisible)
+{
+	IsVisible = InIsVisible;
+	return *this;
+}
+
+FPropertyUpdatedWidgetBuilder& FPropertyUpdatedWidgetBuilder::Set_ResetToDefault(FResetToDefault InResetToDefault)
+{
+	ResetToDefault.Unbind();
+	ResetToDefault = InResetToDefault;
+	return *this;
+}
+
+FPropertyUpdatedWidgetBuilder::FResetToDefault FPropertyUpdatedWidgetBuilder::Get_ResetToDefault()
+{
+	return ResetToDefault;
+}
+
+void FPropertyUpdatedWidgetBuilder::Bind_IsRowHovered(TAttribute<bool> InIsRowHoveredAttr)
+{
+	IsRowHoveredAttr = InIsRowHoveredAttr;
+}
diff --git a/Engine/Source/Editor/PropertyEditor/Public/DetailsDisplayManager.h b/Engine/Source/Editor/PropertyEditor/Public/DetailsDisplayManager.h
--- a/Engine/Source/Editor/PropertyEditor/Public/DetailsDisplayManager.h
+++ b/Engine/Source/Editor/PropertyEditor/Public/DetailsDisplayManager.h
@@ -3,6 +3,7 @@
 
 #include "DetailsViewStyleKey.h"
 #include "Templates/SharedPointer.h"
+#include "UserInterface/Widgets/PropertyUpdatedWidgetBuilder.h"
 #include "Widgets/SWidget.h"
 
 DECLARE_DELEGATE(FOnDetailsNeedsUpdate)
@@ -16,6 +17,8 @@ class FDetailLayoutBuilderImpl;
 class FDetailsDisplayManager : public TSharedFromThis<FDetailsDisplayManager>
 {
 public:
+	DECLARE_DELEGATE(FResetToDefault);
+	
 	FOnDetailsNeedsUpdate OnDetailsNeedsUpdate;
 
 	PROPERTYEDITOR_API FDetailsDisplayManager();
@@ -104,6 +107,12 @@ public:
 	 */	
 	PROPERTYEDITOR_API virtual bool AddEmptyCategoryToDetailLayoutIfNeeded(TSharedRef<FComplexPropertyNode> Node, TSharedRef<FDetailLayoutBuilderImpl> DetailLayoutBuilder);
 
+	virtual TSharedPtr<FPropertyUpdatedWidgetBuilder> GetPropertyUpdatedWidget(FResetToDefault ResetToDefault, bool bIsCategoryUpdateWidget = false);
+
+	void UpdatePropertyForCategory(FName InCategoryObjectName, FProperty* Property, bool bAddProperty);
+
+	bool GetCategoryHasAnyUpdatedProperties(FName InCategoryObjectName) const;
+
 	/**
 	* Returns true if the specified UObject is a Root Node Object and should show an empty Category stub even if it
 	* has no UProperty Data, else it returns false
@@ -135,5 +144,8 @@ protected:
 	* here means that in this value the work has been done to figure out if the scrollbar should show 
 	 */
 	bool bIsScrollBarNeeded = false;
+
+	/** A map of category object name to a set of properties that has been updated for it */
+	TMap<FName, TSet<FProperty*>> CategoryNameToUpdatePropertySetMap;
 };
 
diff --git a/Engine/Source/Editor/PropertyEditor/Public/UserInterface/Widgets/OverridesComboButtonBuilder.h b/Engine/Source/Editor/PropertyEditor/Public/UserInterface/Widgets/OverridesComboButtonBuilder.h
new file mode 100644
--- /dev/null
+++ b/Engine/Source/Editor/PropertyEditor/Public/UserInterface/Widgets/OverridesComboButtonBuilder.h
@@ -0,0 +1,58 @@
+//  Copyright Epic Games, Inc. All Rights Reserved.
+
+#pragma once
+#include "DetailsDisplayManager.h"
+#include "PropertyUpdatedWidgetBuilder.h"
+#include "Widgets/SNullWidget.h"
+#include "Widgets/Input/SComboButton.h"
+#include "Templates/SharedPointer.h"
+
+class FDetailsDisplayManager;
+
+DECLARE_DELEGATE_OneParam(FGetIsRowHoveredOver, bool)
+
+/**
+ * A Display builder for the overrides combo button 
+ */
+class FOverridesComboButtonBuilder : public FPropertyUpdatedWidgetBuilder
+{
+public:
+
+	/**
+	 * the delegate to get the menu content
+	 */
+	FOnGetContent OnGetContent;
+
+	/**
+	 * The constructor, which takes a @code TSharedRef<FDetailsDisplayManager> @endcode to initialize
+	 * the Details Display Manager
+	 */
+	PROPERTYEDITOR_API FOverridesComboButtonBuilder( TSharedRef<FDetailsDisplayManager> InDetailsDisplayManager );
+
+	/**
+	 * Set the OnGetContent for the menu that this button is responsible for
+	 */
+	PROPERTYEDITOR_API FOverridesComboButtonBuilder& Set_OnGetContent(FOnGetContent InOnGetContent);
+
+	/**
+	 * Implements the generation of the Category Menu button SWidget
+	 */
+	virtual TSharedPtr<SWidget> GenerateWidget() override;
+
+	/**
+	 * Converts this into the SWidget it builds
+	 */
+	PROPERTYEDITOR_API virtual ~FOverridesComboButtonBuilder() override;
+
+	/**
+	 * Converts this into the SWidget it builds
+	 */
+	TSharedRef<SWidget> operator*();
+
+private:
+	/**
+	 * The @code DetailsDisplayManager @endcode which provides an API to manage some of the characteristics of the
+	 * details display
+	 */
+	TSharedRef<FDetailsDisplayManager> DisplayManager;
+};
diff --git a/Engine/Source/Editor/PropertyEditor/Public/UserInterface/Widgets/PropertyUpdatedWidgetBuilder.h b/Engine/Source/Editor/PropertyEditor/Public/UserInterface/Widgets/PropertyUpdatedWidgetBuilder.h
new file mode 100644
--- /dev/null
+++ b/Engine/Source/Editor/PropertyEditor/Public/UserInterface/Widgets/PropertyUpdatedWidgetBuilder.h
@@ -0,0 +1,60 @@
+// Copyright Epic Games, Inc. All Rights Reserved.
+
+#pragma once
+
+#include "ToolElementRegistry.h"
+#include "Layout/Visibility.h"
+#include "Misc/Attribute.h"
+
+class FPropertyUpdatedWidgetBuilder : public FToolElementRegistrationArgs
+{
+public:
+
+	DECLARE_DELEGATE(FResetToDefault);
+
+	/**
+	 * default constructor
+	 */
+	FPropertyUpdatedWidgetBuilder();
+
+	/**
+	 * Set the IsVisible delegate for the menu that this button is responsible for
+	 */
+	FPropertyUpdatedWidgetBuilder& Bind_IsVisible(TAttribute<EVisibility> InIsVisible);
+
+	/**
+	* Sets the ResetToDefault delegate
+	*
+	* @param InResetToDefault the FResetToDefault delegate to set for this widget builder
+	*/
+	FPropertyUpdatedWidgetBuilder& Set_ResetToDefault(FResetToDefault InResetToDefault);
+
+	/**
+	 * returns the ResetToDefault delegate
+	 */
+	FResetToDefault Get_ResetToDefault();
+
+	/**
+	 *  binds the IsRowHovered Attribute
+	 *
+	 * @param bInIsRowHoveredAttr the attribute which will determine if the row for this property is hovered
+	 */
+	void Bind_IsRowHovered(TAttribute<bool> InIsRowHoveredAttr);
+
+	/**
+	 * If true, the row that this is associated with is currently hovered over
+	 */
+	TAttribute<bool> IsRowHoveredAttr;
+protected:
+	
+	/**
+	 * the attribute which tells whether this menu button is visible
+	 */
+	TAttribute<EVisibility> IsVisible;
+
+	/**
+	 * A delegate for resetting the active property row
+	 */
+	FResetToDefault ResetToDefault;
+
+};
