Sequencer: Column UX

- Update icons to 14x14
- Remove inner border
- Update opacities: Default 10%, Row Hover 25%, Row and Icon Hover 80%, Selected 100%

diff --git a/Engine/Source/Editor/EditorStyle/Private/StarshipStyle.cpp b/Engine/Source/Editor/EditorStyle/Private/StarshipStyle.cpp
--- a/Engine/Source/Editor/EditorStyle/Private/StarshipStyle.cpp
+++ b/Engine/Source/Editor/EditorStyle/Private/StarshipStyle.cpp
@@ -2299,11 +2299,11 @@ void FStarshipEditorStyle::FStyle::SetupSequencerStyles()
 	// Sequencer outliner column UI
 	if (IncludeEditorSpecificStyles())
 	{
-		Set("Sequencer.Column.Mute", new IMAGE_BRUSH_SVG("Sequencer/Column_Widgets/SequencerMute", Icon12x12));
-		Set("Sequencer.Column.Locked", new IMAGE_BRUSH_SVG("Sequencer/Column_Widgets/SequencerLocked", Icon12x12));
-		Set("Sequencer.Column.Solo", new IMAGE_BRUSH_SVG("Sequencer/Column_Widgets/SequencerSolo", Icon12x12));
-		Set("Sequencer.Column.Unpinned", new IMAGE_BRUSH_SVG("Sequencer/Column_Widgets/SequencerUnpinned", Icon12x12));
-		Set("Sequencer.Column.CheckBoxIndeterminate", new CORE_IMAGE_BRUSH_SVG("Starship/CoreWidgets/CheckBox/CheckBoxIndeterminate_12", Icon12x12));
+		Set("Sequencer.Column.Mute", new IMAGE_BRUSH_SVG("Sequencer/Column_Widgets/SequencerMute", Icon14x14));
+		Set("Sequencer.Column.Locked", new IMAGE_BRUSH_SVG("Sequencer/Column_Widgets/SequencerLocked", Icon14x14));
+		Set("Sequencer.Column.Solo", new IMAGE_BRUSH_SVG("Sequencer/Column_Widgets/SequencerSolo", Icon14x14));
+		Set("Sequencer.Column.Unpinned", new IMAGE_BRUSH_SVG("Sequencer/Column_Widgets/SequencerUnpinned", Icon14x14));
+		Set("Sequencer.Column.CheckBoxIndeterminate", new CORE_IMAGE_BRUSH_SVG("Starship/CoreWidgets/CheckBox/CheckBoxIndeterminate_14", Icon14x14));
 		Set("Sequencer.Column.OutlinerColumnBox", new FSlateRoundedBoxBrush(FStyleColors::Header, 2.0f));
 	}
 }
diff --git a/Engine/Source/Editor/Sequencer/Private/Widgets/OutlinerColumns/SColumnToggleWidget.cpp b/Engine/Source/Editor/Sequencer/Private/Widgets/OutlinerColumns/SColumnToggleWidget.cpp
--- a/Engine/Source/Editor/Sequencer/Private/Widgets/OutlinerColumns/SColumnToggleWidget.cpp
+++ b/Engine/Source/Editor/Sequencer/Private/Widgets/OutlinerColumns/SColumnToggleWidget.cpp
@@ -126,7 +126,7 @@ FSlateColor SColumnToggleWidget::GetImageColorAndOpacity() const
 	else if (bIsMouseOverWidget)
 	{
 		// Mouse is over widget and it is not directly active.
-		Opacity = .65f;
+		Opacity = .8f;
 	}
 	else if (bIsImplicitlyActive)
 	{
@@ -136,12 +136,12 @@ FSlateColor SColumnToggleWidget::GetImageColorAndOpacity() const
 	else if (Editor->GetOutliner()->GetHoveredItem() == OutlinerItem)
 	{
 		// Mouse is hovered over outliner item and not the widget itself, preview icons for widget.
-		Opacity = .2f;
+		Opacity = .25f;
 	}
 	else
 	{
 		// Not active in any way and mouse is not over widget or item.
-		Opacity = 0.0f;
+		Opacity = 0.1f;
 	}
 
 	OutColor.A = Opacity;
diff --git a/Engine/Source/Editor/SequencerCore/Private/MVVM/Views/SOutlinerColumnBorder.cpp b/Engine/Source/Editor/SequencerCore/Private/MVVM/Views/SOutlinerColumnBorder.cpp
--- a/Engine/Source/Editor/SequencerCore/Private/MVVM/Views/SOutlinerColumnBorder.cpp
+++ b/Engine/Source/Editor/SequencerCore/Private/MVVM/Views/SOutlinerColumnBorder.cpp
@@ -22,7 +22,7 @@ void SOutlinerColumnInnerBorder::Construct(const FArguments& InArgs, const FCrea
 
 	TSharedPtr<SScaleBox> UniformScaleBox;
 
-	// Size of outliner column widgets stretch to desired height and are usually 12x12 based off padding and fixed width of columns
+	// Size of outliner column widgets stretch to desired height and are usually 14x14 based off padding and fixed width of columns
 	TSharedRef<SWidget>	FinalWidget = SNew(SVerticalBox)
 		+ SVerticalBox::Slot()
 		.AutoHeight()
@@ -31,13 +31,7 @@ void SOutlinerColumnInnerBorder::Construct(const FArguments& InArgs, const FCrea
 		[
 			SAssignNew(UniformScaleBox, SScaleBox)
 			[
-				SNew(SBorder)
-				.VAlign(VAlign_Center)
-				.BorderImage(FAppStyle::GetBrush("Sequencer.Column.OutlinerColumnBox"))
-				.Padding(FMargin(1.f))
-				[
-					InArgs._Content.Widget
-				]
+				InArgs._Content.Widget
 			]
 		];
 
@@ -85,7 +79,6 @@ void SOutlinerColumnBorder::Construct(const FArguments& InArgs, const FCreateOut
 	ChildSlot
 	.VAlign(VAlign_Center)
 	.HAlign(HAlign_Center)
-	.Padding(FMargin(2.0f))
 	[
 		InnerWidget.ToSharedRef()
 	];
diff --git a/Engine/Source/Editor/SequencerCore/Private/MVVM/Views/SOutlinerView.cpp b/Engine/Source/Editor/SequencerCore/Private/MVVM/Views/SOutlinerView.cpp
--- a/Engine/Source/Editor/SequencerCore/Private/MVVM/Views/SOutlinerView.cpp
+++ b/Engine/Source/Editor/SequencerCore/Private/MVVM/Views/SOutlinerView.cpp
@@ -509,7 +509,7 @@ void SOutlinerView::AddOutlinerColumnCallback(TSharedPtr<ISequencerOutlinerColum
 		return SNew(SBox).HeightOverride(10.f);
 	};
 
-	Columns.Add(InColumn->GetColumnName(), FOutlinerViewColumn(GenerateToolColumn, 16.f, true));
+	Columns.Add(InColumn->GetColumnName(), FOutlinerViewColumn(GenerateToolColumn, 22.f, true));
 }
 
 void SOutlinerView::UpdateTrackGutterColumns(TSharedPtr<FEditorViewModel> InEditorViewModel)
