Sequencer: Additional Tree View UX

- Brighter top level hover

diff --git a/Engine/Source/Editor/EditorStyle/Private/StarshipStyle.cpp b/Engine/Source/Editor/EditorStyle/Private/StarshipStyle.cpp
--- a/Engine/Source/Editor/EditorStyle/Private/StarshipStyle.cpp
+++ b/Engine/Source/Editor/EditorStyle/Private/StarshipStyle.cpp
@@ -2018,10 +2018,15 @@ void FStarshipEditorStyle::FStyle::SetupSequencerStyles()
 
 		const FTableRowStyle AlternatingTableRowStyle = GetWidgetStyle<FTableRowStyle>("TableView.AlternatingRow");
 
+		// Top parent hover is 2x brighter than header color
+		FLinearColor BrighterHeader = FStyleColors::Header.GetSpecifiedColor().LinearRGBToHSV();
+		BrighterHeader.B = FMath::Min(1.f, BrighterHeader.B * 2.0f);
+		BrighterHeader = BrighterHeader.HSVToLinearRGB();
+
 		Set("Sequencer.TableView.Row", FTableRowStyle(AlternatingTableRowStyle)
 			.SetUseParentRowBrush(true)
 			.SetParentRowBackgroundBrush(FSlateRoundedBoxBrush(FStyleColors::Header, 2.f, FStyleColors::Transparent, 1.f))
-			.SetParentRowBackgroundHoveredBrush(FSlateRoundedBoxBrush(FStyleColors::SelectHover, 2.f)));
+			.SetParentRowBackgroundHoveredBrush(FSlateRoundedBoxBrush(BrighterHeader, 2.f)));
 
 		Set("Sequencer.IconKeySmartAuto", new IMAGE_BRUSH("Sequencer/IconKeySmartAuto", Icon12x12));
 		Set("Sequencer.IconKeyAuto", new IMAGE_BRUSH("Sequencer/IconKeyAuto", Icon12x12));
