[Slate] Added missing styles for EditorViewportToolBar.Button.*

diff --git a/Engine/Source/Editor/EditorStyle/Private/StarshipStyle.cpp b/Engine/Source/Editor/EditorStyle/Private/StarshipStyle.cpp
--- a/Engine/Source/Editor/EditorStyle/Private/StarshipStyle.cpp
+++ b/Engine/Source/Editor/EditorStyle/Private/StarshipStyle.cpp
@@ -2369,12 +2369,30 @@ void FStarshipEditorStyle::FStyle::SetupViewportStyles()
 			.SetPressedPadding(FMargin(4.0f, 4.0f, 3.0f, 4.0f));
 		Set("EditorViewportToolBar.Button", ViewportMenuButton);
 
+		FButtonStyle ViewportMenuButtonLeft = FButtonStyle(ViewportMenuButton)
+			.SetNormal(		BOX_BRUSH("Starship/EditorViewport/ToolBarLeftGroup", 12.f/25.f, FStyleColors::Dropdown))
+			.SetHovered(	BOX_BRUSH("Starship/EditorViewport/ToolBarLeftGroup", 12.f/25.f, FStyleColors::Hover))
+			.SetPressed(	BOX_BRUSH("Starship/EditorViewport/ToolBarLeftGroup", 12.f/25.f, FStyleColors::Recessed));
+		Set("EditorViewportToolBar.Button.Start", ViewportMenuButtonLeft);
+
+		FButtonStyle ViewportMenuButtonMiddle = FButtonStyle(ViewportMenuButton)
+        	.SetNormal(		BOX_BRUSH("Starship/EditorViewport/ToolBarMiddleGroup", 12.f/25.f, FStyleColors::Dropdown))
+        	.SetHovered(	BOX_BRUSH("Starship/EditorViewport/ToolBarMiddleGroup", 12.f/25.f, FStyleColors::Hover))
+        	.SetPressed(	BOX_BRUSH("Starship/EditorViewport/ToolBarMiddleGroup", 12.f/25.f, FStyleColors::Recessed));
+		Set("EditorViewportToolBar.Button.Middle", ViewportMenuButton);
+
+		FButtonStyle ViewportMenuButtonRight = FButtonStyle(ViewportMenuButton)
+			.SetNormal(		BOX_BRUSH("Starship/EditorViewport/ToolBarRightGroup", 12.f/25.f, FStyleColors::Dropdown))
+			.SetHovered(	BOX_BRUSH("Starship/EditorViewport/ToolBarRightGroup", 12.f/25.f, FStyleColors::Hover))
+			.SetPressed(	BOX_BRUSH("Starship/EditorViewport/ToolBarRightGroup", 12.f/25.f, FStyleColors::Recessed));
+		Set("EditorViewportToolBar.Button.End", ViewportMenuButton);
+
 		Set("EditorViewportToolBar.StartToolbarImage", new BOX_BRUSH("Starship/EditorViewport/ToolBarLeftGroup", 12.f/25.f, FStyleColors::Dropdown));
 
 		const FCheckBoxStyle ViewportMenuToggleLeftButtonStyle = FCheckBoxStyle(ViewportToggleButton)
 			.SetCheckBoxType(ESlateCheckBoxType::ToggleButton)
 			.SetUncheckedImage(		  BOX_BRUSH("Starship/EditorViewport/ToolBarLeftGroup", 12.f/25.f, FStyleColors::Dropdown))
-			.SetUncheckedPressedImage(BOX_BRUSH("Starship/EditorViewport/ToolBarLeftGroup", 12.f/25.f,FStyleColors::Recessed))
+			.SetUncheckedPressedImage(BOX_BRUSH("Starship/EditorViewport/ToolBarLeftGroup", 12.f/25.f, FStyleColors::Recessed))
 			.SetUncheckedHoveredImage(BOX_BRUSH("Starship/EditorViewport/ToolBarLeftGroup", 12.f/25.f, FStyleColors::Hover))
 			.SetCheckedHoveredImage(  BOX_BRUSH("Starship/EditorViewport/ToolBarLeftGroup", 12.f/25.f, FStyleColors::PrimaryHover))
 			.SetCheckedPressedImage(  BOX_BRUSH("Starship/EditorViewport/ToolBarLeftGroup", 12.f/25.f, FStyleColors::PrimaryPress))
@@ -2404,12 +2422,6 @@ void FStarshipEditorStyle::FStyle::SetupViewportStyles()
 			.SetPadding(ViewportMarginRight);
 		Set("EditorViewportToolBar.ToggleButton.End", ViewportMenuToggleRightButtonStyle);
 
-		Set("EditorViewportToolBar.Button.Middle", FButtonStyle(ViewportMenuButton)
-			.SetNormal(BOX_BRUSH("Starship/EditorViewport/ToolBarMiddleGroup", 12.f/25.f, FStyleColors::Dropdown))
-			.SetPressed(BOX_BRUSH("Starship/EditorViewport/ToolBarMiddleGroup", 12.f/25.f, FStyleColors::Recessed))
-			.SetHovered(BOX_BRUSH("Starship/EditorViewport/ToolBarMiddleGroup", 12.f/25.f, FStyleColors::Hover))
-		);
-
 		// We want a background-less version as the ComboMenu has its own unified background
 		const FToolBarStyle& SlimCoreToolBarStyle = FStarshipCoreStyle::GetCoreStyle().GetWidgetStyle<FToolBarStyle>("SlimToolBar");
 
