EditorStyles fixed reference for EditorViewportToolBar.Button.* styles

diff --git a/Engine/Source/Editor/EditorStyle/Private/StarshipStyle.cpp b/Engine/Source/Editor/EditorStyle/Private/StarshipStyle.cpp
--- a/Engine/Source/Editor/EditorStyle/Private/StarshipStyle.cpp
+++ b/Engine/Source/Editor/EditorStyle/Private/StarshipStyle.cpp
@@ -2379,13 +2379,13 @@ void FStarshipEditorStyle::FStyle::SetupViewportStyles()
         	.SetNormal(		BOX_BRUSH("Starship/EditorViewport/ToolBarMiddleGroup", 12.f/25.f, FStyleColors::Dropdown))
         	.SetHovered(	BOX_BRUSH("Starship/EditorViewport/ToolBarMiddleGroup", 12.f/25.f, FStyleColors::Hover))
         	.SetPressed(	BOX_BRUSH("Starship/EditorViewport/ToolBarMiddleGroup", 12.f/25.f, FStyleColors::Recessed));
-		Set("EditorViewportToolBar.Button.Middle", ViewportMenuButton);
+		Set("EditorViewportToolBar.Button.Middle", ViewportMenuButtonMiddle);
 
 		FButtonStyle ViewportMenuButtonRight = FButtonStyle(ViewportMenuButton)
 			.SetNormal(		BOX_BRUSH("Starship/EditorViewport/ToolBarRightGroup", 12.f/25.f, FStyleColors::Dropdown))
 			.SetHovered(	BOX_BRUSH("Starship/EditorViewport/ToolBarRightGroup", 12.f/25.f, FStyleColors::Hover))
 			.SetPressed(	BOX_BRUSH("Starship/EditorViewport/ToolBarRightGroup", 12.f/25.f, FStyleColors::Recessed));
-		Set("EditorViewportToolBar.Button.End", ViewportMenuButton);
+		Set("EditorViewportToolBar.Button.End", ViewportMenuButtonRight);
 
 		Set("EditorViewportToolBar.StartToolbarImage", new BOX_BRUSH("Starship/EditorViewport/ToolBarLeftGroup", 12.f/25.f, FStyleColors::Dropdown));
 
