Submit missing stylesheet updates with overrides support

diff --git a/Engine/Source/Editor/EditorStyle/Private/StarshipStyle.cpp b/Engine/Source/Editor/EditorStyle/Private/StarshipStyle.cpp
--- a/Engine/Source/Editor/EditorStyle/Private/StarshipStyle.cpp
+++ b/Engine/Source/Editor/EditorStyle/Private/StarshipStyle.cpp
@@ -457,7 +457,6 @@ void FStarshipEditorStyle::FStyle::SetupGeneralStyles()
 		Set("Icons.Heart", new IMAGE_BRUSH_SVG("Starship/Common/Heart", Icon16x16));
 		Set("Icons.RecentAssets", new IMAGE_BRUSH_SVG("Starship/Common/RecentAssets", Icon16x16));
 
-
 		Set("Icons.Toolbar.Play", new IMAGE_BRUSH_SVG("Starship/Common/play", Icon20x20));
 		Set("Icons.Toolbar.Pause", new IMAGE_BRUSH_SVG("Starship/MainToolbar/pause", Icon20x20));
 		Set("Icons.Toolbar.Stop", new CORE_IMAGE_BRUSH_SVG("Starship/Common/stop", Icon20x20));
@@ -1109,6 +1108,21 @@ void FStarshipEditorStyle::FStyle::SetupGeneralStyles()
 
 	// Scene Outliner
 	{
+		// Override icons
+		
+		// Background images for outliner override badges (first layer)
+		Set("SceneOutliner.OverrideBase", new IMAGE_BRUSH_SVG("Starship/SceneOutliner/OutlinerOverrideBG", Icon16x16, FStyleColors::AccentBlue));
+		Set("SceneOutliner.OverrideAddedBase", new IMAGE_BRUSH_SVG("Starship/SceneOutliner/OutlinerOverrideBG", Icon16x16, FStyleColors::AccentGreen));
+		Set("SceneOutliner.OverrideRemovedBase", new IMAGE_BRUSH_SVG("Starship/SceneOutliner/OutlinerOverrideBG", Icon16x16, FStyleColors::AccentRed));
+		Set("SceneOutliner.OverrideInsideBase", new IMAGE_BRUSH_SVG("Starship/SceneOutliner/OutlinerOverrideInsideBG", Icon16x16, FStyleColors::Input));
+
+		// Foreground images for all override badges (second layer)
+		Set("SceneOutliner.OverrideHere", new IMAGE_BRUSH_SVG("Starship/SceneOutliner/OutlinerOverrideHere", Icon16x16, FStyleColors::Input));
+		Set("SceneOutliner.OverrideAdded", new IMAGE_BRUSH_SVG("Starship/SceneOutliner/OutlinerOverrideAdded", Icon16x16, FStyleColors::Input));
+		Set("SceneOutliner.OverrideRemoved", new IMAGE_BRUSH_SVG("Starship/SceneOutliner/OutlinerOverrideRemoved", Icon16x16, FStyleColors::Input));
+		Set("SceneOutliner.OverrideInside", new IMAGE_BRUSH_SVG("Starship/SceneOutliner/OutlinerOverrideInside", Icon16x16, FStyleColors::AccentBlue));
+		Set("SceneOutliner.OverrideHereAndInside", new IMAGE_BRUSH_SVG("SceneOutliner/Common/OutlinerOverrideHereAndInside", Icon16x16, FStyleColors::AccentBlue));
+
 		Set( "SceneOutliner.NewFolderIcon", new CORE_IMAGE_BRUSH_SVG("Starship/Common/folder-plus", Icon16x16 ) );
 		Set( "SceneOutliner.FolderClosed",  new CORE_IMAGE_BRUSH_SVG("Starship/Common/folder-closed", Icon16x16, FStyleColors::AccentFolder ) );
 		Set( "SceneOutliner.FolderOpen",    new CORE_IMAGE_BRUSH_SVG("Starship/Common/folder-open", Icon16x16, FStyleColors::AccentFolder ) );
diff --git a/Engine/Source/Runtime/SlateCore/Private/Styling/StarshipCoreStyle.cpp b/Engine/Source/Runtime/SlateCore/Private/Styling/StarshipCoreStyle.cpp
--- a/Engine/Source/Runtime/SlateCore/Private/Styling/StarshipCoreStyle.cpp
+++ b/Engine/Source/Runtime/SlateCore/Private/Styling/StarshipCoreStyle.cpp
@@ -412,19 +412,7 @@ TSharedRef<ISlateStyle> FStarshipCoreStyle::Create()
 		Style->Set("Icons.Recent", new IMAGE_BRUSH_SVG("Starship/Common/Recent", Icon16x16));
 
 		Style->Set("Icons.BadgeModified", new IMAGE_BRUSH_SVG("Starship/Common/badge-modified", Icon16x16));
-
-		// Background images for all override badges (first layer)
-		Style->Set("Icons.BadgeOverrideBase", new IMAGE_BRUSH_SVG("Starship/Common/CircleBadgeDot", Icon16x16, FStyleColors::AccentBlue));
-		Style->Set("Icons.BadgeOverrideAddBase", new IMAGE_BRUSH_SVG("Starship/Common/CircleBadgeDot", Icon16x16, FStyleColors::AccentGreen));
-		Style->Set("Icons.BadgeOverrideRemoveBase", new IMAGE_BRUSH_SVG("Starship/Common/CircleBadgeDot", Icon16x16, FStyleColors::AccentRed));
-
-		// Foreground images for all override badges (second layer)
-		Style->Set("Icons.BadgeOverrideSelf", new IMAGE_BRUSH_SVG("Starship/Common/CircleBadgeBorder", Icon16x16, FStyleColors::Recessed));
-		Style->Set("Icons.BadgeOverrideChild", new IMAGE_BRUSH_SVG("Starship/Common/CircleBadgeOutline", Icon16x16, FStyleColors::Recessed));
-		Style->Set("Icons.BadgeOverrideSelfAndChild", new IMAGE_BRUSH_SVG("Starship/Common/CircleBadgeHalf", Icon16x16, FStyleColors::Recessed));
-		Style->Set("Icons.BadgeOverrideAdd", new IMAGE_BRUSH_SVG("Starship/Common/CircleBadgePlus", Icon16x16, FStyleColors::Recessed));
-		Style->Set("Icons.BadgeOverrideRemove", new IMAGE_BRUSH_SVG("Starship/Common/CircleBadgeMinus", Icon16x16, FStyleColors::Recessed));
-
+		
 		// Background images for all the details panels override states
 		Style->Set("Icons.DetailsOverrideHere", new IMAGE_BRUSH_SVG("Starship/Common/DetailsOverrideHere", Icon16x16, FStyleColors::AccentBlue));
 		Style->Set("Icons.DetailsOverrideAdded", new IMAGE_BRUSH_SVG("Starship/Common/DetailsOverrideAdded", Icon16x16, FStyleColors::AccentGreen));
