Added icons for WorldDataLayers actor.

diff --git a/Engine/Source/Editor/EditorStyle/Private/StarshipStyle.cpp b/Engine/Source/Editor/EditorStyle/Private/StarshipStyle.cpp
--- a/Engine/Source/Editor/EditorStyle/Private/StarshipStyle.cpp
+++ b/Engine/Source/Editor/EditorStyle/Private/StarshipStyle.cpp
@@ -6118,6 +6118,7 @@ void FStarshipEditorStyle::FStyle::SetupClassIconsAndThumbnails()
 			{TEXT("WidgetBlueprintGeneratedClass")},
 			{TEXT("WindDirectionalSource")},
 			{TEXT("World")},
+			{TEXT("WorldDataLayers")},
 		};
 	
 		// SVG Asset icons
