Change shadow brush for virtual folders in asset tile view to match its shape.

diff --git a/Engine/Source/Editor/ContentBrowser/Private/ContentBrowserUtils.cpp b/Engine/Source/Editor/ContentBrowser/Private/ContentBrowserUtils.cpp
--- a/Engine/Source/Editor/ContentBrowser/Private/ContentBrowserUtils.cpp
+++ b/Engine/Source/Editor/ContentBrowser/Private/ContentBrowserUtils.cpp
@@ -484,7 +484,7 @@ bool ContentBrowserUtils::TryGetFolderBrushAndShadowName(const FContentBrowserIt
 	else if (bVirtualFolder && ShouldShowCustomVirtualFolderIcon())
 	{
 		OutBrushName = TEXT("ContentBrowser.ListViewVirtualFolderIcon");
-		OutShadowBrushName = TEXT("None");
+		OutShadowBrushName = TEXT("ContentBrowser.ListViewVirtualFolderShadow");
 	}
 	else if (bPluginFolder && ShouldShowPluginFolderIcon())
 	{
diff --git a/Engine/Source/Editor/EditorStyle/Private/StarshipStyle.cpp b/Engine/Source/Editor/EditorStyle/Private/StarshipStyle.cpp
--- a/Engine/Source/Editor/EditorStyle/Private/StarshipStyle.cpp
+++ b/Engine/Source/Editor/EditorStyle/Private/StarshipStyle.cpp
@@ -6360,6 +6360,7 @@ void FStarshipEditorStyle::FStyle::SetupContentBrowserStyle()
 
 		Set("ContentBrowser.ListViewFolderIcon", new IMAGE_BRUSH_SVG("Starship/ContentBrowser/folder", FVector2D(64, 64)));
 		Set("ContentBrowser.ListViewVirtualFolderIcon", new IMAGE_BRUSH_SVG("Starship/ContentBrowser/FolderLargeVirtual", FVector2D(64, 64)));
+		Set("ContentBrowser.ListViewVirtualFolderShadow", new IMAGE_BRUSH("Starship/ContentBrowser/FolderLargeVirtualShadow", FVector2D(256, 256)));
 		Set("ContentBrowser.ListViewDeveloperFolderIcon", new IMAGE_BRUSH_SVG("Starship/ContentBrowser/FolderLargeDeveloper", FVector2D(64, 64)));
 		Set("ContentBrowser.ListViewCodeFolderIcon", new IMAGE_BRUSH_SVG("Starship/ContentBrowser/FolderLargeCode", FVector2D(64, 64)));
 		Set("ContentBrowser.ListViewPluginFolderIcon", new IMAGE_BRUSH_SVG("Starship/ContentBrowser/FolderLargePlugin", FVector2D(64, 64)));
