Add icon for browse to nanite override material.

diff --git a/Engine/Source/Editor/EditorStyle/Private/StarshipStyle.cpp b/Engine/Source/Editor/EditorStyle/Private/StarshipStyle.cpp
--- a/Engine/Source/Editor/EditorStyle/Private/StarshipStyle.cpp
+++ b/Engine/Source/Editor/EditorStyle/Private/StarshipStyle.cpp
@@ -776,6 +776,7 @@ void FStarshipEditorStyle::FStyle::SetupGeneralStyles()
 		Set( "MaterialList.DragDropBorder", new BOX_BRUSH( "Old/Window/ViewportDebugBorder", 0.8f ) );
 		Set( "MaterialList.HyperlinkStyle", FTextBlockStyle(NormalText) .SetFont( DEFAULT_FONT( "Regular", 8 ) ) );
 		Set( "MaterialList.HyperlinkStyle.ShadowOffset", FVector2D::ZeroVector );
+		Set( "Icons.NaniteBrowseContent", new IMAGE_BRUSH_SVG("Starship/Common/NaniteBrowseContent", Icon16x16));
 	}
 
 	// Dialogue Wave Details
diff --git a/Engine/Source/Editor/PropertyEditor/Private/MaterialList.cpp b/Engine/Source/Editor/PropertyEditor/Private/MaterialList.cpp
--- a/Engine/Source/Editor/PropertyEditor/Private/MaterialList.cpp
+++ b/Engine/Source/Editor/PropertyEditor/Private/MaterialList.cpp
@@ -343,7 +343,7 @@ TSharedRef<SWidget> FMaterialItemView::MakeBrowseNaniteOverrideMaterialButton()
 				}))
 			[ 
 				SNew(SImage)
-				.Image(FAppStyle::Get().GetBrush("Icons.BrowseContent")) //todo: UE-168435 Get custom icon for this.
+				.Image(FAppStyle::Get().GetBrush("Icons.NaniteBrowseContent"))
 				.ColorAndOpacity(FSlateColor::UseForeground())
 			]
 		];
