Refine styles for toolbar and buttons within.

Adding label visibility updates to the category toolbars and clipping vertical box, and refining the styles for the view to match UX.

diff --git a/Engine/Source/Developer/WidgetRegistration/Private/Layout/CategoryDrivenContentBuilder.cpp b/Engine/Source/Developer/WidgetRegistration/Private/Layout/CategoryDrivenContentBuilder.cpp
--- a/Engine/Source/Developer/WidgetRegistration/Private/Layout/CategoryDrivenContentBuilder.cpp
+++ b/Engine/Source/Developer/WidgetRegistration/Private/Layout/CategoryDrivenContentBuilder.cpp
@@ -242,18 +242,19 @@ void FCategoryDrivenContentBuilder::SetShowNoCategorySelection(bool bInShowNoCat
 	}
 }
 
-void FCategoryDrivenContentBuilder::SetCommands(TArray<TSharedPtr<FUICommandInfo>> InContentLoaderCommands)
+void FCategoryDrivenContentBuilder::InitializeCategoryButtons(TArray<UE::DisplayBuilders::FBuilderInput> InBuilderInputArray)
 {
-	// TODO delete method
+	BuilderInputArray = InBuilderInputArray;
+	InitializeCategoryButtons();
 }
 
-void FCategoryDrivenContentBuilder::InitializeCategoryButtons(TArray<UE::DisplayBuilders::FBuilderInput> InBuilderInputArray)
+void FCategoryDrivenContentBuilder::InitializeCategoryButtons()
 {
 	CategoryNameToBuilderInputMap.Empty();
 	
-	for ( int32 Index = 0; Index < InBuilderInputArray.Num(); Index++ )
+	for ( int32 Index = 0; Index < BuilderInputArray.Num(); Index++ )
 	{
-		UE::DisplayBuilders::FBuilderInput& BuilderInput = InBuilderInputArray[Index];
+		UE::DisplayBuilders::FBuilderInput& BuilderInput = BuilderInputArray[Index];
 		BuilderInput.Index = Index;
 		
 		if ( BuilderInput.Name == FavoritesCategoryName && GetDecoratedButtonDelegate.IsBound() )
@@ -268,7 +269,9 @@ void FCategoryDrivenContentBuilder::InitializeCategoryButtons(TArray<UE::Display
 	
 	LoadToolPaletteCommandList = MakeShared<FUICommandList>();
 	LoadPaletteToolBarBuilder = MakeShared<FVerticalToolBarBuilder>(LoadToolPaletteCommandList, FMultiBoxCustomization::None, TSharedPtr<FExtender>(), bForceSmallIcons);
-	LoadPaletteToolBarBuilder->SetStyle(&FAppStyle::Get(), "FCategoryDrivenContentBuilderToolbar");
+	
+	const FName StyleName = GetCategoryToolBarStyleName(); 
+	LoadPaletteToolBarBuilder->SetStyle(&FAppStyle::Get(), StyleName );
 	LoadPaletteToolBarBuilder->SetLabelVisibility( CategoryButtonLabelVisibility );
 	MainContentVerticalBox = SNew(SVerticalBox);
 	InitializeCategoryToolbar();
diff --git a/Engine/Source/Developer/WidgetRegistration/Private/Layout/CategoryDrivenContentBuilderBase.cpp b/Engine/Source/Developer/WidgetRegistration/Private/Layout/CategoryDrivenContentBuilderBase.cpp
--- a/Engine/Source/Developer/WidgetRegistration/Private/Layout/CategoryDrivenContentBuilderBase.cpp
+++ b/Engine/Source/Developer/WidgetRegistration/Private/Layout/CategoryDrivenContentBuilderBase.cpp
@@ -14,6 +14,9 @@
 namespace UE::CategoryDrivenContentBuilderBase::Private
 {
 	const FName StyleName = TEXT("Name");
+	const float LabelledButtonToolbarButtonWidth = 64;
+	const float UnlabelledButtonToolbarButtonWidth = 44;
+	
 }
 
 FToolElementRegistry FCategoryDrivenContentBuilderBase::ToolRegistry = FToolElementRegistry::Get();
@@ -58,14 +61,26 @@ void FCategoryDrivenContentBuilderBase::InitCategoryToolbarContainerWidget()
 	{
 		CategoryToolbarVBox->ClearChildren();
 	}
+
+	float ToolbarBoxWidth = UE::CategoryDrivenContentBuilderBase::Private::UnlabelledButtonToolbarButtonWidth;
+
+	if ( CategoryButtonLabelVisibility.IsVisible() )
+	{
+		ToolbarBoxWidth = UE::CategoryDrivenContentBuilderBase::Private::LabelledButtonToolbarButtonWidth;
+	}
+	
 	CategoryToolbarVBox->AddSlot()
 	.Padding(0.f)
 	[
-		CreateToolbarWidget()
+		SNew(SBox)
+		.WidthOverride( ToolbarBoxWidth )
+		[
+			CreateToolbarWidget()
+		]
 	];
 }
 
-void FCategoryDrivenContentBuilderBase::RefreshCategoryToolbarWidget()
+void FCategoryDrivenContentBuilderBase::RefreshCategoryToolbarWidget(bool bShouldReinitialize)
 {
 	if ( !LoadPaletteToolBarBuilder.IsValid() )
 	{
@@ -73,10 +88,17 @@ void FCategoryDrivenContentBuilderBase::RefreshCategoryToolbarWidget()
 	}
 	FToolElementRegistrationKey ElementKey = FToolElementRegistrationKey(BuilderName, EToolElement::Toolbar);
 	VerticalToolbarElement = ToolRegistry.GetToolElementSP(ElementKey);
+	LoadPaletteToolBarBuilder->SetLabelVisibility( CategoryButtonLabelVisibility );
 	const TSharedRef<FToolbarRegistrationArgs> VerticalToolbarRegistrationArgs = MakeShareable<FToolbarRegistrationArgs>(
 		new FToolbarRegistrationArgs(LoadPaletteToolBarBuilder.ToSharedRef()));
 	
-	if (!VerticalToolbarElement.IsValid())
+	if (VerticalToolbarElement.IsValid() && bShouldReinitialize)
+	{
+		ToolRegistry.UnregisterElement(VerticalToolbarElement.ToSharedRef());
+		VerticalToolbarElement = nullptr;
+	}
+	
+	if (!VerticalToolbarElement.IsValid() || bShouldReinitialize)
 	{
 		VerticalToolbarElement = MakeShareable(new FToolElement
 			(BuilderName,
@@ -166,4 +188,12 @@ void FCategoryDrivenContentBuilderBase::SetCategoryButtonLabelVisibility(bool bI
 TSharedRef<SWidget> FCategoryDrivenContentBuilderBase::CreateToolbarWidget() const
 {
 	return ToolRegistry.GenerateWidget(VerticalToolbarElement.ToSharedRef());
+}
+
+
+FName FCategoryDrivenContentBuilderBase::GetCategoryToolBarStyleName() const
+{
+	return CategoryButtonLabelVisibility.IsVisible() ?
+		"CategoryDrivenContentBuilderToolbarWithLabels" :
+		"CategoryDrivenContentBuilderToolbarWithoutLabels";
 }
\ No newline at end of file
diff --git a/Engine/Source/Developer/WidgetRegistration/Private/Persistence/BuilderPersistenceManager.cpp b/Engine/Source/Developer/WidgetRegistration/Private/Persistence/BuilderPersistenceManager.cpp
--- a/Engine/Source/Developer/WidgetRegistration/Private/Persistence/BuilderPersistenceManager.cpp
+++ b/Engine/Source/Developer/WidgetRegistration/Private/Persistence/BuilderPersistenceManager.cpp
@@ -4,7 +4,11 @@
 
 TObjectPtr<UBuilderPersistenceManager> UBuilderPersistenceManager::Instance = nullptr;
 
-FName UBuilderPersistenceManager::FavoritesSuffix = "BuilderFavorites";
+namespace UE::DisplayBuilders::BuilderPersistenceManager
+{
+	FName FavoritesSuffix = "BuilderFavorites";
+	FName ShowButtonLabelsSuffix = "ShowButtonLabels";
+}
 
 void UBuilderPersistenceManager::Initialize()
 {
@@ -12,7 +16,6 @@ void UBuilderPersistenceManager::Initialize()
 	{
 		Instance = NewObject<UBuilderPersistenceManager>(); 
 		Instance->AddToRoot();
-		Instance->LoadEditorConfig();
 	}
 }
 
@@ -27,19 +30,20 @@ void UBuilderPersistenceManager::ShutDown()
 
 TArray<FName> UBuilderPersistenceManager::GetPersistedFavoritesNamesArray(const UE::DisplayBuilders::FBuilderKey& Key)
 {
-	return GetPersistedArrayOfNames(Key, FavoritesSuffix);
+	return GetPersistedArrayOfNames( Key, UE::DisplayBuilders::BuilderPersistenceManager::FavoritesSuffix );
 }
 
 void UBuilderPersistenceManager::SetPersistedFavoritesNamesArray(const UE::DisplayBuilders::FBuilderKey& Key, TArray<FName>& Favorites)
 {
-	return PersistArrayOfNames(Key, FavoritesSuffix, Favorites);
+	return PersistArrayOfNames(Key, UE::DisplayBuilders::BuilderPersistenceManager::FavoritesSuffix, Favorites);
 }
 
 TArray<FName> UBuilderPersistenceManager::GetPersistedArrayOfNames( const UE::DisplayBuilders::FBuilderKey& Key, FName PersistenceKeySuffix )
 {
+	LoadEditorConfig();
 	if ( !Key.IsNone()  && !PersistenceKeySuffix.IsNone() )
 	{
-		if ( const FPersistedNameArray* Settings =  PersistenceSettings.Find( Key.GetKeyWithSuffix( PersistenceKeySuffix ) ) )
+		if ( const FPersistedNameArray* Settings =  SavedNameToPersistedFNameArrayMap.Find( Key.GetKeyWithSuffix( PersistenceKeySuffix ) ) )
 		{
 			return Settings->ArrayOfNamesToPersist;
 		}
@@ -55,8 +59,48 @@ void UBuilderPersistenceManager::PersistArrayOfNames( const UE::DisplayBuilders:
 		return;
 	}
 
-	FPersistedNameArray& Settings =  PersistenceSettings.Add( Key.GetKeyWithSuffix( PersistenceKeySuffix ) );
+	FPersistedNameArray& Settings =  SavedNameToPersistedFNameArrayMap.Add( Key.GetKeyWithSuffix( PersistenceKeySuffix ) );
 
 	Settings.ArrayOfNamesToPersist = ArrayOfNamesToPersist;;
 	SaveEditorConfig();
-}
\ No newline at end of file
+}
+
+bool UBuilderPersistenceManager::GetPersistedBool( const UE::DisplayBuilders::FBuilderKey& Key, FName PersistenceKeySuffix, bool PersistedBoolIfNoneFound )
+{
+	LoadEditorConfig();
+	if ( !Key.IsNone()  && !PersistenceKeySuffix.IsNone() )
+	{
+		if ( const FPersistedBool* BoolSettings =  SavedNameToPersistedBoolMap.Find( Key.GetKeyWithSuffix( PersistenceKeySuffix ) ) )
+		{
+			return BoolSettings->PersistedBool;
+		}
+		else
+		{
+			PersistBool( Key, PersistenceKeySuffix, PersistedBoolIfNoneFound );
+		}
+	}
+
+	return PersistedBoolIfNoneFound;
+}
+
+void UBuilderPersistenceManager::SetPersistedButtonLabelBool( const UE::DisplayBuilders::FBuilderKey& Key, bool InPersistedBool )
+{
+	return PersistBool( Key, UE::DisplayBuilders::BuilderPersistenceManager::ShowButtonLabelsSuffix, InPersistedBool );
+}
+
+bool UBuilderPersistenceManager::GetPersistedButtonLabelBool( const UE::DisplayBuilders::FBuilderKey& Key, bool PersistedBoolIfNoneFound )
+{
+	return GetPersistedBool( Key, UE::DisplayBuilders::BuilderPersistenceManager::ShowButtonLabelsSuffix, PersistedBoolIfNoneFound );
+}
+
+void UBuilderPersistenceManager::PersistBool( const UE::DisplayBuilders::FBuilderKey& Key, FName PersistenceKeySuffix, bool InPersistedBool )
+{
+	if (PersistenceKeySuffix.IsNone())
+	{
+		return;
+	}
+
+	FPersistedBool& Settings =  SavedNameToPersistedBoolMap.Add( Key.GetKeyWithSuffix( PersistenceKeySuffix ) );
+	Settings.PersistedBool = InPersistedBool;;
+	SaveEditorConfig();
+}
diff --git a/Engine/Source/Developer/WidgetRegistration/Private/Styles/BuilderStyleManager.cpp b/Engine/Source/Developer/WidgetRegistration/Private/Styles/BuilderStyleManager.cpp
--- a/Engine/Source/Developer/WidgetRegistration/Private/Styles/BuilderStyleManager.cpp
+++ b/Engine/Source/Developer/WidgetRegistration/Private/Styles/BuilderStyleManager.cpp
@@ -7,10 +7,14 @@
 #include "Misc/Paths.h"
 #include "Styling/SlateStyleMacros.h"
 #include "Styling/SlateStyleRegistry.h"
+#include "Styling/ToolBarStyle.h"
+#include "Styling/AppStyle.h"
+#include "Styling/CoreStyle.h"
+#include "Styling/StyleColors.h"
 
 void FBuilderStyleManager::RegisterSlateIcon( FBuilderIconKey Key )
 {
-	Set( Key.FileNameWithoutExtension, new IMAGE_BRUSH_SVG( Key.RelativePathToFileWithoutExtension, Key.SizeKey.Size ));
+	Set( Key.FileNameWithoutExtension, new CORE_IMAGE_BRUSH_SVG( Key.RelativePathToFileWithoutExtension, Key.SizeKey.Size ));
 	Key.InitializeIcon();
 }
 
@@ -21,8 +25,6 @@ FBuilderStyleManager::FBuilderStyleManager()
 	SetContentRoot(FPaths::EngineContentDir() / TEXT("Editor/Slate/Builders"));
 	SetCoreContentRoot(FPaths::EngineContentDir() / TEXT("Slate"));
 
-	FText Text = FText::GetEmpty();
-	
 	FSlateStyleRegistry::RegisterSlateStyle(*this);
 }
 
diff --git a/Engine/Source/Developer/WidgetRegistration/Private/ToolkitBuilder.cpp b/Engine/Source/Developer/WidgetRegistration/Private/ToolkitBuilder.cpp
--- a/Engine/Source/Developer/WidgetRegistration/Private/ToolkitBuilder.cpp
+++ b/Engine/Source/Developer/WidgetRegistration/Private/ToolkitBuilder.cpp
@@ -169,7 +169,9 @@ void FToolkitBuilder::AddPalette(TSharedPtr<FToolPalette> Palette)
 		PaletteCommandNameToButtonArgsMap.Add(Button->Command->GetCommandName().ToString(), Button);
 		PaletteCommandInfos.Add(Button->Command->GetCommandName().ToString(), Button->Command);
 	}
-	LoadCommandNameToToolPaletteMap.Add(Palette->LoadToolPaletteAction->GetCommandName().ToString(), Palette);
+	const FName CommandName = Palette->LoadToolPaletteAction->GetCommandName();
+	LoadCommandNameToToolPaletteMap.Add(CommandName.ToString(), Palette);
+	LoadCommandArray.Add( CommandName ); 
 
 	LoadToolPaletteCommandList->MapAction(
 				Palette->LoadToolPaletteAction,
@@ -259,7 +261,11 @@ void FToolkitBuilder::InitializeCategoryToolbar(bool bInitLoadToolPaletteMap)
 	Style = FToolkitStyle::Get().GetWidgetStyle<FToolkitWidgetStyle>("FToolkitWidgetStyle");
 	LoadToolPaletteCommandList = MakeShareable(new FUICommandList);
 	LoadPaletteToolBarBuilder = MakeShared<FVerticalToolBarBuilder>(LoadToolPaletteCommandList, FMultiBoxCustomization::None, TSharedPtr<FExtender>(), bForceSmallIcons);
+
 	LoadPaletteToolBarBuilder->SetLabelVisibility( CategoryButtonLabelVisibility );
+	const FName StyleName = GetCategoryToolBarStyleName(); 
+	LoadPaletteToolBarBuilder->SetStyle(&FAppStyle::Get(), StyleName );
+	
 	EditablePalettesArray.Reset();
 
 	if (bInitLoadToolPaletteMap)
@@ -356,6 +362,7 @@ void FToolkitBuilder::CreatePalette(TSharedPtr<FToolPalette> Palette)
 		PaletteToolbarBuilder->AddToolBarButton(PaletteButton.Get());
 	}
 	CreatePaletteWidget(*Palette.Get(), *Element.Get());
+	LoadPaletteToolBarBuilder->SetLastSelectedCommandIndex( LoadCommandArray.IndexOfByKey(Palette->LoadToolPaletteAction->GetCommandName()) );
 }
 
 void FToolkitBuilder::CreatePaletteWidget(FToolPalette& Palette, FToolElement& Element)
diff --git a/Engine/Source/Developer/WidgetRegistration/Public/Layout/CategoryDrivenContentBuilder.h b/Engine/Source/Developer/WidgetRegistration/Public/Layout/CategoryDrivenContentBuilder.h
--- a/Engine/Source/Developer/WidgetRegistration/Public/Layout/CategoryDrivenContentBuilder.h
+++ b/Engine/Source/Developer/WidgetRegistration/Public/Layout/CategoryDrivenContentBuilder.h
@@ -84,15 +84,18 @@ public:
 	 */
 	WIDGETREGISTRATION_API void SetShowNoCategorySelection( bool bInShowNoCategorySelection );
 	
-	WIDGETREGISTRATION_API void SetCommands(TArray<TSharedPtr<FUICommandInfo>> InContentLoaderCommands);
-
 	/*
 	 * Initializes the category buttons with the given FBuilderInputs
 	 *
 	 * @param InBuilderInputArray the array of FBuilderInput instances that will initialize the Category buttons
 	 */
 	WIDGETREGISTRATION_API void InitializeCategoryButtons(TArray<UE::DisplayBuilders::FBuilderInput> InBuilderInputArray);
-	
+
+	/*
+	 * Initializes the category buttons with the given FBuilderInputs
+	 */
+	WIDGETREGISTRATION_API void InitializeCategoryButtons();
+
 	/**
 	 * Converts the SWidget Widget to a FSlateBuilder and adds it to the main content for the currently selected category
 	 *
@@ -150,6 +153,9 @@ private:
 
 private:
 
+	/** The array of BuilderInputs, in the correct display order */
+	TArray<UE::DisplayBuilders::FBuilderInput> BuilderInputArray;
+
 	/** The array of FSlateBuilders that will build the content for the currently selected category */
 	TArray<TSharedRef<FSlateBuilder>> ChildBuilderArray;
 
diff --git a/Engine/Source/Developer/WidgetRegistration/Public/Layout/CategoryDrivenContentBuilderBase.h b/Engine/Source/Developer/WidgetRegistration/Public/Layout/CategoryDrivenContentBuilderBase.h
--- a/Engine/Source/Developer/WidgetRegistration/Public/Layout/CategoryDrivenContentBuilderBase.h
+++ b/Engine/Source/Developer/WidgetRegistration/Public/Layout/CategoryDrivenContentBuilderBase.h
@@ -85,7 +85,7 @@ public:
 	/**
 	 * refreshes the UI display of the category toolbar
 	 */
-    	WIDGETREGISTRATION_API void RefreshCategoryToolbarWidget();
+	WIDGETREGISTRATION_API void RefreshCategoryToolbarWidget( bool bShouldReinitialize = false );
 
 	/** Implements the generation of the TSharedPtr<SWidget> */
 	WIDGETREGISTRATION_API virtual TSharedPtr<SWidget> GenerateWidget() override;
@@ -110,11 +110,7 @@ private:
 	/* The SWidget that is the whole Toolkit */
 	TSharedPtr<SWidget> ToolkitWidget;
 
-    	/**
-    	 * Creates the SWidget
-    	 */
-    	void CreateWidget();
-     
+
 	/** The SVerticalBox which contains the category toolbar */
 	TSharedPtr<SVerticalBox> CategoryToolbarVBox;
 
@@ -122,6 +118,14 @@ private:
 	TSharedPtr<SVerticalBox> ToolkitWidgetContainerVBox;
 
 protected:
+
+	/**
+	 * Creates the SWidget
+	 */
+	void CreateWidget();
+
+	FName GetCategoryToolBarStyleName() const;
+	
 	virtual void UpdateContentForCategory( FName ActiveCategoryName = NAME_None, FText InActiveCategoryText = FText::GetEmpty() ) = 0;
 	
 	/** The SVerticalBox which holds the main content ~ -all but the Category chooser */
diff --git a/Engine/Source/Developer/WidgetRegistration/Public/Persistence/BuilderPersistenceManager.h b/Engine/Source/Developer/WidgetRegistration/Public/Persistence/BuilderPersistenceManager.h
--- a/Engine/Source/Developer/WidgetRegistration/Public/Persistence/BuilderPersistenceManager.h
+++ b/Engine/Source/Developer/WidgetRegistration/Public/Persistence/BuilderPersistenceManager.h
@@ -22,6 +22,19 @@ public:
 	TArray<FName> ArrayOfNamesToPersist;
 };
 
+/**
+ * Manages a Bool to persist
+ */
+USTRUCT()
+struct WIDGETREGISTRATION_API FPersistedBool
+{
+	GENERATED_BODY()
+public:
+	
+	UPROPERTY()
+	bool PersistedBool = false;
+};
+
 
 /**
  * The Builder Persistence Manager handles persistence for Builders through use of FBuilderKeys
@@ -85,9 +98,44 @@ public:
 private:
 
 	UPROPERTY(meta=(EditorConfig))
-	TMap<FString, FPersistedNameArray> PersistenceSettings;
+	TMap<FString, FPersistedNameArray> SavedNameToPersistedFNameArrayMap;
 
-	static FName FavoritesSuffix;
+	UPROPERTY(meta=(EditorConfig))
+	TMap<FString, FPersistedBool> SavedNameToPersistedBoolMap;
 	
 	static TObjectPtr<UBuilderPersistenceManager> Instance;
+
+private:
+	/**
+	 * Sets the default button label EVisisbility for the Builder with FBuilderKey Key
+	 * 
+	 * @param Key the FBuilderKey to retrieve the button label Bool for
+	 * @param Bool the bool for the button labels to persist
+	 */
+	void SetPersistedButtonLabelBool( const UE::DisplayBuilders::FBuilderKey& Key, bool Bool );
+	
+	/**
+	 * @return the bool or the buttons labels for the Builder with FBuilderKey Key, if one exists, else an empty array is returned
+	 * 
+	 * @param Key the FBuilderKey to retrieve the button label bool
+	 * @param PersistedBoolIfNoneFound the Bool to set as the new persistence value, if one is not found
+	 */
+	bool GetPersistedButtonLabelBool( const UE::DisplayBuilders::FBuilderKey& Key, bool PersistedBoolIfNoneFound  );
+
+	/**
+	 * Returns the bool for the Builder with FBuilderKey Key and the suffix that the array was persisted with, if one exists, else an empty array is returned
+	 * 
+	 * @param Key the FBuilderKey to retrieve the bool
+	 * @param PersistenceKeySuffix the suffix to add to the FBuilderKey to get the bool with
+	 * @param PersistedBoolIfNoneFound the Bool to set as the new persistence value, if one is not found
+	 */
+	bool GetPersistedBool( const UE::DisplayBuilders::FBuilderKey& Key, FName PersistenceKeySuffix, bool PersistedBoolIfNoneFound );
+	
+	/**
+	 * Persists the Bool the Builder with FBuilderKey Key and suffix PersistenceKeySuffix
+	 * 
+	 * @param Key the FBuilderKey to retrieve the favorites for
+	 * @param PersistenceKeySuffix the suffix added to the FBuilderKey to persist the Bool with
+	 */
+	void PersistBool( const UE::DisplayBuilders::FBuilderKey& Key, FName PersistenceKeySuffix, bool InPersistedBool );
 };
\ No newline at end of file
diff --git a/Engine/Source/Developer/WidgetRegistration/Public/ToolkitBuilder.h b/Engine/Source/Developer/WidgetRegistration/Public/ToolkitBuilder.h
--- a/Engine/Source/Developer/WidgetRegistration/Public/ToolkitBuilder.h
+++ b/Engine/Source/Developer/WidgetRegistration/Public/ToolkitBuilder.h
@@ -334,4 +334,7 @@ private:
 	
 	/** The display name of the currently active tool */
 	FText ActiveToolDisplayName = FText::GetEmpty();
+
+	/** The array of Command names, in the correct display order */
+	TArray<FName> LoadCommandArray;
 };
\ No newline at end of file
diff --git a/Engine/Source/Editor/EditorStyle/Private/StarshipStyle.cpp b/Engine/Source/Editor/EditorStyle/Private/StarshipStyle.cpp
--- a/Engine/Source/Editor/EditorStyle/Private/StarshipStyle.cpp
+++ b/Engine/Source/Editor/EditorStyle/Private/StarshipStyle.cpp
@@ -1737,10 +1737,11 @@ void FStarshipEditorStyle::FStyle::SetupGeneralStyles()
 		FVerticalToolBarStyle.SetBackgroundPadding(   FMargin(6.0f, 4.0f));
 		FVerticalToolBarStyle.ButtonStyle.SetNormalPadding(FMargin(12.0f, 6.0f));
 		FVerticalToolBarStyle.ButtonStyle.SetPressedPadding(FMargin(12.0f, 6.0f));
-		FVerticalToolBarStyle.SetExpandBrush(CORE_IMAGE_BRUSH_SVG("Starship/Common/ellipsis-horizontal-narrow", FVector2D(16, 16)));
+		FVerticalToolBarStyle.SetExpandBrush(IMAGE_BRUSH("Icons/toolbar_expand_16x", Icon8x8));
 		FVerticalToolBarStyle.SetIconPadding(FMargin(8.0f, 8.0f));
 		FVerticalToolBarStyle.SetIconPaddingWithVisibleLabel(FMargin(8.0f, 8.0f, 8.0f, 2.0f));
  
+		FVerticalToolBarStyle.SetExpandBrush(CORE_IMAGE_BRUSH_SVG("Starship/Common/ellipsis-horizontal-narrow", FVector2D(16, 16)));
 		FCheckBoxStyle CheckBoxStyle = FCheckBoxStyle(FVerticalToolBarStyle.ToggleButton)
 			.SetCheckedImage(FSlateRoundedBoxBrush(FStyleColors::Primary, 4.f, FLinearColor(0, 0, 0, .8), 0.5))
 			.SetCheckedHoveredImage(FSlateRoundedBoxBrush(FStyleColors::PrimaryHover, 4.f, FLinearColor(0, 0, 0, .8), 0.5))
@@ -1762,10 +1763,24 @@ void FStarshipEditorStyle::FStyle::SetupGeneralStyles()
 
 	{
 		// FCategoryDrivenContentBuilder vertical toolbar style
-		FToolBarStyle CategoryDrivenContentBuilderToolBarStyle = GetWidgetStyle<FToolBarStyle>("FVerticalToolBar");
-		CategoryDrivenContentBuilderToolBarStyle.SetIconPaddingWithVisibleLabel(FMargin(8.0f, 12.0f, 8.0f, 8.0f));
+		FToolBarStyle CategoryDrivenContentBuilderToolBarStyle =  GetWidgetStyle<FToolBarStyle>("FVerticalToolBar");
+		CategoryDrivenContentBuilderToolBarStyle.SetIconPadding( FMargin(0, 9, 0, 4 ) );
+		CategoryDrivenContentBuilderToolBarStyle.SetBackgroundPadding( FMargin( 4 ) );
+		CategoryDrivenContentBuilderToolBarStyle.SetIconPaddingWithVisibleLabel( FMargin(0, 9, 0, 0 ) );
+		CategoryDrivenContentBuilderToolBarStyle.SetLabelPadding( FMargin(4, 5, 4, 9) );
 		CategoryDrivenContentBuilderToolBarStyle.SetButtonContentMaxWidth( 56.0f );
-		Set("FCategoryDrivenContentBuilderToolbar", CategoryDrivenContentBuilderToolBarStyle);
+		FTextBlockStyle LabelTextStyle = FTextBlockStyle(NormalText)
+		                                 .SetOverflowPolicy(ETextOverflowPolicy::Ellipsis)
+		                                 .SetFont(DEFAULT_FONT("roboto", FCoreStyle::SmallTextSize));
+		CategoryDrivenContentBuilderToolBarStyle.SetLabelStyle( LabelTextStyle );
+		CategoryDrivenContentBuilderToolBarStyle.SetButtonPadding( FMargin( 0, 2, 0, 2 ));
+		CategoryDrivenContentBuilderToolBarStyle.SetShowLabels( true );
+		Set("CategoryDrivenContentBuilderToolbarWithLabels", CategoryDrivenContentBuilderToolBarStyle);
+
+		CategoryDrivenContentBuilderToolBarStyle.SetIconPadding( FMargin(8, 8, 8, 8 ) );
+		CategoryDrivenContentBuilderToolBarStyle.SetShowLabels( false );
+		CategoryDrivenContentBuilderToolBarStyle.SetButtonContentMaxWidth( 36.0f );
+		Set("CategoryDrivenContentBuilderToolbarWithoutLabels", CategoryDrivenContentBuilderToolBarStyle);
 	}
 	
 	// Vertical ToolPalette 
@@ -4526,9 +4541,7 @@ void FStarshipEditorStyle::FStyle::SetupLevelEditorStyle()
 
 		Set( "PlacementBrowser.Asset.Name", FTextBlockStyle( NormalText )
 			.SetFont( DEFAULT_FONT( "Regular", 10 ) )
-			.SetColorAndOpacity( FLinearColor( 1.0f, 1.0f, 1.0f, 0.9f ) )
-			.SetShadowOffset( FVector2D( 1, 1 ) )
-			.SetShadowColorAndOpacity( FLinearColor( 0, 0, 0, 0.9f ) ) );
+			.SetColorAndOpacity( FLinearColor( 1.0f, 1.0f, 1.0f, 0.9f ) ) );
 
 		Set( "PlacementBrowser.Asset.Type", FTextBlockStyle( NormalText )
 			.SetFont( DEFAULT_FONT( "Regular", 8 ) )
@@ -4595,15 +4608,15 @@ void FStarshipEditorStyle::FStyle::SetupLevelEditorStyle()
 			.SetLastControlStyle(PlacementSegmentedBox)
 		);
 
-		Set("PlacementBrowser.Icons.Recent",        new CORE_IMAGE_BRUSH_SVG("Starship/Common/Recent",    Icon16x16));
-		Set("PlacementBrowser.Icons.Basic",         new IMAGE_BRUSH_SVG("Starship/Common/Basic",          Icon16x16));
-		Set("PlacementBrowser.Icons.Lights",        new IMAGE_BRUSH_SVG("Starship/Common/LightBulb",      Icon16x16));
-		Set("PlacementBrowser.Icons.Cinematics",    new IMAGE_BRUSH_SVG("Starship/Common/Cinematics",     Icon16x16));
-		Set("PlacementBrowser.Icons.VisualEffects", new IMAGE_BRUSH_SVG("Starship/Common/VisualEffects",  Icon16x16));
-		Set("PlacementBrowser.Icons.BSP",           new IMAGE_BRUSH_SVG("Starship/Common/Geometry",       Icon16x16));
-		Set("PlacementBrowser.Icons.Volumes",       new IMAGE_BRUSH_SVG("Starship/Common/Volumes",        Icon16x16));
-		Set("PlacementBrowser.Icons.All",           new IMAGE_BRUSH_SVG("Starship/Common/AllClasses",     Icon16x16));
-		Set("PlacementBrowser.Icons.Testing",       new CORE_IMAGE_BRUSH_SVG("Starship/Common/Test",           Icon16x16));
+		Set("PlacementBrowser.Icons.Recent",        new CORE_IMAGE_BRUSH_SVG("Starship/Common/Recent",    Icon20x20));
+		Set("PlacementBrowser.Icons.Basic",         new IMAGE_BRUSH_SVG("Starship/Common/Basic",          Icon20x20));
+		Set("PlacementBrowser.Icons.Lights",        new IMAGE_BRUSH_SVG("Starship/Common/LightBulb",      Icon20x20));
+		Set("PlacementBrowser.Icons.Cinematics",    new IMAGE_BRUSH_SVG("Starship/Common/Cinematics",     Icon20x20));
+		Set("PlacementBrowser.Icons.VisualEffects", new IMAGE_BRUSH_SVG("Starship/Common/VisualEffects",  Icon20x20));
+		Set("PlacementBrowser.Icons.BSP",           new IMAGE_BRUSH_SVG("Starship/Common/Geometry",       Icon20x20));
+		Set("PlacementBrowser.Icons.Volumes",       new IMAGE_BRUSH_SVG("Starship/Common/Volumes",        Icon20x20));
+		Set("PlacementBrowser.Icons.All",           new IMAGE_BRUSH_SVG("Starship/Common/AllClasses",     Icon20x20));
+		Set("PlacementBrowser.Icons.Testing",       new CORE_IMAGE_BRUSH_SVG("Starship/Common/Test",           Icon20x20));
 
 		Set( "ContentPalette.ShowAllPlaceables", new IMAGE_BRUSH( "Icons/icon_Placement_FilterAll_20px", Icon20x20 ) );
 		Set( "ContentPalette.ShowAllPlaceables.Small", new IMAGE_BRUSH( "Icons/icon_Placement_FilterAll_20px", Icon20x20 ) );
diff --git a/Engine/Source/Runtime/Slate/Private/Framework/MultiBox/SClippingVerticalBox.cpp b/Engine/Source/Runtime/Slate/Private/Framework/MultiBox/SClippingVerticalBox.cpp
--- a/Engine/Source/Runtime/Slate/Private/Framework/MultiBox/SClippingVerticalBox.cpp
+++ b/Engine/Source/Runtime/Slate/Private/Framework/MultiBox/SClippingVerticalBox.cpp
@@ -22,6 +22,7 @@ void SClippingVerticalBox::OnArrangeChildren( const FGeometry& AllottedGeometry,
 	// been called and this method isn't going to behave properly
 	check(WrapButton.IsValid());
 
+	const FToolBarStyle& ToolBarStyle = StyleSet->GetWidgetStyle<FToolBarStyle>(StyleName);
 	LastClippedIdx = ClippedIdx;
 
 	NumClippedChildren = 0;
@@ -31,8 +32,9 @@ void SClippingVerticalBox::OnArrangeChildren( const FGeometry& AllottedGeometry,
 	const int32 NumChildren = ArrangedChildren.Num();
 	const int32 OverflowButtonIndex = NumChildren - 1;
 	LastToolBarButtonIndex = OverflowButtonIndex - 1;
+	const bool bHasLabel = ToolBarStyle.bShowLabels;
 	
-	constexpr int32 OverflowButtonSize = 48;
+	const int32 OverflowButtonSize =  ToolBarStyle.ButtonContentMaxWidth;
 	
 	for (int32 ChildIdx = LastToolBarButtonIndex; ChildIdx >= 0; --ChildIdx)
 	{
@@ -165,7 +167,7 @@ void SClippingVerticalBox::InitializeWrapButton( TSharedPtr<SComboButton>& Butto
 		[
 		SNew(SVerticalBox)
 				+ SVerticalBox::Slot()
-				.Padding(0)
+				.Padding(  ToolBarStyle.bShowLabels ? ToolBarStyle.IconPaddingWithVisibleLabel : ToolBarStyle.IconPadding )
 				.AutoHeight()
 				.HAlign(HAlign_Center)	// Center the icon horizontally, so that large labels don't stretch out the artwork
 				[
@@ -176,6 +178,7 @@ void SClippingVerticalBox::InitializeWrapButton( TSharedPtr<SComboButton>& Butto
 				.HAlign(HAlign_Center)
 				[
 					SNew(STextBlock)
+					.Visibility( ToolBarStyle.bShowLabels ? EVisibility::Visible : EVisibility::Collapsed )
 					.Text( LOCTEXT("ClippingVerticalBox.Icon.More", "More") )
 					.TextStyle(&FCoreStyle::Get().GetWidgetStyle<FTextBlockStyle>("SmallText"))
 				]
diff --git a/Engine/Source/Runtime/Slate/Private/Framework/MultiBox/SClippingVerticalBox.h b/Engine/Source/Runtime/Slate/Private/Framework/MultiBox/SClippingVerticalBox.h
--- a/Engine/Source/Runtime/Slate/Private/Framework/MultiBox/SClippingVerticalBox.h
+++ b/Engine/Source/Runtime/Slate/Private/Framework/MultiBox/SClippingVerticalBox.h
@@ -21,6 +21,7 @@ public:
 	SLATE_BEGIN_ARGS(SClippingVerticalBox) 
 		: _StyleSet(&FCoreStyle::Get())
 		, _StyleName(NAME_None)
+		, _LabelVisibility(EVisibility::Visible)
 		, _IsFocusable(true)
 		, _SelectedIndex(INDEX_NONE )
 		{ }
@@ -28,6 +29,7 @@ public:
 		SLATE_ARGUMENT(FOnGetContent, OnWrapButtonClicked)
 		SLATE_ARGUMENT(const ISlateStyle*, StyleSet)
 		SLATE_ARGUMENT(FName, StyleName)
+		SLATE_ATTRIBUTE( EVisibility, LabelVisibility)
 		SLATE_ARGUMENT(bool, IsFocusable)
 		SLATE_ATTRIBUTE( int32, SelectedIndex )
 	SLATE_END_ARGS()
