PR #11805: Improved default blueprints Grid Line appearance

Note: Added additional change to comment bubble default color based on feedback from design.

#11805

diff --git a/Engine/Source/Editor/EditorStyle/Private/SlateEditorStyle.cpp b/Engine/Source/Editor/EditorStyle/Private/SlateEditorStyle.cpp
--- a/Engine/Source/Editor/EditorStyle/Private/SlateEditorStyle.cpp
+++ b/Engine/Source/Editor/EditorStyle/Private/SlateEditorStyle.cpp
@@ -3912,9 +3912,9 @@ void FSlateEditorStyle::FStyle::SetupGraphEditorStyles()
 		Set( "Graph.ReadOnlyBorder", new BOX_BRUSH( "/Graph/Graph_readonly_border", FMargin(18.0f / 64.0f) ) );
 
 		Set( "Graph.Panel.SolidBackground", new IMAGE_BRUSH( "/Graph/GraphPanel_SolidBackground", FVector2D(16, 16), FLinearColor::White, ESlateBrushTileType::Both) );
-		Set( "Graph.Panel.GridLineColor",   FLinearColor(0.035f, 0.035f, 0.035f) );
-		Set( "Graph.Panel.GridRuleColor",   FLinearColor(0.008f, 0.008f, 0.008f) );
-		Set( "Graph.Panel.GridCenterColor", FLinearColor(0.000f, 0.000f, 0.000f) );
+		Set( "Graph.Panel.GridLineColor",   FLinearColor(0.024f, 0.024, 0.024f) );
+		Set( "Graph.Panel.GridRuleColor",   FLinearColor(0.010f, 0.010f, 0.010f) );
+		Set( "Graph.Panel.GridCenterColor", FLinearColor(0.005f, 0.005f, 0.005f) );
 		
 		Set( "Graph.Panel.GridRulePeriod", 8.0f ); // should be a strictly positive integral value
 
diff --git a/Engine/Source/Editor/EditorStyle/Private/StarshipStyle.cpp b/Engine/Source/Editor/EditorStyle/Private/StarshipStyle.cpp
--- a/Engine/Source/Editor/EditorStyle/Private/StarshipStyle.cpp
+++ b/Engine/Source/Editor/EditorStyle/Private/StarshipStyle.cpp
@@ -3383,9 +3383,9 @@ void FStarshipEditorStyle::FStyle::SetupGraphEditorStyles()
 		Set( "Graph.ReadOnlyBorder", new BOX_BRUSH( "/Graph/Graph_readonly_border", FMargin(18.0f / 64.0f) ) );
 
 		Set( "Graph.Panel.SolidBackground", new IMAGE_BRUSH( "/Graph/GraphPanel_SolidBackground", FVector2D(16, 16), FLinearColor::White, ESlateBrushTileType::Both) );
-		Set( "Graph.Panel.GridLineColor",   FLinearColor(0.035f, 0.035f, 0.035f) );
-		Set( "Graph.Panel.GridRuleColor",   FLinearColor(0.008f, 0.008f, 0.008f) );
-		Set( "Graph.Panel.GridCenterColor", FLinearColor(0.000f, 0.000f, 0.000f) );
+		Set( "Graph.Panel.GridLineColor",   FLinearColor(0.024f, 0.024, 0.024f) );
+		Set( "Graph.Panel.GridRuleColor",   FLinearColor(0.010f, 0.010f, 0.010f) );
+		Set( "Graph.Panel.GridCenterColor", FLinearColor(0.005f, 0.005f, 0.005f) );
 		
 		Set( "Graph.Panel.GridRulePeriod", 8.0f ); // should be a strictly positive integral value
 
diff --git a/Engine/Source/Editor/GraphEditor/Private/GraphEditorSettings.cpp b/Engine/Source/Editor/GraphEditor/Private/GraphEditorSettings.cpp
--- a/Engine/Source/Editor/GraphEditor/Private/GraphEditorSettings.cpp
+++ b/Engine/Source/Editor/GraphEditor/Private/GraphEditorSettings.cpp
@@ -94,7 +94,7 @@ UGraphEditorSettings::UGraphEditorSettings( const FObjectInitializer& ObjectInit
 	ExecBranchNodeTitleColor = FLinearColor(1.0f, 1.0f, 1.0f, 1.0f);
 	ExecSequenceNodeTitleColor = FLinearColor(0.8f, 0.4f, 0.4f, 1.0f);
 	ResultNodeTitleColor = FLinearColor(1.0f, 0.65f, 0.4f, 1.0f);
-	DefaultCommentNodeTitleColor = FLinearColor::White;
+	DefaultCommentNodeTitleColor = FLinearColor(0.15f, 0.15f, 0.15f, 0.5f);
 	PreviewNodeTitleColor = FLinearColor(0.0f, 0.0f, 1.0f, 1.0f);
 }
 
diff --git a/Engine/Source/Editor/UnrealEd/Private/Settings/EditorStyleClasses.cpp b/Engine/Source/Editor/UnrealEd/Private/Settings/EditorStyleClasses.cpp
--- a/Engine/Source/Editor/UnrealEd/Private/Settings/EditorStyleClasses.cpp
+++ b/Engine/Source/Editor/UnrealEd/Private/Settings/EditorStyleClasses.cpp
@@ -45,9 +45,9 @@ UEditorStyleSettings::UEditorStyleSettings( const FObjectInitializer& ObjectInit
 	
 	bUseGrid = true;
 
-	RegularColor = FLinearColor(0.035f, 0.035f, 0.035f);
-	RuleColor = FLinearColor(0.008f, 0.008f, 0.008f);
-	CenterColor = FLinearColor::Black;
+	RegularColor = FLinearColor(0.024f, 0.024, 0.024f);
+	RuleColor = FLinearColor(0.010f, 0.010f, 0.010f);
+	CenterColor = FLinearColor(0.005f, 0.005f, 0.005f) ;
 
 	GridSnapSize = 16;
 
