Move hair show flag among the main items for consistency.

diff --git a/Engine/Source/Editor/EditorStyle/Private/StarshipStyle.cpp b/Engine/Source/Editor/EditorStyle/Private/StarshipStyle.cpp
--- a/Engine/Source/Editor/EditorStyle/Private/StarshipStyle.cpp
+++ b/Engine/Source/Editor/EditorStyle/Private/StarshipStyle.cpp
@@ -4771,6 +4771,7 @@ void FStarshipEditorStyle::FStyle::SetupLevelEditorStyle()
 		Set( "ShowFlagsMenu.Translucency", new IMAGE_BRUSH_SVG( "Starship/Common/Transparency", Icon16x16 ) );
 		Set( "ShowFlagsMenu.WidgetComponents", new IMAGE_BRUSH_SVG( "Starship/Common/WidgetComponents", Icon16x16 ) );
 		Set( "ShowFlagsMenu.Cameras", new IMAGE_BRUSH_SVG( "Starship/AssetIcons/Camera_16", Icon16x16 ) );
+		Set( "ShowFlagsMenu.Hair", new IMAGE_BRUSH( "Common/icon_ShowHair_16x", Icon16x16 ) );
 
 		Set("ShowFlagsMenu.SubMenu.PostProcessing", new IMAGE_BRUSH_SVG("Starship/Common/PostProcessing", Icon16x16));
 		Set("ShowFlagsMenu.SubMenu.LightTypes", new IMAGE_BRUSH_SVG("Starship/Common/LightTypes", Icon16x16));
diff --git a/Engine/Source/Runtime/Engine/Public/ShowFlagsValues.inl b/Engine/Source/Runtime/Engine/Public/ShowFlagsValues.inl
--- a/Engine/Source/Runtime/Engine/Public/ShowFlagsValues.inl
+++ b/Engine/Source/Runtime/Engine/Public/ShowFlagsValues.inl
@@ -458,7 +458,7 @@ SHOWFLAG_ALWAYS_ACCESSIBLE(SceneCaptureCopySceneDepth, SFG_Hidden, NSLOCTEXT("Un
 SHOWFLAG_ALWAYS_ACCESSIBLE(Cameras, SFG_Normal, NSLOCTEXT("UnrealEd", "CamerasSF", "Cameras"))
 
 /** If false, hair geometry won't be rendered. */
-SHOWFLAG_ALWAYS_ACCESSIBLE(Hair, SFG_Advanced, NSLOCTEXT("UnrealEd", "HairSF", "Hair"))
+SHOWFLAG_ALWAYS_ACCESSIBLE(Hair, SFG_Normal, NSLOCTEXT("UnrealEd", "HairSF", "Hair"))
 
 #undef SHOWFLAG_ALWAYS_ACCESSIBLE
 #undef SHOWFLAG_FIXED_IN_SHIPPING
