Create a Slate style for the new viewport toolbar

* Create a style named "ViewportToolbar" for toolbar entries that currently sets a brighter background color to help visually group the toolbar.
* Create a separate style named "ViewportToolbar.Raised" for any entries raised to the top-level toolbar. Such entries do not get the brighter background because they're a part of the submenu they were raised from and that submenus brighter background will provide the visual grouping hint.
* Apply the new "ViewportToolbar" style to the UToolMenu representing the new level editor viewport toolbar.
* Add the ".Raised" suffix to any set style when an entry is raised to the top-level of a toolbar. Any entry raised within the viewport toolbar (using style "ViewportToolbar") will therefore hit the "ViewportToolbar.Raised" style.

diff --git a/Engine/Source/Developer/ToolMenus/Private/ToolMenus.cpp b/Engine/Source/Developer/ToolMenus/Private/ToolMenus.cpp
--- a/Engine/Source/Developer/ToolMenus/Private/ToolMenus.cpp
+++ b/Engine/Source/Developer/ToolMenus/Private/ToolMenus.cpp
@@ -1492,7 +1492,38 @@ void UToolMenus::PopulateToolBarBuilderWithEntry(
 		return;
 	}
 
-	ToolBarBuilder.BeginStyleOverride(Block.StyleNameOverride);
+	// Override the style name.
+	{
+		FName OverrideStyleName = Block.StyleNameOverride;
+
+		// Add the .Raised suffix for menu entries raised to the top level.
+		if (bIsRaisingToTopLevel)
+		{
+			if (Block.StyleNameOverride != NAME_None)
+			{
+				OverrideStyleName = ISlateStyle::Join(Block.StyleNameOverride, ".Raised");
+			}
+			else
+			{
+				// We have to search up the submenu parent chain here because the immediate menu we're a part of might
+				// not have a style set while a parent could.
+				const UToolMenu* CurrentMenu = MenuData;
+				FName MenuStyleName = NAME_None;
+				while (MenuStyleName == NAME_None && CurrentMenu->SubMenuParent)
+				{
+					CurrentMenu = CurrentMenu->SubMenuParent.Get();
+					MenuStyleName = CurrentMenu->StyleName;
+				}
+
+				if (MenuStyleName != NAME_None)
+				{
+					OverrideStyleName = ISlateStyle::Join(MenuStyleName, ".Raised");
+				}
+			}
+		}
+
+		ToolBarBuilder.BeginStyleOverride(OverrideStyleName);
+	}
 
 	FUIAction UIAction = UToolMenus::ConvertUIAction(Block, MenuData->Context);
 	if (bIsRaisingToTopLevel && Block.ShowInToolbarTopLevel.IsBound())
diff --git a/Engine/Source/Developer/ToolMenus/Public/ToolMenuEntry.h b/Engine/Source/Developer/ToolMenus/Public/ToolMenuEntry.h
--- a/Engine/Source/Developer/ToolMenus/Public/ToolMenuEntry.h
+++ b/Engine/Source/Developer/ToolMenus/Public/ToolMenuEntry.h
@@ -164,6 +164,9 @@ struct TOOLMENUS_API FToolMenuEntry
 	 *
 	 * This flag only effects entries within submenus of toolbar-type ToolMenus.
 	 *
+	 * THIS AFFECTS STYLING. When an entry is raised to the top level of a toolbar, the ".Raised" suffix is added to
+	 * the style name that would otherwise have been applied.
+	 *
 	 * @param InTopLevel True shows the entry in the top-level next to its submenu, false (default) only displays it in
 	 * the submenu itself. Pass a delegate to drive the top-level state from code.
 	 */
diff --git a/Engine/Source/Editor/EditorStyle/Private/StarshipStyle.cpp b/Engine/Source/Editor/EditorStyle/Private/StarshipStyle.cpp
--- a/Engine/Source/Editor/EditorStyle/Private/StarshipStyle.cpp
+++ b/Engine/Source/Editor/EditorStyle/Private/StarshipStyle.cpp
@@ -2408,6 +2408,64 @@ void FStarshipEditorStyle::FStyle::SetupSequencerStyles()
 
 void FStarshipEditorStyle::FStyle::SetupViewportStyles()
 {
+	// New viewport toolbar.
+	{
+		FButtonStyle ButtonStyle;
+		{
+			FLinearColor ToolbarBackgroundColor = FStyleColors::Dropdown.GetSpecifiedColor();
+			ToolbarBackgroundColor.A = .90f;
+			FSlateRoundedBoxBrush* BackgroundBrush =
+				new FSlateRoundedBoxBrush(ToolbarBackgroundColor, 5.f, ToolbarBackgroundColor, 1.0);
+
+			FLinearColor ToolbarHoveredColor = FStyleColors::Dropdown.GetSpecifiedColor();
+			ToolbarHoveredColor.A = 1;
+			FSlateRoundedBoxBrush* HoveredBrush =
+				new FSlateRoundedBoxBrush(ToolbarHoveredColor, 5.f, ToolbarHoveredColor, 1.0);
+
+			FLinearColor ToolbarPressedColor = FStyleColors::Recessed.GetSpecifiedColor();
+			ToolbarPressedColor.A = .50f;
+			FSlateRoundedBoxBrush* PressedBrush =
+				new FSlateRoundedBoxBrush(ToolbarPressedColor, 5.f, ToolbarPressedColor, 1.0);
+
+			ButtonStyle.SetNormal(*BackgroundBrush)
+				.SetHovered(*HoveredBrush)
+				.SetPressed(*PressedBrush)
+				.SetNormalForeground(FStyleColors::ForegroundHeader)
+				.SetHoveredForeground(FStyleColors::ForegroundHover)
+				.SetPressedForeground(FStyleColors::ForegroundHover)
+				.SetDisabledForeground(FStyleColors::Foreground)
+				.SetNormalPadding(FMargin(4.0f, 4.0f, 3.0f, 4.0f))
+				.SetPressedPadding(FMargin(4.0f, 4.0f, 3.0f, 4.0f));
+		}
+
+		FToolBarStyle ToolBarStyle = FStarshipCoreStyle::GetCoreStyle().GetWidgetStyle<FToolBarStyle>("SlimToolBar");
+		{
+			ToolBarStyle.SetIconSize(Icon16x16).SetButtonStyle(ButtonStyle).SetComboButtonPadding(FMargin(4.f, 0.0f));
+
+			Set("ViewportToolbar", ToolBarStyle);
+		}
+
+		// Special styling for top-level raised buttons.
+		{
+			FButtonStyle ButtonStyleRaised = ButtonStyle;
+			{
+				FLinearColor BackgroundColorRaised = FStyleColors::Dropdown.GetSpecifiedColor();
+				BackgroundColorRaised.A = .0f;
+				FSlateRoundedBoxBrush* BackgroundBrushRaised =
+					new FSlateRoundedBoxBrush(BackgroundColorRaised, 5.f, BackgroundColorRaised, 1.0);
+
+				ButtonStyleRaised.SetNormal(*BackgroundBrushRaised).SetNormalForeground(FStyleColors::Foreground);
+			}
+
+			FToolBarStyle NewViewportToolbarTopLevelRaisedStyle = ToolBarStyle;
+			{
+				NewViewportToolbarTopLevelRaisedStyle.SetButtonStyle(ButtonStyleRaised);
+
+				Set("ViewportToolbar.Raised", NewViewportToolbarTopLevelRaisedStyle);
+			}
+		}
+	}
+
 	{
 		FToolBarStyle ViewportToolbarStyle = FStarshipCoreStyle::GetCoreStyle().GetWidgetStyle<FToolBarStyle>("SlimToolBar");
 
diff --git a/Engine/Source/Editor/LevelEditor/Private/SLevelViewport.cpp b/Engine/Source/Editor/LevelEditor/Private/SLevelViewport.cpp
--- a/Engine/Source/Editor/LevelEditor/Private/SLevelViewport.cpp
+++ b/Engine/Source/Editor/LevelEditor/Private/SLevelViewport.cpp
@@ -1953,6 +1953,8 @@ TSharedPtr<SWidget> SLevelViewport::BuildViewportToolbar()
 				LevelEditorViewportToolbarName, NAME_None /* parent */, EMultiBoxType::SlimHorizontalToolBar
 			);
 
+			ViewportToolbarMenu->StyleName = "ViewportToolbar";
+
 			// Add the left-aligned part of the viewport toolbar.
 			{
 				FToolMenuSection& LeftSection = ViewportToolbarMenu->FindOrAddSection("Left");
