Fix static analysis warning

diff --git a/Engine/Source/Editor/EditorStyle/Private/StarshipStyle.cpp b/Engine/Source/Editor/EditorStyle/Private/StarshipStyle.cpp
--- a/Engine/Source/Editor/EditorStyle/Private/StarshipStyle.cpp
+++ b/Engine/Source/Editor/EditorStyle/Private/StarshipStyle.cpp
@@ -2414,22 +2414,19 @@ void FStarshipEditorStyle::FStyle::SetupViewportStyles()
 		{
 			FLinearColor ToolbarBackgroundColor = FStyleColors::Dropdown.GetSpecifiedColor();
 			ToolbarBackgroundColor.A = .90f;
-			FSlateRoundedBoxBrush* BackgroundBrush =
-				new FSlateRoundedBoxBrush(ToolbarBackgroundColor, 5.f, ToolbarBackgroundColor, 1.0);
+			const FSlateRoundedBoxBrush BackgroundBrush(ToolbarBackgroundColor, 5.f, ToolbarBackgroundColor, 1.0);
 
 			FLinearColor ToolbarHoveredColor = FStyleColors::Dropdown.GetSpecifiedColor();
 			ToolbarHoveredColor.A = 1;
-			FSlateRoundedBoxBrush* HoveredBrush =
-				new FSlateRoundedBoxBrush(ToolbarHoveredColor, 5.f, ToolbarHoveredColor, 1.0);
+			const FSlateRoundedBoxBrush HoveredBrush(ToolbarHoveredColor, 5.f, ToolbarHoveredColor, 1.0);
 
 			FLinearColor ToolbarPressedColor = FStyleColors::Recessed.GetSpecifiedColor();
 			ToolbarPressedColor.A = .50f;
-			FSlateRoundedBoxBrush* PressedBrush =
-				new FSlateRoundedBoxBrush(ToolbarPressedColor, 5.f, ToolbarPressedColor, 1.0);
+			const FSlateRoundedBoxBrush PressedBrush(ToolbarPressedColor, 5.f, ToolbarPressedColor, 1.0);
 
-			ButtonStyle.SetNormal(*BackgroundBrush)
-				.SetHovered(*HoveredBrush)
-				.SetPressed(*PressedBrush)
+			ButtonStyle.SetNormal(BackgroundBrush)
+				.SetHovered(HoveredBrush)
+				.SetPressed(PressedBrush)
 				.SetNormalForeground(FStyleColors::ForegroundHeader)
 				.SetHoveredForeground(FStyleColors::ForegroundHover)
 				.SetPressedForeground(FStyleColors::ForegroundHover)
@@ -2451,10 +2448,9 @@ void FStarshipEditorStyle::FStyle::SetupViewportStyles()
 			{
 				FLinearColor BackgroundColorRaised = FStyleColors::Dropdown.GetSpecifiedColor();
 				BackgroundColorRaised.A = .0f;
-				FSlateRoundedBoxBrush* BackgroundBrushRaised =
-					new FSlateRoundedBoxBrush(BackgroundColorRaised, 5.f, BackgroundColorRaised, 1.0);
+				const FSlateRoundedBoxBrush BackgroundBrushRaised(BackgroundColorRaised, 5.f, BackgroundColorRaised, 1.0);
 
-				ButtonStyleRaised.SetNormal(*BackgroundBrushRaised).SetNormalForeground(FStyleColors::Foreground);
+				ButtonStyleRaised.SetNormal(BackgroundBrushRaised).SetNormalForeground(FStyleColors::Foreground);
 			}
 
 			FToolBarStyle NewViewportToolbarTopLevelRaisedStyle = ToolBarStyle;
