Update the "ViewportToolbar" style to highlight the icons of checked toggle buttons in AccentBlue instead of altering their background color

diff --git a/Engine/Source/Editor/EditorStyle/Private/StarshipStyle.cpp b/Engine/Source/Editor/EditorStyle/Private/StarshipStyle.cpp
--- a/Engine/Source/Editor/EditorStyle/Private/StarshipStyle.cpp
+++ b/Engine/Source/Editor/EditorStyle/Private/StarshipStyle.cpp
@@ -2410,23 +2410,24 @@ void FStarshipEditorStyle::FStyle::SetupViewportStyles()
 {
 	// New viewport toolbar.
 	{
+
+		FLinearColor ToolbarHoveredColor = FStyleColors::Dropdown.GetSpecifiedColor();
+		ToolbarHoveredColor.A = 1;
+		const FSlateRoundedBoxBrush HoveredButtonBrush(ToolbarHoveredColor, 5.f, ToolbarHoveredColor, 1.0);
+
+		FLinearColor ToolbarPressedColor = FStyleColors::Recessed.GetSpecifiedColor();
+		ToolbarPressedColor.A = .50f;
+		const FSlateRoundedBoxBrush PressedButtonBrush(ToolbarPressedColor, 5.f, ToolbarPressedColor, 1.0);
+
 		FButtonStyle ButtonStyle;
 		{
 			FLinearColor ToolbarBackgroundColor = FStyleColors::Dropdown.GetSpecifiedColor();
 			ToolbarBackgroundColor.A = .90f;
 			const FSlateRoundedBoxBrush BackgroundBrush(ToolbarBackgroundColor, 5.f, ToolbarBackgroundColor, 1.0);
 
-			FLinearColor ToolbarHoveredColor = FStyleColors::Dropdown.GetSpecifiedColor();
-			ToolbarHoveredColor.A = 1;
-			const FSlateRoundedBoxBrush HoveredBrush(ToolbarHoveredColor, 5.f, ToolbarHoveredColor, 1.0);
-
-			FLinearColor ToolbarPressedColor = FStyleColors::Recessed.GetSpecifiedColor();
-			ToolbarPressedColor.A = .50f;
-			const FSlateRoundedBoxBrush PressedBrush(ToolbarPressedColor, 5.f, ToolbarPressedColor, 1.0);
-
 			ButtonStyle.SetNormal(BackgroundBrush)
-				.SetHovered(HoveredBrush)
-				.SetPressed(PressedBrush)
+				.SetHovered(HoveredButtonBrush)
+				.SetPressed(PressedButtonBrush)
 				.SetNormalForeground(FStyleColors::ForegroundHeader)
 				.SetHoveredForeground(FStyleColors::ForegroundHover)
 				.SetPressedForeground(FStyleColors::ForegroundHover)
@@ -2437,7 +2438,25 @@ void FStarshipEditorStyle::FStyle::SetupViewportStyles()
 
 		FToolBarStyle ToolBarStyle = FStarshipCoreStyle::GetCoreStyle().GetWidgetStyle<FToolBarStyle>("SlimToolBar");
 		{
-			ToolBarStyle.SetIconSize(Icon16x16).SetButtonStyle(ButtonStyle).SetComboButtonPadding(FMargin(4.f, 0.0f));
+			FCheckBoxStyle ToggleButtonStyle = ToolBarStyle.ToggleButton;
+			{
+				const FSlateRoundedBoxBrush TransparentBrush(FLinearColor(0, 0, 0, 0), 0, FLinearColor(0, 0, 0, 0), 0);
+
+				ToggleButtonStyle.SetCheckedImage(TransparentBrush)
+					.SetCheckedHoveredImage(HoveredButtonBrush)
+					.SetCheckedPressedImage(PressedButtonBrush)
+					.SetUncheckedPressedImage(PressedButtonBrush)
+					.SetCheckedForegroundColor(FStyleColors::AccentBlue)
+					.SetCheckedHoveredForegroundColor(FStyleColors::AccentBlue)
+					.SetCheckedPressedForegroundColor(FStyleColors::AccentBlue)
+					.SetPressedForegroundColor(FStyleColors::AccentBlue)
+					.SetPadding(FMargin(4));
+			}
+
+			ToolBarStyle.SetIconSize(Icon16x16)
+				.SetButtonStyle(ButtonStyle)
+				.SetToggleButtonStyle(ToggleButtonStyle)
+				.SetComboButtonPadding(FMargin(4.f, 0.0f));
 
 			Set("ViewportToolbar", ToolBarStyle);
 		}
@@ -2462,6 +2481,7 @@ void FStarshipEditorStyle::FStyle::SetupViewportStyles()
 		}
 	}
 
+	// Old viewport toolbar.
 	{
 		FToolBarStyle ViewportToolbarStyle = FStarshipCoreStyle::GetCoreStyle().GetWidgetStyle<FToolBarStyle>("SlimToolBar");
 
