Tighten layout of the new viewport toolbar

* Add FToolBarStyle::SeparatorThickness to specify the thickness of separators in toolbars.
* Reduce button paddings in the new viewport toolbar.
* Adjust combo button paddings to add extra space on the left side of the main sub menu combo buttons in the new viewport toolbar.
* Set up separator styling to be thin using FToolBarStyle::SeparatorThickness and using the FStyleColors::Dropdown color.
* Update the Concert and Avalanche modules to use the new FToolBarStyle::SeparatorThickness default value of 2.0f instead of their hardcoded 2.0f values.

diff --git a/Engine/Plugins/Developer/Concert/ConcertUI/ConcertSharedSlate/Source/ConcertSharedSlate/Private/Session/Browser/SConcertSessionBrowser.cpp b/Engine/Plugins/Developer/Concert/ConcertUI/ConcertSharedSlate/Source/ConcertSharedSlate/Private/Session/Browser/SConcertSessionBrowser.cpp
--- a/Engine/Plugins/Developer/Concert/ConcertUI/ConcertSharedSlate/Source/ConcertSharedSlate/Private/Session/Browser/SConcertSessionBrowser.cpp
+++ b/Engine/Plugins/Developer/Concert/ConcertUI/ConcertSharedSlate/Source/ConcertSharedSlate/Private/Session/Browser/SConcertSessionBrowser.cpp
@@ -96,7 +96,7 @@ TSharedRef<SWidget> SConcertSessionBrowser::MakeBrowserContent(const FArguments&
 			.AutoHeight()
 			[
 				SNew(SSeparator)
-				.Thickness(2.f)
+				.Thickness(FCoreStyle::Get().GetWidgetStyle<FToolBarStyle>("ToolBar").SeparatorThickness)
 				.SeparatorImage(&FCoreStyle::Get().GetWidgetStyle<FToolBarStyle>("ToolBar").SeparatorBrush)
 			]
 
diff --git a/Engine/Plugins/Developer/Concert/ConcertUI/ConcertSharedSlate/Source/ConcertSharedSlate/Private/Session/History/SSessionHistory.cpp b/Engine/Plugins/Developer/Concert/ConcertUI/ConcertSharedSlate/Source/ConcertSharedSlate/Private/Session/History/SSessionHistory.cpp
--- a/Engine/Plugins/Developer/Concert/ConcertUI/ConcertSharedSlate/Source/ConcertSharedSlate/Private/Session/History/SSessionHistory.cpp
+++ b/Engine/Plugins/Developer/Concert/ConcertUI/ConcertSharedSlate/Source/ConcertSharedSlate/Private/Session/History/SSessionHistory.cpp
@@ -147,7 +147,7 @@ void SSessionHistory::Construct(const FArguments& InArgs)
 			.Padding(0, 3)
 			[
 				SNew(SSeparator)
-				.Thickness(2.f)
+				.Thickness(FCoreStyle::Get().GetWidgetStyle<FToolBarStyle>("ToolBar").SeparatorThickness)
 				.SeparatorImage(&FCoreStyle::Get().GetWidgetStyle<FToolBarStyle>("ToolBar").SeparatorBrush)
 			];
 	}
diff --git a/Engine/Plugins/Experimental/Avalanche/Source/AvalancheOutliner/Private/Slate/SAvaOutlinerItemFilters.cpp b/Engine/Plugins/Experimental/Avalanche/Source/AvalancheOutliner/Private/Slate/SAvaOutlinerItemFilters.cpp
--- a/Engine/Plugins/Experimental/Avalanche/Source/AvalancheOutliner/Private/Slate/SAvaOutlinerItemFilters.cpp
+++ b/Engine/Plugins/Experimental/Avalanche/Source/AvalancheOutliner/Private/Slate/SAvaOutlinerItemFilters.cpp
@@ -139,7 +139,7 @@ void SAvaOutlinerItemFilters::Construct(const FArguments& InArgs, const TSharedR
 		[
 			SNew(SSeparator)
 			.Orientation(Orient_Horizontal)
-			.Thickness(2.0f)
+			.Thickness(FStyleUtils::GetSlimToolBarStyle().SeparatorThickness)
 			.SeparatorImage(&FStyleUtils::GetSlimToolBarStyle().SeparatorBrush)
 		];
 
diff --git a/Engine/Source/Editor/EditorStyle/Private/StarshipStyle.cpp b/Engine/Source/Editor/EditorStyle/Private/StarshipStyle.cpp
--- a/Engine/Source/Editor/EditorStyle/Private/StarshipStyle.cpp
+++ b/Engine/Source/Editor/EditorStyle/Private/StarshipStyle.cpp
@@ -2432,8 +2432,8 @@ void FStarshipEditorStyle::FStyle::SetupViewportStyles()
 				.SetHoveredForeground(FStyleColors::ForegroundHover)
 				.SetPressedForeground(FStyleColors::ForegroundHover)
 				.SetDisabledForeground(FStyleColors::Foreground)
-				.SetNormalPadding(FMargin(4.0f, 4.0f, 3.0f, 4.0f))
-				.SetPressedPadding(FMargin(4.0f, 4.0f, 3.0f, 4.0f));
+				.SetNormalPadding(FMargin(4.0f))
+				.SetPressedPadding(FMargin(4.0f));
 		}
 
 		FToolBarStyle ToolBarStyle = FStarshipCoreStyle::GetCoreStyle().GetWidgetStyle<FToolBarStyle>("SlimToolBar");
@@ -2450,7 +2450,7 @@ void FStarshipEditorStyle::FStyle::SetupViewportStyles()
 					.SetCheckedHoveredForegroundColor(FStyleColors::AccentBlue)
 					.SetCheckedPressedForegroundColor(FStyleColors::AccentBlue)
 					.SetPressedForegroundColor(FStyleColors::AccentBlue)
-					.SetPadding(FMargin(4));
+					.SetPadding(FMargin(4.0f));
 			}
 
 			FSlateBrush SeparatorBrush;
@@ -2460,9 +2460,11 @@ void FStarshipEditorStyle::FStyle::SetupViewportStyles()
 			ToolBarStyle.SetIconSize(Icon16x16)
 				.SetButtonStyle(ButtonStyle)
 				.SetToggleButtonStyle(ToggleButtonStyle)
-				.SetComboButtonPadding(FMargin(4.f, 0.0f))
+				.SetButtonPadding(FMargin(0))
+				.SetComboButtonPadding(FMargin(6.0f, 0, 2.0f, 0))
 				.SetSeparatorBrush(SeparatorBrush)
-				.SetSeparatorPadding(FMargin(0, 0.0f));
+				.SetSeparatorThickness(1.0f)
+				.SetSeparatorPadding(FMargin(4.0f, 0));
 
 			Set("ViewportToolbar", ToolBarStyle);
 		}
@@ -2480,7 +2482,7 @@ void FStarshipEditorStyle::FStyle::SetupViewportStyles()
 
 			FToolBarStyle NewViewportToolbarTopLevelRaisedStyle = ToolBarStyle;
 			{
-				NewViewportToolbarTopLevelRaisedStyle.SetButtonStyle(ButtonStyleRaised);
+				NewViewportToolbarTopLevelRaisedStyle.SetButtonStyle(ButtonStyleRaised).SetComboButtonPadding(FMargin(2.0f, 0));
 
 				Set("ViewportToolbar.Raised", NewViewportToolbarTopLevelRaisedStyle);
 			}
diff --git a/Engine/Source/Runtime/Slate/Private/Framework/MultiBox/SToolBarSeparatorBlock.cpp b/Engine/Source/Runtime/Slate/Private/Framework/MultiBox/SToolBarSeparatorBlock.cpp
--- a/Engine/Source/Runtime/Slate/Private/Framework/MultiBox/SToolBarSeparatorBlock.cpp
+++ b/Engine/Source/Runtime/Slate/Private/Framework/MultiBox/SToolBarSeparatorBlock.cpp
@@ -62,7 +62,7 @@ void SToolBarSeparatorBlock::BuildMultiBlockWidget(const ISlateStyle* StyleSet,
 		[
 			SNew(SSeparator)
 				.Orientation(Orient_Vertical)
-				.Thickness(2.0f)
+				.Thickness(ToolBarStyle.SeparatorThickness)
 				.SeparatorImage(&ToolBarStyle.SeparatorBrush)
 		]
 	];
diff --git a/Engine/Source/Runtime/SlateCore/Private/Styling/ToolBarStyle.cpp b/Engine/Source/Runtime/SlateCore/Private/Styling/ToolBarStyle.cpp
--- a/Engine/Source/Runtime/SlateCore/Private/Styling/ToolBarStyle.cpp
+++ b/Engine/Source/Runtime/SlateCore/Private/Styling/ToolBarStyle.cpp
@@ -22,13 +22,14 @@ FToolBarStyle::FToolBarStyle()
 	, NumColumns(0)
 	, IconPadding(0)
 	, SeparatorPadding(0)
+	, SeparatorThickness(2.0f)
 	, ComboButtonPadding(0)
 	, ButtonPadding(0)
 	, CheckBoxPadding(0)
 	, BlockPadding(0)
 	, IndentedBlockPadding(0)
 	, BackgroundPadding(0)
-	, IconSize(16,16)
+	, IconSize(16, 16)
 	, bShowLabels(true)
 {}
 
diff --git a/Engine/Source/Runtime/SlateCore/Public/Styling/ToolBarStyle.h b/Engine/Source/Runtime/SlateCore/Public/Styling/ToolBarStyle.h
--- a/Engine/Source/Runtime/SlateCore/Public/Styling/ToolBarStyle.h
+++ b/Engine/Source/Runtime/SlateCore/Public/Styling/ToolBarStyle.h
@@ -97,6 +97,14 @@ struct FToolBarStyle : public FSlateWidgetStyle
 	FMargin SeparatorPadding;
 	FToolBarStyle& SetSeparatorPadding(const FMargin& InSeparatorPadding) { SeparatorPadding = InSeparatorPadding; return *this; }
 
+	UPROPERTY(EditAnywhere, BlueprintReadWrite, Category = Appearance)
+	float SeparatorThickness;
+	FToolBarStyle& SetSeparatorThickness(float InSeparatorThickness)
+	{
+		SeparatorThickness = InSeparatorThickness;
+		return *this;
+	}
+
 	UPROPERTY(EditAnywhere, BlueprintReadWrite, Category = Appearance)
 	FMargin ComboButtonPadding;
 	FToolBarStyle& SetComboButtonPadding(const FMargin& InComboButtonPadding) { ComboButtonPadding = InComboButtonPadding; return *this; }
