Updated icons for animation asset find/replace tab

[RN] Updated icons for animation asset find/replace tab

diff --git a/Engine/Source/Editor/EditorStyle/Private/StarshipStyle.cpp b/Engine/Source/Editor/EditorStyle/Private/StarshipStyle.cpp
--- a/Engine/Source/Editor/EditorStyle/Private/StarshipStyle.cpp
+++ b/Engine/Source/Editor/EditorStyle/Private/StarshipStyle.cpp
@@ -5112,6 +5112,7 @@ void FStarshipEditorStyle::FStyle::SetupPersonaStyle()
 		Set("Persona.Tabs.SkeletonCurves", new IMAGE_BRUSH(TEXT("Icons/icon_Persona_Skeleton_Curves_16x"), Icon16x16));
 		Set("Persona.Tabs.AnimAssetDetails", new IMAGE_BRUSH(TEXT("Icons/icon_Persona_Anim_Asset_Details_16x"), Icon16x16));
 		Set("Persona.Tabs.ControlRigMappingWindow", new IMAGE_BRUSH(TEXT("Icons/icon_Persona_Skeleton_Tree_16x"), Icon16x16));
+		Set("Persona.Tabs.FindReplace", new IMAGE_BRUSH_SVG(TEXT("Starship/Persona/FindReplace_16"), Icon16x16));
 
 		// AssetFamilyIcons
 		Set("Persona.AssetClass.Skeleton", new IMAGE_BRUSH_SVG("Starship/Persona/Skeleton", Icon20x20));
@@ -5119,6 +5120,10 @@ void FStarshipEditorStyle::FStyle::SetupPersonaStyle()
 		Set("Persona.AssetClass.SkeletalMesh", new IMAGE_BRUSH_SVG("Starship/Persona/SkeletalMesh", Icon20x20));
 		Set("Persona.AssetClass.Blueprint", new IMAGE_BRUSH_SVG( "Starship/MainToolbar/blueprints", Icon20x20) );
 		Set("Persona.AssetClass.Physics", new IMAGE_BRUSH_SVG("Starship/Persona/Physics", Icon20x20));
+
+		// Find/replace tab
+		Set("Persona.FindReplace.MatchCase", new IMAGE_BRUSH_SVG("Starship/Persona/MatchCase_20", Icon20x20));
+		Set("Persona.FindReplace.MatchWholeWord", new IMAGE_BRUSH_SVG("Starship/Persona/MatchWord_20", Icon20x20));
 	}
 
 	// Skeleton editor
diff --git a/Engine/Source/Editor/Persona/Private/AnimAssetFindReplace.cpp b/Engine/Source/Editor/Persona/Private/AnimAssetFindReplace.cpp
--- a/Engine/Source/Editor/Persona/Private/AnimAssetFindReplace.cpp
+++ b/Engine/Source/Editor/Persona/Private/AnimAssetFindReplace.cpp
@@ -254,7 +254,7 @@ void UAnimAssetFindReplaceProcessor_StringBase::ExtendToolbar(FToolMenuSection&
 			MatchCaseCheckbox,
 			LOCTEXT("MatchCaseCheckboxLabel", "Match Case"),
 			LOCTEXT("MatchCaseCheckboxTooltip", "Whether to match case when searching."),
-			FSlateIcon(FAppStyle::GetAppStyleSetName(), "Icons.Clipboard"),
+			FSlateIcon(FAppStyle::GetAppStyleSetName(), "Persona.FindReplace.MatchCase"),
 			EUserInterfaceActionType::ToggleButton));
 
 	FToolUIAction MatchWholeWordCheckbox;
@@ -274,7 +274,7 @@ void UAnimAssetFindReplaceProcessor_StringBase::ExtendToolbar(FToolMenuSection&
 			MatchWholeWordCheckbox,
 			LOCTEXT("MatchWholeWordCheckboxLabel", "Match Whole Word"),
 			LOCTEXT("MatchWholeWordCheckboxTooltip", "Whether to match the whole word or just part of the word when searching."),
-			FSlateIcon(FAppStyle::GetAppStyleSetName(), "Icons.Use"),
+			FSlateIcon(FAppStyle::GetAppStyleSetName(), "Persona.FindReplace.MatchWholeWord"),
 			EUserInterfaceActionType::ToggleButton));
 
 	InSection.AddSeparator(NAME_None);
diff --git a/Engine/Source/Editor/Persona/Private/SAnimAssetFindReplace.cpp b/Engine/Source/Editor/Persona/Private/SAnimAssetFindReplace.cpp
--- a/Engine/Source/Editor/Persona/Private/SAnimAssetFindReplace.cpp
+++ b/Engine/Source/Editor/Persona/Private/SAnimAssetFindReplace.cpp
@@ -36,7 +36,7 @@ FAnimAssetFindReplaceSummoner::FAnimAssetFindReplaceSummoner(TSharedPtr<FAssetEd
 	, Config(InConfig)
 {
 	TabLabel = LOCTEXT("AnimAssetFindReplaceTabLabel", "Find/Replace");
-	TabIcon = FSlateIcon(FAppStyle::GetAppStyleSetName(), "Kismet.Tabs.FindResults");
+	TabIcon = FSlateIcon(FAppStyle::GetAppStyleSetName(), "Persona.Tabs.FindReplace");
 }
 
 TSharedRef<SWidget> FAnimAssetFindReplaceSummoner::CreateTabBody(const FWorkflowTabSpawnInfo& Info) const
