New icons for landscape

#rnx

diff --git a/Engine/Source/Editor/EditorStyle/Private/SlateEditorStyle.cpp b/Engine/Source/Editor/EditorStyle/Private/SlateEditorStyle.cpp
--- a/Engine/Source/Editor/EditorStyle/Private/SlateEditorStyle.cpp
+++ b/Engine/Source/Editor/EditorStyle/Private/SlateEditorStyle.cpp
@@ -6969,8 +6969,7 @@ void FSlateEditorStyle::FStyle::SetupLandscapeEditorStyle()
 		Set("LandscapeEditor.TargetList.DropZone.Above", new BOX_BRUSH("Common/VerticalBoxDropZoneIndicator_Above", FMargin(10.0f / 16.0f, 10.0f / 16.0f, 0, 0), SelectionColor_Subdued));
 		Set("LandscapeEditor.TargetList.DropZone.Below", new BOX_BRUSH("Common/VerticalBoxDropZoneIndicator_Below", FMargin(10.0f / 16.0f, 0, 0, 10.0f / 16.0f), SelectionColor_Subdued));
 
-		//todo don.boogert - replace with final icon
-		Set("LandscapeEditor.Layer.Sync", new IMAGE_BRUSH("Icons/icon_MatEd_HLSL_Code_16x", Icon16x16));
+		Set("LandscapeEditor.Layer.Sync", new IMAGE_BRUSH("Icons/icon_Landscape_Layers_From_Material_16x", Icon16x16));
 		
 		Set("LandscapeEditor.SpinBox", FSpinBoxStyle(GetWidgetStyle<FSpinBoxStyle>("SpinBox"))
 			.SetTextPadding(FMargin(0))
diff --git a/Engine/Source/Editor/EditorStyle/Private/StarshipStyle.cpp b/Engine/Source/Editor/EditorStyle/Private/StarshipStyle.cpp
--- a/Engine/Source/Editor/EditorStyle/Private/StarshipStyle.cpp
+++ b/Engine/Source/Editor/EditorStyle/Private/StarshipStyle.cpp
@@ -6749,8 +6749,7 @@ void FStarshipEditorStyle::FStyle::SetupLandscapeEditorStyle()
 		Set("LandscapeEditor.TargetList.DropZone.Above", new BOX_BRUSH("Common/VerticalBoxDropZoneIndicator_Above", FMargin(10.0f / 16.0f, 10.0f / 16.0f, 0, 0), SelectionColor_Subdued));
 		Set("LandscapeEditor.TargetList.DropZone.Below", new BOX_BRUSH("Common/VerticalBoxDropZoneIndicator_Below", FMargin(10.0f / 16.0f, 0, 0, 10.0f / 16.0f), SelectionColor_Subdued));
 
-		//todo don.boogert - replace with final icon
-		Set("LandscapeEditor.Layer.Sync", new IMAGE_BRUSH("Icons/icon_MatEd_HLSL_Code_16x", Icon16x16));
+		Set("LandscapeEditor.Layer.Sync", new IMAGE_BRUSH_SVG("Icons/icon_Landscape_Layers_From_Material_16x", Icon16x16));
 		
 		Set("LandscapeEditor.SpinBox", FSpinBoxStyle(GetWidgetStyle<FSpinBoxStyle>("SpinBox"))
 			.SetTextPadding(FMargin(0))
