Fix ellipsis button padding on the toolbar combo button that opens play settings

diff --git a/Engine/Source/Editor/EditorStyle/Private/StarshipStyle.cpp b/Engine/Source/Editor/EditorStyle/Private/StarshipStyle.cpp
--- a/Engine/Source/Editor/EditorStyle/Private/StarshipStyle.cpp
+++ b/Engine/Source/Editor/EditorStyle/Private/StarshipStyle.cpp
@@ -5599,9 +5599,10 @@ void FStarshipEditorStyle::FStyle::SetupPersonaStyle()
 			.SetNormalPadding(FMargin(7.f, 2.f, 6.f, 2.f))
 			.SetPressedPadding(FMargin(7.f, 2.f, 6.f, 2.f));
 
-		FComboButtonStyle PlayToolbarComboButton = FComboButtonStyle(FStarshipCoreStyle::GetCoreStyle().GetWidgetStyle<FComboButtonStyle>("ComboButton"))
-			.SetDownArrowPadding(FMargin(-19.f, 0.f, 2.f, 0.f))
-			.SetDownArrowImage(CORE_IMAGE_BRUSH_SVG("Starship/Common/ellipsis-vertical-narrow", FVector2D(6, 24)));
+		FComboButtonStyle PlayToolbarComboButton =
+			FComboButtonStyle(FStarshipCoreStyle::GetCoreStyle().GetWidgetStyle<FComboButtonStyle>("ComboButton"))
+				.SetDownArrowPadding(FMargin(1.f, 0.f, 2.f, 0.f))
+				.SetDownArrowImage(CORE_IMAGE_BRUSH_SVG("Starship/Common/ellipsis-vertical-narrow", FVector2D(6, 24)));
 		PlayToolbarComboButton.ButtonStyle = RightToolbarComboButton;
 
 		MainToolbarRightComboButton.SetButtonStyle(RightToolbarComboButton);
