Make it clearer to users of the new viewport toolbar when the viewport is not updating in realtime

* Use a warning-colored background for the realtime toggle when the viewport is not updating in real time. Do this by using a style instead of a colored icon when raising the button to the top-level toolbar to warn.
* Remove the now unused icon EditorViewport.ToggleRealTimeWarning (the colored icon previously used to warn).
* Unrelated to the above, use normal forgeround color when a raised toggle button is pressed (i.e. not the active AccentBlue color).
* Unrelated to the above, add a hover color for viewport toolbar submenus (combo buttons).

diff --git a/Engine/Source/Editor/EditorStyle/Private/StarshipStyle.cpp b/Engine/Source/Editor/EditorStyle/Private/StarshipStyle.cpp
--- a/Engine/Source/Editor/EditorStyle/Private/StarshipStyle.cpp
+++ b/Engine/Source/Editor/EditorStyle/Private/StarshipStyle.cpp
@@ -2437,7 +2437,9 @@ void FStarshipEditorStyle::FStyle::SetupViewportStyles()
 	{
 
 		FLinearColor ToolbarHoveredColor = FStyleColors::Dropdown.GetSpecifiedColor();
-		ToolbarHoveredColor.A = 1;
+		ToolbarHoveredColor.R = FMath::Min(1.0f, ToolbarHoveredColor.R * 1.5);
+		ToolbarHoveredColor.G = FMath::Min(1.0f, ToolbarHoveredColor.G * 1.5);
+		ToolbarHoveredColor.B = FMath::Min(1.0f, ToolbarHoveredColor.B * 1.5);
 		const FSlateRoundedBoxBrush HoveredButtonBrush(ToolbarHoveredColor, 5.f, ToolbarHoveredColor, 1.0);
 
 		FLinearColor ToolbarPressedColor = FStyleColors::Recessed.GetSpecifiedColor();
@@ -2461,23 +2463,25 @@ void FStarshipEditorStyle::FStyle::SetupViewportStyles()
 				.SetPressedPadding(FMargin(4.0f));
 		}
 
-		FToolBarStyle ToolBarStyle = FStarshipCoreStyle::GetCoreStyle().GetWidgetStyle<FToolBarStyle>("SlimToolBar");
+		const FToolBarStyle SlimToolBarStyle =
+			FStarshipCoreStyle::GetCoreStyle().GetWidgetStyle<FToolBarStyle>("SlimToolBar");
+
+		FCheckBoxStyle ToggleButtonStyle = SlimToolBarStyle.ToggleButton;
 		{
-			FCheckBoxStyle ToggleButtonStyle = ToolBarStyle.ToggleButton;
-			{
-				const FSlateRoundedBoxBrush TransparentBrush(FLinearColor(0, 0, 0, 0), 0, FLinearColor(0, 0, 0, 0), 0);
-
-				ToggleButtonStyle.SetCheckedImage(TransparentBrush)
-					.SetCheckedHoveredImage(HoveredButtonBrush)
-					.SetCheckedPressedImage(PressedButtonBrush)
-					.SetUncheckedPressedImage(PressedButtonBrush)
-					.SetCheckedForegroundColor(FStyleColors::AccentBlue)
-					.SetCheckedHoveredForegroundColor(FStyleColors::AccentBlue)
-					.SetCheckedPressedForegroundColor(FStyleColors::AccentBlue)
-					.SetPressedForegroundColor(FStyleColors::AccentBlue)
-					.SetPadding(FMargin(4.0f));
-			}
+			const FSlateRoundedBoxBrush TransparentBrush(FLinearColor(0, 0, 0, 0), 0, FLinearColor(0, 0, 0, 0), 0);
+
+			ToggleButtonStyle.SetCheckedImage(TransparentBrush)
+				.SetCheckedHoveredImage(HoveredButtonBrush)
+				.SetCheckedPressedImage(PressedButtonBrush)
+				.SetUncheckedPressedImage(PressedButtonBrush)
+				.SetCheckedForegroundColor(FStyleColors::AccentBlue)
+				.SetCheckedHoveredForegroundColor(FStyleColors::AccentBlue)
+				.SetCheckedPressedForegroundColor(FStyleColors::AccentBlue)
+				.SetPadding(FMargin(4.0f));
+		}
 
+		FToolBarStyle ToolBarStyle = SlimToolBarStyle;
+		{
 			FSlateBrush SeparatorBrush;
 			SeparatorBrush.DrawAs = ESlateBrushDrawType::Box;
 			SeparatorBrush.TintColor = FStyleColors::Dropdown;
@@ -2485,7 +2489,7 @@ void FStarshipEditorStyle::FStyle::SetupViewportStyles()
 			ToolBarStyle.SetIconSize(Icon16x16)
 				.SetButtonStyle(ButtonStyle)
 				.SetToggleButtonStyle(ToggleButtonStyle)
-				.SetButtonPadding(FMargin(0))
+				.SetButtonPadding(FMargin(2.0f, 0))
 				.SetComboButtonPadding(FMargin(6.0f, 0, 2.0f, 0))
 				.SetSeparatorBrush(SeparatorBrush)
 				.SetSeparatorThickness(1.0f)
@@ -2494,6 +2498,47 @@ void FStarshipEditorStyle::FStyle::SetupViewportStyles()
 			Set("ViewportToolbar", ToolBarStyle);
 		}
 
+		// Special styling for warnings.
+		{
+			// Use normal styling for non-raised buttons.
+			Set("ViewportToolbarWarning", ToolBarStyle);
+
+			FCheckBoxStyle ToggleButtonStyleWarning = ToolBarStyle.ToggleButton;
+			{
+				FLinearColor WarningColor = FStyleColors::Warning.GetSpecifiedColor();
+				const FSlateRoundedBoxBrush WarningBrush(WarningColor, 5.0f, WarningColor, 1.0f);
+
+				FLinearColor WarningHoveredColor = FStyleColors::Warning.GetSpecifiedColor();
+				WarningHoveredColor.R = FMath::Min(1.0f, WarningHoveredColor.R * 1.5);
+				WarningHoveredColor.G = FMath::Min(1.0f, WarningHoveredColor.G * 1.5);
+				WarningHoveredColor.B = FMath::Min(1.0f, WarningHoveredColor.B * 1.5);
+				const FSlateRoundedBoxBrush WarningHoveredBrush(WarningHoveredColor, 5.0f, WarningHoveredColor, 1.0f);
+
+				FLinearColor WarningPressedColor = FStyleColors::Warning.GetSpecifiedColor();
+				WarningPressedColor.A = .50f;
+				const FSlateRoundedBoxBrush WarningPressedBrush(WarningPressedColor, 5.0f, WarningPressedColor, 1.0f);
+
+				ToggleButtonStyleWarning.SetBackgroundImage(WarningBrush)
+					.SetBackgroundHoveredImage(WarningHoveredBrush)
+					.SetUncheckedImage(WarningBrush)
+					.SetUncheckedHoveredImage(WarningHoveredBrush)
+					.SetUncheckedPressedImage(WarningPressedBrush)
+					.SetCheckedImage(WarningBrush)
+					.SetCheckedHoveredImage(WarningHoveredBrush)
+					.SetCheckedPressedImage(WarningPressedBrush)
+					.SetForegroundColor(FStyleColors::ForegroundInverted)
+					.SetHoveredForegroundColor(FStyleColors::ForegroundInverted)
+					.SetUndeterminedForegroundColor(FStyleColors::ForegroundInverted)
+					.SetPressedForegroundColor(FStyleColors::ForegroundInverted);
+			}
+
+			FToolBarStyle NewViewportToolbarTopLevelWarningStyle = ToolBarStyle;
+			NewViewportToolbarTopLevelWarningStyle.SetToggleButtonStyle(ToggleButtonStyleWarning)
+				.SetComboButtonPadding(FMargin(2.0f, 0));
+
+			Set("ViewportToolbarWarning.Raised", NewViewportToolbarTopLevelWarningStyle);
+		}
+
 		// Special styling for top-level raised buttons.
 		{
 			FButtonStyle ButtonStyleRaised = ButtonStyle;
@@ -2555,7 +2600,9 @@ void FStarshipEditorStyle::FStyle::SetupViewportStyles()
 				}
 				FToolBarStyle ViewportSizingRaisedStyle = ToolBarStyle;
 				{
-					ViewportSizingRaisedStyle.SetButtonStyle(ButtonStyleRaised).SetComboButtonPadding(FMargin(2.0f, 0));
+					ViewportSizingRaisedStyle.SetButtonStyle(ButtonStyleRaised)
+						.SetButtonPadding(0)
+						.SetComboButtonPadding(FMargin(2.0f, 0));
 
 					Set("ViewportToolbarViewportSizingSubmenu.Raised", ViewportSizingRaisedStyle);
 				}
@@ -4407,8 +4454,6 @@ void FStarshipEditorStyle::FStyle::SetupLevelEditorStyle()
 		Set( "EditorViewport.TranslateRotate2DMode", new IMAGE_BRUSH_SVG("Starship/EditorViewport/TranslateRotate2D", Icon16x16 ) );
 
 		Set( "EditorViewport.ToggleRealTime", new IMAGE_BRUSH_SVG("Starship/Common/Realtime", Icon16x16));
-		Set("EditorViewport.ToggleRealTimeWarning",
-			new IMAGE_BRUSH_SVG("Starship/Common/Realtime", Icon16x16, FStyleColors::Warning));
 
 		Set( "EditorViewport.LocationGridSnap", new IMAGE_BRUSH_SVG("Starship/EditorViewport/grid", Icon16x16));
 		Set( "EditorViewport.RotationGridSnap", new IMAGE_BRUSH_SVG("Starship/EditorViewport/angle", Icon16x16));
diff --git a/Engine/Source/Editor/UnrealEd/Private/ViewportToolbar/UnrealEdViewportToolbar.cpp b/Engine/Source/Editor/UnrealEd/Private/ViewportToolbar/UnrealEdViewportToolbar.cpp
--- a/Engine/Source/Editor/UnrealEd/Private/ViewportToolbar/UnrealEdViewportToolbar.cpp
+++ b/Engine/Source/Editor/UnrealEd/Private/ViewportToolbar/UnrealEdViewportToolbar.cpp
@@ -3131,26 +3131,11 @@ FToolMenuEntry CreateToggleRealtimeEntry()
 					}
 				}
 
-				const TAttribute<FSlateIcon> Icon = TAttribute<FSlateIcon>::CreateLambda(
-					[EditorViewportWeak]() -> FSlateIcon
-					{
-						bool bIsViewportRealtime = true;
-						if (const TSharedPtr<SEditorViewport> EditorViewport = EditorViewportWeak.Pin())
-						{
-							bIsViewportRealtime = EditorViewport->IsRealtime();
-						}
-
-						return bIsViewportRealtime
-								 ? FSlateIcon(FAppStyle::GetAppStyleSetName(), "EditorViewport.ToggleRealTime")
-								 : FSlateIcon(FAppStyle::GetAppStyleSetName(), "EditorViewport.ToggleRealTimeWarning");
-					}
-				);
-
 				FToolMenuEntry ToggleRealtime = FToolMenuEntry::InitMenuEntry(
 					"ToggleRealtime",
 					LOCTEXT("ToggleRealtimeLabel", "Realtime Viewport"),
 					Tooltip,
-					Icon,
+					FSlateIcon(FAppStyle::GetAppStyleSetName(), "EditorViewport.ToggleRealTime"),
 					RealtimeToggleAction,
 					EUserInterfaceActionType::ToggleButton
 				);
@@ -3172,6 +3157,14 @@ FToolMenuEntry CreateToggleRealtimeEntry()
 					));
 				}
 
+				if (const TSharedPtr<SEditorViewport> EditorViewport = EditorViewportWeak.Pin())
+				{
+					if (!EditorViewport->IsRealtime())
+					{
+						ToggleRealtime.StyleNameOverride = "ViewportToolbarWarning";
+					}
+				}
+
 				InnerSection.AddEntry(ToggleRealtime);
 			}
 		)
