Use the brighter Foreground color for the viewport toolbar separator to make it stand out better

diff --git a/Engine/Source/Editor/EditorStyle/Private/StarshipStyle.cpp b/Engine/Source/Editor/EditorStyle/Private/StarshipStyle.cpp
--- a/Engine/Source/Editor/EditorStyle/Private/StarshipStyle.cpp
+++ b/Engine/Source/Editor/EditorStyle/Private/StarshipStyle.cpp
@@ -2484,7 +2484,7 @@ void FStarshipEditorStyle::FStyle::SetupViewportStyles()
 		{
 			FSlateBrush SeparatorBrush;
 			SeparatorBrush.DrawAs = ESlateBrushDrawType::Box;
-			SeparatorBrush.TintColor = FStyleColors::Dropdown;
+			SeparatorBrush.TintColor = FStyleColors::Foreground;
 
 			ToolBarStyle.SetIconSize(Icon16x16)
 				.SetButtonStyle(ButtonStyle)
