[Viewport Toolbar] Add "Scalability: High" warning to new viewport toolbar

Adds style elements to create the yellow warning.
Adds new icon.

This change also removes some duplicate code used to retrieve visibility and labels for the scalability warning

diff --git a/Engine/Source/Editor/EditorStyle/Private/StarshipStyle.cpp b/Engine/Source/Editor/EditorStyle/Private/StarshipStyle.cpp
--- a/Engine/Source/Editor/EditorStyle/Private/StarshipStyle.cpp
+++ b/Engine/Source/Editor/EditorStyle/Private/StarshipStyle.cpp
@@ -459,6 +459,7 @@ void FStarshipEditorStyle::FStyle::SetupGeneralStyles()
 		Set("Icons.Heart", new IMAGE_BRUSH_SVG("Starship/Common/Heart", Icon16x16));
 		Set("Icons.RecentAssets", new IMAGE_BRUSH_SVG("Starship/Common/RecentAssets", Icon16x16));
 		Set("Icons.Scalability", new IMAGE_BRUSH_SVG("Starship/Common/Scalability_16", Icon20x20));
+		Set("Icons.ViewportScalability", new IMAGE_BRUSH_SVG("Starship/Common/ViewportScalability_16", Icon16x16));
 		Set("Icons.EllipsisVerticalNarrow",
 			new CORE_IMAGE_BRUSH_SVG("Starship/Common/ellipsis-vertical-narrow", FVector2D(6, 24)));
 
@@ -2512,38 +2513,58 @@ void FStarshipEditorStyle::FStyle::SetupViewportStyles()
 			// Use normal styling for non-raised buttons.
 			Set("ViewportToolbarWarning", ToolBarStyle);
 
-			FCheckBoxStyle ToggleButtonStyleWarning = ToolBarStyle.ToggleButton;
+			FLinearColor WarningColor = FStyleColors::Warning.GetSpecifiedColor();
+			const FSlateRoundedBoxBrush WarningBrush(WarningColor, 5.0f, WarningColor, 1.0f);
+
+			FLinearColor WarningHoveredColor = FStyleColors::Warning.GetSpecifiedColor();
+			WarningHoveredColor.R = FMath::Min(1.0f, WarningHoveredColor.R * 1.5);
+			WarningHoveredColor.G = FMath::Min(1.0f, WarningHoveredColor.G * 1.5);
+			WarningHoveredColor.B = FMath::Min(1.0f, WarningHoveredColor.B * 1.5);
+			const FSlateRoundedBoxBrush WarningHoveredBrush(WarningHoveredColor, 5.0f, WarningHoveredColor, 1.0f);
+
+			FLinearColor WarningPressedColor = FStyleColors::Warning.GetSpecifiedColor();
+			WarningPressedColor.A = .50f;
+			const FSlateRoundedBoxBrush WarningPressedBrush(WarningPressedColor, 5.0f, WarningPressedColor, 1.0f);
+
+			FToolBarStyle NewViewportToolbarTopLevelWarningStyle = ToolBarStyle;
+
+			// Toggle Button Style
 			{
-				FLinearColor WarningColor = FStyleColors::Warning.GetSpecifiedColor();
-				const FSlateRoundedBoxBrush WarningBrush(WarningColor, 5.0f, WarningColor, 1.0f);
-
-				FLinearColor WarningHoveredColor = FStyleColors::Warning.GetSpecifiedColor();
-				WarningHoveredColor.R = FMath::Min(1.0f, WarningHoveredColor.R * 1.5);
-				WarningHoveredColor.G = FMath::Min(1.0f, WarningHoveredColor.G * 1.5);
-				WarningHoveredColor.B = FMath::Min(1.0f, WarningHoveredColor.B * 1.5);
-				const FSlateRoundedBoxBrush WarningHoveredBrush(WarningHoveredColor, 5.0f, WarningHoveredColor, 1.0f);
-
-				FLinearColor WarningPressedColor = FStyleColors::Warning.GetSpecifiedColor();
-				WarningPressedColor.A = .50f;
-				const FSlateRoundedBoxBrush WarningPressedBrush(WarningPressedColor, 5.0f, WarningPressedColor, 1.0f);
-
-				ToggleButtonStyleWarning.SetBackgroundImage(WarningBrush)
-					.SetBackgroundHoveredImage(WarningHoveredBrush)
-					.SetUncheckedImage(WarningBrush)
-					.SetUncheckedHoveredImage(WarningHoveredBrush)
-					.SetUncheckedPressedImage(WarningPressedBrush)
-					.SetCheckedImage(WarningBrush)
-					.SetCheckedHoveredImage(WarningHoveredBrush)
-					.SetCheckedPressedImage(WarningPressedBrush)
-					.SetForegroundColor(FStyleColors::ForegroundInverted)
-					.SetHoveredForegroundColor(FStyleColors::ForegroundInverted)
-					.SetUndeterminedForegroundColor(FStyleColors::ForegroundInverted)
-					.SetPressedForegroundColor(FStyleColors::ForegroundInverted);
+				FCheckBoxStyle ToggleButtonStyleWarning = ToolBarStyle.ToggleButton;
+				{
+					ToggleButtonStyleWarning.SetBackgroundImage(WarningBrush)
+						.SetBackgroundHoveredImage(WarningHoveredBrush)
+						.SetUncheckedImage(WarningBrush)
+						.SetUncheckedHoveredImage(WarningHoveredBrush)
+						.SetUncheckedPressedImage(WarningPressedBrush)
+						.SetCheckedImage(WarningBrush)
+						.SetCheckedHoveredImage(WarningHoveredBrush)
+						.SetCheckedPressedImage(WarningPressedBrush)
+						.SetForegroundColor(FStyleColors::ForegroundInverted)
+						.SetHoveredForegroundColor(FStyleColors::ForegroundInverted)
+						.SetUndeterminedForegroundColor(FStyleColors::ForegroundInverted)
+						.SetPressedForegroundColor(FStyleColors::ForegroundInverted);
+				}
+
+				NewViewportToolbarTopLevelWarningStyle.SetToggleButtonStyle(ToggleButtonStyleWarning)
+					.SetComboButtonPadding(FMargin(2.0f, 0));
 			}
 
-			FToolBarStyle NewViewportToolbarTopLevelWarningStyle = ToolBarStyle;
-			NewViewportToolbarTopLevelWarningStyle.SetToggleButtonStyle(ToggleButtonStyleWarning)
-				.SetComboButtonPadding(FMargin(2.0f, 0));
+			// Button Style
+			{
+				FButtonStyle ButtonStyleWarning = ToolBarStyle.ButtonStyle;
+				{
+					ButtonStyleWarning.SetHovered(WarningHoveredBrush)
+					.SetPressed(WarningPressedBrush)
+					.SetNormal(WarningBrush)
+					.SetNormalForeground(FStyleColors::ForegroundInverted)
+					.SetHoveredForeground(FStyleColors::ForegroundInverted)
+					.SetPressedForeground(FStyleColors::ForegroundInverted);
+				}
+
+				NewViewportToolbarTopLevelWarningStyle.SetButtonStyle(ButtonStyleWarning)
+				.SetButtonPadding(FMargin(2.0f, 0));
+			}
 
 			Set("ViewportToolbarWarning.Raised", NewViewportToolbarTopLevelWarningStyle);
 		}
diff --git a/Engine/Source/Editor/LevelEditor/Private/SLevelViewportToolBar.cpp b/Engine/Source/Editor/LevelEditor/Private/SLevelViewportToolBar.cpp
--- a/Engine/Source/Editor/LevelEditor/Private/SLevelViewportToolBar.cpp
+++ b/Engine/Source/Editor/LevelEditor/Private/SLevelViewportToolBar.cpp
@@ -238,11 +238,11 @@ void SLevelViewportToolBar::Construct( const FArguments& InArgs )
 						// Button to show scalability warnings
 						SNew(SEditorViewportToolbarMenu)
 						.ParentToolBar(SharedThis(this))
-						.Label(this, &SLevelViewportToolBar::GetScalabilityWarningLabel)
+						.Label_Static(&UE::UnrealEd::GetScalabilityWarningLabel)
 						.MenuStyle(&FAppStyle::Get().GetWidgetStyle<FButtonStyle>("EditorViewportToolBar.WarningButton"))
 						.OnGetMenuContent(this, &SLevelViewportToolBar::GetScalabilityWarningMenuContent)
 						.Visibility(this, &SLevelViewportToolBar::GetScalabilityWarningVisibility)
-						.ToolTipText(LOCTEXT("ScalabilityWarning_ToolTip", "Non-default scalability settings could be affecting what is shown in this viewport.\nFor example you may experience lower visual quality, reduced particle counts, and other artifacts that don't match what the scene would look like when running outside of the editor. Click to make changes."))
+						.ToolTipText_Static(&UE::UnrealEd::GetScalabilityWarningTooltip)
 					]
 					+ SHorizontalBox::Slot()
 					.Padding(ToolbarSlotPadding)
@@ -1094,8 +1094,7 @@ FText SLevelViewportToolBar::GetScalabilityWarningLabel() const
 
 EVisibility SLevelViewportToolBar::GetScalabilityWarningVisibility() const
 {
-	//This method returns magic numbers. 3 means epic
-	return GetDefault<UEditorPerformanceSettings>()->bEnableScalabilityWarningIndicator && Scalability::GetQualityLevels().GetMinQualityLevel() != 3 ? EVisibility::Visible : EVisibility::Collapsed;
+	return UE::UnrealEd::IsScalabilityWarningVisible() ? EVisibility::Visible : EVisibility::Collapsed;
 }
 
 TSharedRef<SWidget> SLevelViewportToolBar::GetScalabilityWarningMenuContent() const
diff --git a/Engine/Source/Editor/LevelEditor/Private/ViewportToolbar/LevelEditorViewportToolbarSections.cpp b/Engine/Source/Editor/LevelEditor/Private/ViewportToolbar/LevelEditorViewportToolbarSections.cpp
--- a/Engine/Source/Editor/LevelEditor/Private/ViewportToolbar/LevelEditorViewportToolbarSections.cpp
+++ b/Engine/Source/Editor/LevelEditor/Private/ViewportToolbar/LevelEditorViewportToolbarSections.cpp
@@ -16,9 +16,7 @@
 #include "LevelEditor.h"
 #include "LevelEditorActions.h"
 #include "LevelViewportActions.h"
-#include "SCommonEditorViewportToolbarBase.h"
 #include "SLevelViewport.h"
-#include "SScalabilitySettings.h"
 #include "Selection.h"
 #include "SortHelper.h"
 #include "Stats/StatsData.h"
@@ -421,7 +419,6 @@ FSlateIcon GetCameraSubmenuIconFromLevelViewport(const TWeakPtr<::SLevelViewport
 
 	return FSlateIcon(FAppStyle::GetAppStyleSetName(), "ClassIcon.CameraComponent");
 }
-
 } // namespace UE::LevelEditor::Private
 
 namespace UE::LevelEditor
@@ -1657,20 +1654,7 @@ FToolMenuEntry CreateViewportToolbarPerformanceAndScalabilitySubmenu()
 
 					PerformanceAndScalabilitySection.AddSeparator("PerformanceAndScalabilitySettings");
 
-					PerformanceAndScalabilitySection.AddSubMenu(
-						"Scalability",
-						LOCTEXT("ScalabilitySubMenu", "Viewport Scalability"),
-						LOCTEXT("ScalabilitySubMenu_ToolTip", "Open the viewport scalability settings"),
-						FNewToolMenuDelegate::CreateLambda(
-							[](UToolMenu* InMenu) -> void
-							{
-								FToolMenuSection& Section = InMenu->FindOrAddSection(NAME_None);
-								Section.AddEntry(FToolMenuEntry::InitWidget(
-									"ScalabilitySettings", SNew(SScalabilitySettings), FText(), true
-								));
-							}
-						)
-					);
+					PerformanceAndScalabilitySection.AddEntry(UE::UnrealEd::CreateScalabilitySubmenu());
 
 					PerformanceAndScalabilitySection.AddEntry(CreateMaterialQualityLevelSubmenu());
 
diff --git a/Engine/Source/Editor/UnrealEd/Private/SCommonEditorViewportToolbarBase.cpp b/Engine/Source/Editor/UnrealEd/Private/SCommonEditorViewportToolbarBase.cpp
--- a/Engine/Source/Editor/UnrealEd/Private/SCommonEditorViewportToolbarBase.cpp
+++ b/Engine/Source/Editor/UnrealEd/Private/SCommonEditorViewportToolbarBase.cpp
@@ -313,11 +313,11 @@ void SCommonEditorViewportToolbarBase::Construct(const FArguments& InArgs, TShar
 			// Button to show scalability warnings
 			SNew(SEditorViewportToolbarMenu)
 			.ParentToolBar(SharedThis(this))
-			.Label(this, &SCommonEditorViewportToolbarBase::GetScalabilityWarningLabel)
+			.Label_Static(&UE::UnrealEd::GetScalabilityWarningLabel)
 			.MenuStyle(&FAppStyle::Get().GetWidgetStyle<FButtonStyle>("EditorViewportToolBar.WarningButton"))
 			.OnGetMenuContent(this, &SCommonEditorViewportToolbarBase::GetScalabilityWarningMenuContent)
 			.Visibility(this, &SCommonEditorViewportToolbarBase::GetScalabilityWarningVisibility)
-			.ToolTipText(LOCTEXT("ScalabilityWarning_ToolTip", "Non-default scalability settings could be affecting what is shown in this viewport.\nFor example you may experience lower visual quality, reduced particle counts, and other artifacts that don't match what the scene would look like when running outside of the editor. Click to make changes."))
+			.ToolTipText_Static(&UE::UnrealEd::GetScalabilityWarningTooltip)
 		];
 
 	// Add optional toolbar slots to be added by child classes inherited from this common viewport toolbar
@@ -626,8 +626,7 @@ FText SCommonEditorViewportToolbarBase::GetScalabilityWarningLabel() const
 
 EVisibility SCommonEditorViewportToolbarBase::GetScalabilityWarningVisibility() const
 {
-	//This method returns magic numbers. 3 means epic
-	return GetDefault<UEditorPerformanceSettings>()->bEnableScalabilityWarningIndicator && GetShowScalabilityMenu() && Scalability::GetQualityLevels().GetMinQualityLevel() != 3 ? EVisibility::Visible : EVisibility::Collapsed;
+	return UE::UnrealEd::IsScalabilityWarningVisible() ? EVisibility::Visible : EVisibility::Collapsed;
 }
 
 TSharedRef<SWidget> SCommonEditorViewportToolbarBase::GetScalabilityWarningMenuContent() const
diff --git a/Engine/Source/Editor/UnrealEd/Private/ViewportToolbar/UnrealEdViewportToolbar.cpp b/Engine/Source/Editor/UnrealEd/Private/ViewportToolbar/UnrealEdViewportToolbar.cpp
--- a/Engine/Source/Editor/UnrealEd/Private/ViewportToolbar/UnrealEdViewportToolbar.cpp
+++ b/Engine/Source/Editor/UnrealEd/Private/ViewportToolbar/UnrealEdViewportToolbar.cpp
@@ -15,6 +15,7 @@
 #include "LevelEditorActions.h"
 #include "RayTracingDebugVisualizationMenuCommands.h"
 #include "SEditorViewport.h"
+#include "SScalabilitySettings.h"
 #include "ShowFlagMenuCommands.h"
 #include "Settings/EditorProjectSettings.h"
 #include "Settings/LevelEditorViewportSettings.h"
@@ -2995,6 +2996,52 @@ void ConstructScreenPercentageMenu(UToolMenu* InMenu)
 	}
 }
 
+FToolMenuEntry CreateScalabilitySubmenu()
+{
+	return FToolMenuEntry::InitDynamicEntry(
+		"Scalability",
+		FNewToolMenuSectionDelegate::CreateLambda(
+			[](FToolMenuSection& InnerSection) -> void
+			{
+				FToolMenuEntry& ScalabilitySubmenu = InnerSection.AddSubMenu(
+					"Scalability",
+					LOCTEXT("ScalabilitySubMenu", "Viewport Scalability"),
+					TAttribute<FText>::CreateLambda([]()
+					{
+						if (UE::UnrealEd::IsScalabilityWarningVisible())
+						{
+							return UE::UnrealEd::GetScalabilityWarningTooltip();
+						}
+
+						return LOCTEXT("ScalabilitySubMenu", "Viewport Scalability");
+					}),
+					FNewToolMenuDelegate::CreateLambda(
+						[](UToolMenu* InMenu) -> void
+						{
+							FToolMenuSection& Section = InMenu->FindOrAddSection(NAME_None);
+							Section.AddEntry(FToolMenuEntry::InitWidget(
+								"ScalabilitySettings", SNew(SScalabilitySettings), FText(), true
+							));
+						}
+					),
+					FToolUIAction(),
+					EUserInterfaceActionType::Button,
+					false,
+					 FSlateIcon(FAppStyle::GetAppStyleSetName(), "Icons.ViewportScalability")
+				);
+
+				ScalabilitySubmenu.SetShowInToolbarTopLevel(TAttribute<bool>::Create(&UE::UnrealEd::IsScalabilityWarningVisible
+				));
+
+				if (UE::UnrealEd::IsScalabilityWarningVisible())
+				{
+					ScalabilitySubmenu.StyleNameOverride = "ViewportToolbarWarning";
+				}
+			}
+		)
+	);
+}
+
 bool ShouldShowViewportRealtimeWarning(const FEditorViewportClient& ViewportClient)
 {
 	return !ViewportClient.IsRealtime() && !ViewportClient.IsRealtimeOverrideSet() && ViewportClient.IsPerspective();
@@ -3028,6 +3075,37 @@ FToolMenuEntry CreatePerformanceAndScalabilitySubmenu()
 	return Entry;
 }
 
+bool IsScalabilityWarningVisible()
+{
+	// This method returns magic numbers. 3 means epic
+	return GetDefault<UEditorPerformanceSettings>()->bEnableScalabilityWarningIndicator
+		&& Scalability::GetQualityLevels().GetMinQualityLevel() != 3;
+}
+
+FText GetScalabilityWarningLabel()
+{
+	if (UE::UnrealEd::IsScalabilityWarningVisible())
+	{
+		const int32 QualityLevel = Scalability::GetQualityLevels().GetMinQualityLevel();
+		if (QualityLevel >= 0)
+		{
+			return FText::Format(
+				LOCTEXT("ScalabilityWarning", "Scalability: {0}"),
+				Scalability::GetScalabilityNameFromQualityLevel(QualityLevel)
+			);
+		}
+	}
+
+	return FText();
+}
+
+FText GetScalabilityWarningTooltip()
+{
+	return LOCTEXT(
+		"ScalabilityWarning_ToolTip", "Non-default scalability settings could be affecting what is shown in this viewport.\nFor example you may experience lower visual quality, reduced particle counts, and other artifacts that don't match what the scene would look like when running outside of the editor. Click to make changes."
+	);
+}
+
 FToolMenuEntry CreateDefaultShowSubmenu()
 {
 	return FToolMenuEntry::InitSubMenu(
diff --git a/Engine/Source/Editor/UnrealEd/Public/ViewportToolbar/UnrealEdViewportToolbar.h b/Engine/Source/Editor/UnrealEd/Public/ViewportToolbar/UnrealEdViewportToolbar.h
--- a/Engine/Source/Editor/UnrealEd/Public/ViewportToolbar/UnrealEdViewportToolbar.h
+++ b/Engine/Source/Editor/UnrealEd/Public/ViewportToolbar/UnrealEdViewportToolbar.h
@@ -155,6 +155,15 @@ UNREALED_API bool ShouldShowViewportRealtimeWarning(const FEditorViewportClient&
 
 UNREALED_API FToolMenuEntry CreatePerformanceAndScalabilitySubmenu();
 
+UE_DEPRECATED_FORGAME(5.6, "Do not use. Will be made private.")
+UNREALED_API bool IsScalabilityWarningVisible();
+
+UE_DEPRECATED_FORGAME(5.6, "Do not use. Will be made private.")
+UNREALED_API FText GetScalabilityWarningLabel();
+
+UE_DEPRECATED_FORGAME(5.6, "Do not use. Will be made private.")
+UNREALED_API FText GetScalabilityWarningTooltip();
+
 /**
  * Creates a Show submenu with commonly used show flags
  */
@@ -174,6 +183,8 @@ UNREALED_API void CreateCameraSpeedMenu(UToolMenu* InMenu);
 
 UNREALED_API void ConstructScreenPercentageMenu(UToolMenu* InMenu);
 
+UNREALED_API FToolMenuEntry CreateScalabilitySubmenu();
+
 // Camera Menu Widgets
 TSharedRef<SWidget> CreateCameraMenuWidget(const TSharedRef<SEditorViewport>& InViewport);
 TSharedRef<SWidget> CreateFOVMenuWidget(const TSharedRef<SEditorViewport>& InViewport);
