Updates to Details panel UX for property and component overrides.
- New alert icon for properties.
- Icon style colors are no longer dimmed by compositing against Foreground.
- Use style colors for warnings and errors instead of plain red and yellow.
- Tooltip and menu entry edits.

diff --git a/Engine/Source/Editor/EditorStyle/Private/StarshipStyle.cpp b/Engine/Source/Editor/EditorStyle/Private/StarshipStyle.cpp
--- a/Engine/Source/Editor/EditorStyle/Private/StarshipStyle.cpp
+++ b/Engine/Source/Editor/EditorStyle/Private/StarshipStyle.cpp
@@ -3335,6 +3335,7 @@ void FStarshipEditorStyle::FStyle::SetupPropertyEditorStyles()
 		Set("DetailsView.OverrideInside", new IMAGE_BRUSH_SVG("Starship/DetailsView/DetailsOverrideInside", Icon16x16, FStyleColors::AccentBlue));
 		Set("DetailsView.OverrideHereInside", new IMAGE_BRUSH_SVG("Starship/DetailsView/DetailsOverrideHereInside", Icon16x16, FStyleColors::AccentBlue));
 		Set("DetailsView.OverrideInherited", new IMAGE_BRUSH_SVG("Starship/DetailsView/DetailsOverrideInherited", Icon16x16, FStyleColors::AccentBlue));
+		Set("DetailsView.OverrideAlert", new IMAGE_BRUSH_SVG("Starship/DetailsView/DetailsOverrideAlert", Icon16x16, FStyleColors::White));
 
 		// Hovered images for all the details panel override states
 		Set("DetailsView.OverrideHere.Hovered", new IMAGE_BRUSH_SVG("Starship/DetailsView/DetailsOverrideHere", Icon16x16, FStyleColors::White));
@@ -3344,6 +3345,7 @@ void FStarshipEditorStyle::FStyle::SetupPropertyEditorStyles()
 		Set("DetailsView.OverrideInside.Hovered", new IMAGE_BRUSH_SVG("Starship/DetailsView/DetailsOverrideInside", Icon16x16, FStyleColors::White));
 		Set("DetailsView.OverrideHereInside.Hovered", new IMAGE_BRUSH_SVG("Starship/DetailsView/DetailsOverrideHereInside", Icon16x16, FStyleColors::White));
 		Set("DetailsView.OverrideInherited.Hovered", new IMAGE_BRUSH_SVG("Starship/DetailsView/DetailsOverrideInherited", Icon16x16, FStyleColors::White));
+		Set("DetailsView.OverrideAlert.Hovered", new IMAGE_BRUSH_SVG("Starship/DetailsView/DetailsOverrideAlert", Icon16x16, FStyleColors::White));
 
 		
 		const FButtonStyle CategoryRowButton = FButtonStyle()
