Styling pass over new viewport toolbars to reduce paddings

diff --git a/Engine/Source/Editor/EditorStyle/Private/StarshipStyle.cpp b/Engine/Source/Editor/EditorStyle/Private/StarshipStyle.cpp
--- a/Engine/Source/Editor/EditorStyle/Private/StarshipStyle.cpp
+++ b/Engine/Source/Editor/EditorStyle/Private/StarshipStyle.cpp
@@ -2490,49 +2490,74 @@ void FStarshipEditorStyle::FStyle::SetupViewportStyles()
 				.SetPadding(FMargin(4.0f));
 		}
 
+		FComboButtonStyle ComboButtonStyle = SlimToolBarStyle.ComboButtonStyle;
+		{
+			FButtonStyle InnerButtonStyle = FButtonStyle(ButtonStyle).SetNormalPadding(0).SetPressedPadding(0);
+
+			ComboButtonStyle.SetContentPadding(FMargin(2.0f, 0, 0, 0))
+				.SetDownArrowPadding(FMargin(1.0f, 0, 0, 0))
+				.SetButtonStyle(InnerButtonStyle);
+		}
+
 		FToolBarStyle ToolBarStyle = SlimToolBarStyle;
 		{
 			FSlateBrush SeparatorBrush;
 			SeparatorBrush.DrawAs = ESlateBrushDrawType::Box;
-			SeparatorBrush.TintColor = FStyleColors::Foreground;
+			SeparatorBrush.TintColor = FStyleColors::Hover;
 
 			ToolBarStyle.SetIconSize(Icon16x16)
 				.SetButtonStyle(ButtonStyle)
 				.SetToggleButtonStyle(ToggleButtonStyle)
 				.SetButtonPadding(FMargin(2.0f, 0))
-				.SetComboButtonPadding(FMargin(6.0f, 0, 2.0f, 0))
+				.SetComboButtonStyle(ComboButtonStyle)
+				.SetComboButtonPadding(FMargin(4.0f, 0))
+				.SetButtonPadding(0)
 				.SetSeparatorBrush(SeparatorBrush)
 				.SetSeparatorThickness(1.0f)
-				.SetSeparatorPadding(FMargin(4.0f, 0));
+				.SetSeparatorPadding(FMargin(4.0f, 0))
+				.SetBackgroundPadding(FMargin(0, 4.0f, 4.0f, 4.0f));
 
 			Set("ViewportToolbar", ToolBarStyle);
 		}
 
-		// Special styling for warnings.
+		// Special styling for top-level raised buttons.
+		FToolBarStyle NewViewportToolbarTopLevelRaisedStyle = ToolBarStyle;
 		{
-			// Use normal styling for non-raised buttons.
-			Set("ViewportToolbarWarning", ToolBarStyle);
-
-			FLinearColor WarningColor = FStyleColors::Warning.GetSpecifiedColor();
-			const FSlateRoundedBoxBrush WarningBrush(WarningColor, 5.0f, WarningColor, 1.0f);
+			FButtonStyle ButtonStyleRaised = ButtonStyle;
+			{
+				FLinearColor BackgroundColorRaised = FStyleColors::Dropdown.GetSpecifiedColor();
+				BackgroundColorRaised.A = .0f;
+				const FSlateRoundedBoxBrush BackgroundBrushRaised(BackgroundColorRaised, 5.f, BackgroundColorRaised, 1.0);
 
-			FLinearColor WarningHoveredColor = FStyleColors::Warning.GetSpecifiedColor();
-			WarningHoveredColor.R = FMath::Min(1.0f, WarningHoveredColor.R * 1.5);
-			WarningHoveredColor.G = FMath::Min(1.0f, WarningHoveredColor.G * 1.5);
-			WarningHoveredColor.B = FMath::Min(1.0f, WarningHoveredColor.B * 1.5);
-			const FSlateRoundedBoxBrush WarningHoveredBrush(WarningHoveredColor, 5.0f, WarningHoveredColor, 1.0f);
+				ButtonStyleRaised.SetNormal(BackgroundBrushRaised).SetNormalForeground(FStyleColors::Foreground);
+			}
 
-			FLinearColor WarningPressedColor = FStyleColors::Warning.GetSpecifiedColor();
-			WarningPressedColor.A = .50f;
-			const FSlateRoundedBoxBrush WarningPressedBrush(WarningPressedColor, 5.0f, WarningPressedColor, 1.0f);
+			NewViewportToolbarTopLevelRaisedStyle.SetButtonStyle(ButtonStyleRaised).SetComboButtonPadding(0);
+			Set("ViewportToolbar.Raised", NewViewportToolbarTopLevelRaisedStyle);
+		}
 
-			FToolBarStyle NewViewportToolbarTopLevelWarningStyle = ToolBarStyle;
+		// Special styling for warnings.
+		{
+			// Use normal styling for non-raised buttons.
+			Set("ViewportToolbarWarning", ToolBarStyle);
 
-			// Toggle Button Style
 			{
-				FCheckBoxStyle ToggleButtonStyleWarning = ToolBarStyle.ToggleButton;
-				{
-					ToggleButtonStyleWarning.SetBackgroundImage(WarningBrush)
+				FLinearColor WarningColor = FStyleColors::Warning.GetSpecifiedColor();
+				const FSlateRoundedBoxBrush WarningBrush(WarningColor, 5.0f, WarningColor, 1.0f);
+
+				FLinearColor WarningHoveredColor = FStyleColors::Warning.GetSpecifiedColor();
+				WarningHoveredColor.R = FMath::Min(1.0f, WarningHoveredColor.R * 1.5);
+				WarningHoveredColor.G = FMath::Min(1.0f, WarningHoveredColor.G * 1.5);
+				WarningHoveredColor.B = FMath::Min(1.0f, WarningHoveredColor.B * 1.5);
+				const FSlateRoundedBoxBrush WarningHoveredBrush(WarningHoveredColor, 5.0f, WarningHoveredColor, 1.0f);
+
+				FLinearColor WarningPressedColor = FStyleColors::Warning.GetSpecifiedColor();
+				WarningPressedColor.A = .50f;
+				const FSlateRoundedBoxBrush WarningPressedBrush(WarningPressedColor, 5.0f, WarningPressedColor, 1.0f);
+
+				const FCheckBoxStyle ToggleButtonStyleRaisedWarning =
+					FCheckBoxStyle(NewViewportToolbarTopLevelRaisedStyle.ToggleButton)
+						.SetBackgroundImage(WarningBrush)
 						.SetBackgroundHoveredImage(WarningHoveredBrush)
 						.SetUncheckedImage(WarningBrush)
 						.SetUncheckedHoveredImage(WarningHoveredBrush)
@@ -2544,47 +2569,22 @@ void FStarshipEditorStyle::FStyle::SetupViewportStyles()
 						.SetHoveredForegroundColor(FStyleColors::ForegroundInverted)
 						.SetUndeterminedForegroundColor(FStyleColors::ForegroundInverted)
 						.SetPressedForegroundColor(FStyleColors::ForegroundInverted);
-				}
-
-				NewViewportToolbarTopLevelWarningStyle.SetToggleButtonStyle(ToggleButtonStyleWarning)
-					.SetComboButtonPadding(FMargin(2.0f, 0));
-			}
 
-			// Button Style
-			{
-				FButtonStyle ButtonStyleWarning = ToolBarStyle.ButtonStyle;
-				{
-					ButtonStyleWarning.SetHovered(WarningHoveredBrush)
-					.SetPressed(WarningPressedBrush)
-					.SetNormal(WarningBrush)
-					.SetNormalForeground(FStyleColors::ForegroundInverted)
-					.SetHoveredForeground(FStyleColors::ForegroundInverted)
-					.SetPressedForeground(FStyleColors::ForegroundInverted);
-				}
-
-				NewViewportToolbarTopLevelWarningStyle.SetButtonStyle(ButtonStyleWarning)
-				.SetButtonPadding(FMargin(2.0f, 0));
-			}
-
-			Set("ViewportToolbarWarning.Raised", NewViewportToolbarTopLevelWarningStyle);
-		}
-
-		// Special styling for top-level raised buttons.
-		{
-			FButtonStyle ButtonStyleRaised = ButtonStyle;
-			{
-				FLinearColor BackgroundColorRaised = FStyleColors::Dropdown.GetSpecifiedColor();
-				BackgroundColorRaised.A = .0f;
-				const FSlateRoundedBoxBrush BackgroundBrushRaised(BackgroundColorRaised, 5.f, BackgroundColorRaised, 1.0);
-
-				ButtonStyleRaised.SetNormal(BackgroundBrushRaised).SetNormalForeground(FStyleColors::Foreground);
-			}
-
-			FToolBarStyle NewViewportToolbarTopLevelRaisedStyle = ToolBarStyle;
-			{
-				NewViewportToolbarTopLevelRaisedStyle.SetButtonStyle(ButtonStyleRaised).SetComboButtonPadding(FMargin(2.0f, 0));
-
-				Set("ViewportToolbar.Raised", NewViewportToolbarTopLevelRaisedStyle);
+				const FButtonStyle ButtonStyleRaisedWarning =
+					FButtonStyle(NewViewportToolbarTopLevelRaisedStyle.ButtonStyle)
+						.SetHovered(WarningHoveredBrush)
+						.SetPressed(WarningPressedBrush)
+						.SetNormal(WarningBrush)
+						.SetNormalForeground(FStyleColors::ForegroundInverted)
+						.SetHoveredForeground(FStyleColors::ForegroundInverted)
+						.SetPressedForeground(FStyleColors::ForegroundInverted);
+
+				const FToolBarStyle NewViewportToolbarTopLevelWarningRaisedStyle =
+					FToolBarStyle(NewViewportToolbarTopLevelRaisedStyle)
+						.SetToggleButtonStyle(ToggleButtonStyleRaisedWarning)
+						.SetButtonStyle(ButtonStyleRaisedWarning);
+
+				Set("ViewportToolbarWarning.Raised", NewViewportToolbarTopLevelWarningRaisedStyle);
 			}
 		}
 
@@ -2599,17 +2599,17 @@ void FStarshipEditorStyle::FStyle::SetupViewportStyles()
 					.SetNormalPadding(FMargin(0, 0, 2.f, 0))
 					.SetPressedPadding(FMargin(0, 0, 2.f, 0));
 
-			FToolBarStyle ViewportSizingStyle = ToolBarStyle;
-
-			FComboButtonStyle ViewportSizingComboButtonStyle =
-				FComboButtonStyle(FStarshipCoreStyle::GetCoreStyle().GetWidgetStyle<FComboButtonStyle>("ComboButton"))
-					.SetDownArrowPadding(FMargin(1.f, 0.f, 2.f, 0.f))
-					.SetDownArrowImage(CORE_IMAGE_BRUSH_SVG("Starship/Common/ellipsis-vertical-narrow", FVector2D(6, 20)));
-			ViewportSizingComboButtonStyle.ButtonStyle = ViewportSizingButtonStyle;
+			const FComboButtonStyle ViewportSizingComboButtonStyle =
+				FComboButtonStyle(ComboButtonStyle)
+					.SetDownArrowPadding(FMargin(5.f, 0.f, 2.f, 0.f))
+					.SetDownArrowImage(CORE_IMAGE_BRUSH_SVG("Starship/Common/ellipsis-vertical-narrow", FVector2D(6, 20)))
+					.SetButtonStyle(ViewportSizingButtonStyle);
 
-			ViewportSizingStyle.SetButtonStyle(ViewportSizingButtonStyle);
-			ViewportSizingStyle.SetComboButtonStyle(ViewportSizingComboButtonStyle);
-			ViewportSizingStyle.SetSeparatorPadding(FMargin(0.f, 0.f, 0.f, 0.f));
+			const FToolBarStyle ViewportSizingStyle = FToolBarStyle(ToolBarStyle)
+														  .SetButtonStyle(ViewportSizingButtonStyle)
+														  .SetComboButtonStyle(ViewportSizingComboButtonStyle)
+														  .SetComboButtonPadding(FMargin(4.0f, 0, 1.0f, 0))
+														  .SetSeparatorPadding(FMargin(0.f, 0.f, 0.f, 0.f));
 
 			Set("ViewportToolbarViewportSizingSubmenu", ViewportSizingStyle);
 
@@ -2628,14 +2628,13 @@ void FStarshipEditorStyle::FStyle::SetupViewportStyles()
 						.SetPressed(FSlateRoundedBoxBrush(FStyleColors::Dropdown, FVector4(0, 4.0f, 4.0f, 0)))
 						.SetDisabled(FSlateRoundedBoxBrush(FStyleColors::Dropdown, FVector4(0, 4.0f, 4.0f, 0)));
 				}
-				FToolBarStyle ViewportSizingRaisedStyle = ToolBarStyle;
-				{
-					ViewportSizingRaisedStyle.SetButtonStyle(ButtonStyleRaised)
-						.SetButtonPadding(0)
-						.SetComboButtonPadding(FMargin(2.0f, 0));
 
-					Set("ViewportToolbarViewportSizingSubmenu.Raised", ViewportSizingRaisedStyle);
-				}
+				const FToolBarStyle ViewportSizingRaisedStyle = FToolBarStyle(ToolBarStyle)
+																	.SetButtonStyle(ButtonStyleRaised)
+																	.SetButtonPadding(0)
+																	.SetComboButtonPadding(FMargin(2.0f, 0));
+
+				Set("ViewportToolbarViewportSizingSubmenu.Raised", ViewportSizingRaisedStyle);
 			}
 		}
 	}
