Use new FToolBarStyle::SetRaisedChildrenRightPadding for the new viewport toolbar style

diff --git a/Engine/Source/Editor/EditorStyle/Private/StarshipStyle.cpp b/Engine/Source/Editor/EditorStyle/Private/StarshipStyle.cpp
--- a/Engine/Source/Editor/EditorStyle/Private/StarshipStyle.cpp
+++ b/Engine/Source/Editor/EditorStyle/Private/StarshipStyle.cpp
@@ -2515,7 +2515,8 @@ void FStarshipEditorStyle::FStyle::SetupViewportStyles()
 				.SetSeparatorBrush(SeparatorBrush)
 				.SetSeparatorThickness(1.0f)
 				.SetSeparatorPadding(FMargin(4.0f, 0))
-				.SetBackgroundPadding(FMargin(0, 4.0f, 4.0f, 4.0f));
+				.SetBackgroundPadding(FMargin(0, 4.0f, 4.0f, 4.0f))
+				.SetRaisedChildrenRightPadding(16.0f);
 
 			Set("ViewportToolbar", ToolBarStyle);
 		}
