Sequencer: Reinstate ScrulbHandleUp/Down and VanillaScrubHandleUpDown because it's used outside of Sequencer.

#rnx

diff --git a/Engine/Source/Editor/EditorStyle/Private/StarshipStyle.cpp b/Engine/Source/Editor/EditorStyle/Private/StarshipStyle.cpp
--- a/Engine/Source/Editor/EditorStyle/Private/StarshipStyle.cpp
+++ b/Engine/Source/Editor/EditorStyle/Private/StarshipStyle.cpp
@@ -2199,6 +2199,11 @@ void FStarshipEditorStyle::FStyle::SetupSequencerStyles()
 
 		Set("Sequencer.NotificationImage_AddedPlayMovieSceneEvent", new IMAGE_BRUSH("Old/Checkbox_checked", Icon16x16));
 
+		Set("Sequencer.Timeline.ScrubHandleDown", new BOX_BRUSH("Sequencer/ScrubHandleDown", FMargin(6.f / 13.f, 5 / 12.f, 6 / 13.f, 8 / 12.f)));
+		Set("Sequencer.Timeline.ScrubHandleUp", new BOX_BRUSH("Sequencer/ScrubHandleUp", FMargin(6.f / 13.f, 8 / 12.f, 6 / 13.f, 5 / 12.f)));
+		Set("Sequencer.Timeline.VanillaScrubHandleDown", new BOX_BRUSH("Sequencer/ScrubHandleDown_Clamped", FMargin(6.f / 13.f, 3.f / 12.f, 6.f / 13.f, 7.f / 12.f)));
+		Set("Sequencer.Timeline.VanillaScrubHandleUp", new BOX_BRUSH("Sequencer/ScrubHandleUp_Clamped", FMargin(6.f / 13.f, 8 / 12.f, 6 / 13.f, 5 / 12.f)));
+
 		Set("Sequencer.RestoreAnimatedState", new IMAGE_BRUSH("Sequencer/Main_Icons/Icon_Sequencer_RestoreAnimatedState_24x", Icon48x48));
 		Set("Sequencer.GenericGripLeft", new BOX_BRUSH("Sequencer/GenericGripLeft", FMargin(5.f / 16.f)));
 		Set("Sequencer.GenericGripRight", new BOX_BRUSH("Sequencer/GenericGripRight", FMargin(5.f / 16.f)));
