Update the new viewport toolbar style to move and adjust the style of the overflow menu button

diff --git a/Engine/Source/Editor/EditorStyle/Private/StarshipStyle.cpp b/Engine/Source/Editor/EditorStyle/Private/StarshipStyle.cpp
--- a/Engine/Source/Editor/EditorStyle/Private/StarshipStyle.cpp
+++ b/Engine/Source/Editor/EditorStyle/Private/StarshipStyle.cpp
@@ -2534,6 +2534,8 @@ void FStarshipEditorStyle::FStyle::SetupViewportStyles()
 				.SetSeparatorThickness(1.0f)
 				.SetSeparatorPadding(FMargin(4.0f, 0))
 				.SetBackgroundPadding(FMargin(0, 4.0f, 4.0f, 4.0f))
+				.SetWrapButtonPadding(FMargin(4.0f, 0))
+				.SetWrapButtonIndex(-3)
 				.SetRaisedChildrenRightPadding(16.0f);
 
 			Set("ViewportToolbar", ToolBarStyle);
