Add debug option CVar for SceneGraph PropertyState to show undetermined state

diff --git a/Engine/Source/Editor/EditorStyle/Private/StarshipStyle.cpp b/Engine/Source/Editor/EditorStyle/Private/StarshipStyle.cpp
--- a/Engine/Source/Editor/EditorStyle/Private/StarshipStyle.cpp
+++ b/Engine/Source/Editor/EditorStyle/Private/StarshipStyle.cpp
@@ -3360,6 +3360,7 @@ void FStarshipEditorStyle::FStyle::SetupPropertyEditorStyles()
 		/*******************************************************************************/
 		
 		// Background images for all the details panels override states
+		Set("DetailsView.OverrideUndetermined", new IMAGE_BRUSH_SVG("Starship/Common/QuestionAnswer", Icon16x16, FStyleColors::AccentPurple));
 		Set("DetailsView.OverrideHere", new IMAGE_BRUSH_SVG("Starship/DetailsView/DetailsOverrideHere", Icon16x16, FStyleColors::AccentBlue));
 		Set("DetailsView.OverrideAdded", new IMAGE_BRUSH_SVG("Starship/DetailsView/DetailsOverrideAdded", Icon16x16, FStyleColors::AccentGreen));
 		Set("DetailsView.OverrideNone", new IMAGE_BRUSH_SVG("Starship/DetailsView/DetailsOverrideNone", Icon16x16, FStyleColors::Foreground));
diff --git a/Engine/Source/Runtime/CoreUObject/Public/UObject/OverridableManager.h b/Engine/Source/Runtime/CoreUObject/Public/UObject/OverridableManager.h
--- a/Engine/Source/Runtime/CoreUObject/Public/UObject/OverridableManager.h
+++ b/Engine/Source/Runtime/CoreUObject/Public/UObject/OverridableManager.h
@@ -7,6 +7,8 @@
 #include "UObject/OverriddenPropertySet.h"
 #include "UObject/PropertyBagRepository.h"
 
+#include "OverridableManager.generated.h"
+
 /*
  *************************************************************************************
  * Overridable serialization is experimental, not supported and use at your own risk *
@@ -83,6 +85,7 @@ private:
 
 #endif // WITH_EDITORONLY_DATA
 
+UENUM()
 enum class EOverriddenState : uint8
 {
 	NoOverrides, // no on this object and any of its instanced subobjects
