[Viewport Toolbar] Update level editor viewport toolbar icons

- Adds new set of icons.
- Starts applying new icons to Level Editor Viewport toolbar

Icons still missing:

- e.g. Groom, Screen Percentage
- some instances of Widget based entries have an icon, but cannot currently be shown (FOV, Far View Plane, camera speed submenu)

Yet to be applied:
- Icons for camera spawn menu

diff --git a/Engine/Plugins/2D/Paper2D/Source/SmartSnapping/Private/SmartSnappingModule.cpp b/Engine/Plugins/2D/Paper2D/Source/SmartSnapping/Private/SmartSnappingModule.cpp
--- a/Engine/Plugins/2D/Paper2D/Source/SmartSnapping/Private/SmartSnappingModule.cpp
+++ b/Engine/Plugins/2D/Paper2D/Source/SmartSnapping/Private/SmartSnappingModule.cpp
@@ -71,7 +71,7 @@ public:
 			LOCTEXT(
 				"View_Extension_PlanarSnap_Tooltip", "If Enabled, actors will snap to the nearest location on the constraint plane (NOTE: Only works correctly in perspective views right now!)"
 			),
-			FSlateIcon(),
+			FSlateIcon(FAppStyle::GetAppStyleSetName(), "ViewportToolbar.SnapPlanar"),
 			Action_TogglePlanarSnap,
 			NAME_None,
 			EUserInterfaceActionType::ToggleButton
diff --git a/Engine/Source/Editor/EditorStyle/Private/StarshipStyle.cpp b/Engine/Source/Editor/EditorStyle/Private/StarshipStyle.cpp
--- a/Engine/Source/Editor/EditorStyle/Private/StarshipStyle.cpp
+++ b/Engine/Source/Editor/EditorStyle/Private/StarshipStyle.cpp
@@ -2670,8 +2670,29 @@ void FStarshipEditorStyle::FStyle::SetupViewportStyles()
 
 		// Viewport Toolbar Icons
 		{
-			Set("ViewportToolbar.SetShowGrid", new IMAGE_BRUSH_SVG("Starship/Common/Grid", Icon16x16));
+			Set("ViewportToolbar.CameraSpeed", new IMAGE_BRUSH_SVG("Starship/Common/CameraSpeed_16", Icon16x16));
+			Set("ViewportToolbar.EV100", new IMAGE_BRUSH_SVG("Starship/Common/EV100_16", Icon16x16));
+			Set("ViewportToolbar.ExactCameraView", new IMAGE_BRUSH_SVG("Starship/Common/ExactCameraView_16", Icon16x16));
+			Set("ViewportToolbar.Exposure", new IMAGE_BRUSH_SVG("Starship/Common/Exposure_16", Icon16x16));
+			Set("ViewportToolbar.FarViewPlane", new IMAGE_BRUSH_SVG("Starship/Common/FarViewPlane_16", Icon16x16));
+			Set("ViewportToolbar.FieldOfView", new IMAGE_BRUSH_SVG("Starship/Common/FieldOfView_16", Icon16x16));
+			Set("ViewportToolbar.GizmoScale", new IMAGE_BRUSH_SVG("Starship/Common/GizmoScale_16", Icon16x16));
 			Set("ViewportToolbar.PreviewSceneSettings", new IMAGE_BRUSH_SVG("Starship/AssetEditors/PreviewSceneSettings_16", Icon16x16));
+			Set("ViewportToolbar.SetShowGrid", new IMAGE_BRUSH_SVG("Starship/Common/Grid", Icon16x16));
+			Set("ViewportToolbar.Snap", new IMAGE_BRUSH_SVG("Starship/Common/Snap_16", Icon16x16));
+			Set("ViewportToolbar.SnapLocation", new IMAGE_BRUSH_SVG("Starship/Common/SnapLocation_16", Icon16x16));
+			Set("ViewportToolbar.SnapPlanar", new IMAGE_BRUSH_SVG("Starship/Common/SnapPlanar_16", Icon16x16));
+			Set("ViewportToolbar.SnapRotation", new IMAGE_BRUSH_SVG("Starship/Common/SnapRotation_16", Icon16x16));
+			Set("ViewportToolbar.SnapScale", new IMAGE_BRUSH_SVG("Starship/Common/SnapScale_16", Icon16x16));
+			Set("ViewportToolbar.SpeedScalar", new IMAGE_BRUSH_SVG("Starship/Common/SpeedScalar_16", Icon16x16));
+			Set("ViewportToolbar.SurfaceSnap", new IMAGE_BRUSH_SVG("Starship/Common/SurfaceSnap_16", Icon16x16));
+			Set("ViewportToolbar.SurfaceSnapRotateToNormal", new IMAGE_BRUSH_SVG("Starship/Common/SurfaceSnapRotateToNormal_16", Icon16x16));
+			Set("ViewportToolbar.TransformRotate", new IMAGE_BRUSH_SVG("Starship/Common/TransformRotate_16", Icon16x16));
+			Set("ViewportToolbar.TransformScale", new IMAGE_BRUSH_SVG("Starship/Common/TransformScale_16", Icon16x16));
+			Set("ViewportToolbar.TransformSelect", new IMAGE_BRUSH_SVG("Starship/Common/TransformSelect_16", Icon16x16));
+			Set("ViewportToolbar.TransformTranslate", new IMAGE_BRUSH_SVG("Starship/Common/TransformMove_16", Icon16x16));
+			Set("ViewportToolbar.VPCamera", new IMAGE_BRUSH_SVG("Starship/Common/VPCamera_16", Icon16x16));
+			Set("ViewportToolbar.VPCineCamera", new IMAGE_BRUSH_SVG("Starship/Common/VPCineCamera_16", Icon16x16));
 		}
 	}
 
@@ -4491,6 +4512,26 @@ void FStarshipEditorStyle::FStyle::SetupLevelEditorStyle()
 		Set("LevelEditor.BrowseViewportControls",  new IMAGE_BRUSH_SVG("Starship/Common/ViewportControls", Icon16x16));
 		Set("LevelEditor.PasteHere",				new IMAGE_BRUSH_SVG("Starship/Actors/paste-here", Icon16x16));
 
+		Set("LevelEditor.AllowGroupSelection",			new IMAGE_BRUSH_SVG("Starship/Common/GroupActors", Icon16x16));
+		Set("LevelEditor.AllowTranslucentSelection",		new IMAGE_BRUSH_SVG("Starship/Common/Transparency", Icon16x16));
+		Set("LevelEditor.EnableActorSnap",				new IMAGE_BRUSH_SVG("Starship/Common/SnapActor_16", Icon16x16));
+		Set("LevelEditor.EnableVertexSnap",				new IMAGE_BRUSH_SVG("Starship/Common/SnapVertex_16", Icon16x16));
+		Set("LevelEditor.InvertSelection",				new IMAGE_BRUSH_SVG("Starship/Common/SelectInvert_16", Icon16x16));
+		Set("LevelEditor.PreviewPlatform", new IMAGE_BRUSH_SVG("Starship/Common/PreviewPlatform_16", Icon16x16));
+
+		Set("LevelEditor.SelectAllDescendants",			new IMAGE_BRUSH_SVG("Starship/Common/AllDescendants_16", Icon16x16));
+		Set("LevelEditor.SelectImmediateChildren",		new IMAGE_BRUSH_SVG("Starship/Common/ImmediateChildren_16", Icon16x16));
+		Set("LevelEditor.SelectNone",					new IMAGE_BRUSH_SVG("Starship/Common/DeselectAll_16", Icon16x16));
+		Set("LevelEditor.ShowTransformWidget",			new IMAGE_BRUSH_SVG("Starship/Common/ShowTransformGizmo_16", Icon16x16));
+		Set("LevelEditor.SnapCameraToObject", 			new IMAGE_BRUSH_SVG("Starship/Common/MoveCameraToObject_16", Icon16x16));
+		Set("LevelEditor.SnapObjectToCamera", 			new IMAGE_BRUSH_SVG("Starship/Common/MoveObjectToCamera_16", Icon16x16));
+		Set("LevelEditor.StrictBoxSelect",				new IMAGE_BRUSH_SVG("Starship/Common/StrictMarqueeSelection_16", Icon16x16));
+		Set("LevelEditor.ToggleSocketSnapping", 			new IMAGE_BRUSH_SVG("Starship/Common/SnapSocket_16", Icon16x16));
+		Set("LevelEditor.TransparentBoxSelect", 			new IMAGE_BRUSH_SVG("Starship/Common/MarqueeSelectOccluded_16", Icon16x16));
+
+		Set("LevelEditor.MaterialQuality", new IMAGE_BRUSH_SVG("Starship/Common/MaterialQuality_16", Icon16x16));
+			
+		
 		Set("MainFrame.ToggleFullscreen",          new IMAGE_BRUSH_SVG("Starship/Common/EnableFullscreen", Icon16x16));
 		Set("MainFrame.LoadLayout",                new IMAGE_BRUSH_SVG("Starship/Common/LayoutLoad", Icon16x16));
 		Set("MainFrame.SaveLayout",                new IMAGE_BRUSH_SVG("Starship/Common/LayoutSave", Icon16x16));
@@ -4513,10 +4554,10 @@ void FStarshipEditorStyle::FStyle::SetupLevelEditorStyle()
 		Set("MainFrame.AboutUnrealEd",              new IMAGE_BRUSH_SVG("About/UnrealLogo", Icon16x16));
 		Set("MainFrame.CreditsUnrealEd",            new IMAGE_BRUSH_SVG("Starship/Common/Credits", Icon16x16));
 
-		Set( "EditorViewport.SelectMode", new IMAGE_BRUSH_SVG("Starship/EditorViewport/select", Icon16x16) );
-		Set( "EditorViewport.TranslateMode", new IMAGE_BRUSH_SVG( "Starship/EditorViewport/translate", Icon16x16 ) );
-		Set( "EditorViewport.RotateMode", new IMAGE_BRUSH_SVG("Starship/EditorViewport/rotate", Icon16x16 ) );
-		Set( "EditorViewport.ScaleMode", new IMAGE_BRUSH_SVG( "Starship/EditorViewport/scale", Icon16x16 ) );
+		Set( "EditorViewport.SelectMode", new IMAGE_BRUSH_SVG("Starship/Common/TransformSelect_16", Icon16x16) );
+		Set( "EditorViewport.TranslateMode", new IMAGE_BRUSH_SVG("Starship/Common/TransformMove_16", Icon16x16 ) );
+		Set( "EditorViewport.RotateMode", new IMAGE_BRUSH_SVG("Starship/Common/TransformRotate_16", Icon16x16 ) );
+		Set( "EditorViewport.ScaleMode", new IMAGE_BRUSH_SVG("Starship/Common/TransformScale_16", Icon16x16 ) );
 		Set("EditorViewport.Speed", new IMAGE_BRUSH_SVG("Starship/EditorViewport/speed", Icon16x16));
 
 		Set( "EditorViewport.TranslateRotateMode", new IMAGE_BRUSH_SVG("Starship/EditorViewport/TranslateRotate3D", Icon16x16 ) );
@@ -4551,7 +4592,7 @@ void FStarshipEditorStyle::FStyle::SetupLevelEditorStyle()
 		Set( "EditorViewport.LitMode",            	  new IMAGE_BRUSH_SVG("Starship/Common/LitCube", Icon16x16 ) );
 		Set( "EditorViewport.UnlitMode",          	  new IMAGE_BRUSH_SVG("Starship/Common/UnlitCube", Icon16x16 ) );
 		Set( "EditorViewport.WireframeMode",      	  new IMAGE_BRUSH_SVG("Starship/Common/BrushWireframe", Icon16x16 ) );
-		Set( "EditorViewport.LitWireframeMode",       new IMAGE_BRUSH_SVG("Starship/Common/BrushWireframe", Icon16x16) );
+		Set( "EditorViewport.LitWireframeMode",       new IMAGE_BRUSH_SVG("Starship/Common/LitWireframe_16", Icon16x16) );
 		Set( "EditorViewport.DetailLightingMode", 	  new IMAGE_BRUSH_SVG("Starship/Common/DetailLighting", Icon16x16 ) );
 		Set( "EditorViewport.LightingOnlyMode",   	  new IMAGE_BRUSH_SVG("Starship/Common/LightBulb", Icon16x16 ) );
 		
@@ -4585,15 +4626,15 @@ void FStarshipEditorStyle::FStyle::SetupLevelEditorStyle()
 
 		Set( "EditorViewport.ReflectionOverrideMode", new IMAGE_BRUSH_SVG("Starship/Common/Reflections", Icon16x16 ) );
 		Set( "EditorViewport.VisualizeBufferMode",    new IMAGE_BRUSH_SVG("Starship/Common/BufferVisualization", Icon16x16 ) );
-		Set( "EditorViewport.VisualizeNaniteMode",    new IMAGE_BRUSH_SVG("Starship/Common/BufferVisualization", Icon16x16 ) );
-		Set( "EditorViewport.VisualizeLumenMode",     new IMAGE_BRUSH_SVG("Starship/Common/BufferVisualization", Icon16x16 ) );
+		Set( "EditorViewport.VisualizeNaniteMode",    new IMAGE_BRUSH_SVG("Starship/Common/Nanite_16", Icon16x16 ) );
+		Set( "EditorViewport.VisualizeLumenMode",     new IMAGE_BRUSH_SVG("Starship/Common/Lumen_16", Icon16x16 ) );
 		Set( "EditorViewport.VisualizeSubstrateMode", new IMAGE_BRUSH_SVG("Starship/Common/BufferVisualization", Icon16x16 ) );
 		Set( "EditorViewport.VisualizeGroomMode",     new IMAGE_BRUSH_SVG("Starship/Common/BufferVisualization", Icon16x16 ) );
-		Set( "EditorViewport.VisualizeVirtualShadowMapMode", new IMAGE_BRUSH_SVG("Starship/Common/BufferVisualization", Icon16x16 ) );
-		Set( "EditorViewport.VisualizeActorColorationMode", new IMAGE_BRUSH("Icons/icon_ViewMode_LODColoration_16px", Icon16x16));
+		Set( "EditorViewport.VisualizeVirtualShadowMapMode", new IMAGE_BRUSH_SVG("Starship/Common/VirtualShadowMap_16", Icon16x16 ) );
+		Set( "EditorViewport.VisualizeActorColorationMode", new IMAGE_BRUSH_SVG("Starship/Common/ActorColoration_16", Icon16x16));
 		Set( "EditorViewport.CollisionPawn",          new IMAGE_BRUSH_SVG("Starship/Common/PlayerCollision", Icon16x16 ) );
 		Set( "EditorViewport.CollisionVisibility",    new IMAGE_BRUSH_SVG("Starship/Common/VisibilityCollision", Icon16x16 ) );
-		
+
 		Set( "EditorViewport.Perspective", new IMAGE_BRUSH_SVG("Starship/Common/ViewPerspective", Icon16x16 ) );
 		Set( "EditorViewport.Top",         new IMAGE_BRUSH_SVG("Starship/Common/ViewTop", Icon16x16 ) );
 		Set( "EditorViewport.Left",        new IMAGE_BRUSH_SVG("Starship/Common/ViewLeft", Icon16x16 ) );
@@ -4985,14 +5026,14 @@ void FStarshipEditorStyle::FStyle::SetupLevelEditorStyle()
 		Set( "LevelViewport.ViewportConfig_FourPanes2x2", new IMAGE_BRUSH( "Icons/ViewportLayout_FourPanes2x2", IconLayoutSize ) );
 		Set( "LevelViewport.ViewportConfig_FourPanes2x2.Small", new IMAGE_BRUSH( "Icons/ViewportLayout_FourPanes2x2", IconLayoutSizeSmall ) );
 
-		Set( "LevelViewport.EjectActorPilot", new IMAGE_BRUSH_SVG( "Starship/EditorViewport/eject", Icon16x16 ) );
-		Set( "LevelViewport.EjectActorPilot.Small", new IMAGE_BRUSH_SVG( "Starship/EditorViewport/eject", Icon16x16 ) );
+		Set( "LevelViewport.EjectActorPilot", new IMAGE_BRUSH_SVG( "Starship/Common/StopPiloting_16", Icon16x16 ) );
+		Set( "LevelViewport.EjectActorPilot.Small", new IMAGE_BRUSH_SVG( "Starship/Common/StopPiloting_16", Icon16x16 ) );
 		Set( "LevelViewport.PilotSelectedActor", new IMAGE_BRUSH_SVG( "Starship/EditorViewport/pilot", Icon16x16 ) );
-		
-		Set( "LevelViewport.ToggleActorPilotCameraView",       new IMAGE_BRUSH_SVG( "Starship/EditorViewport/actor-pilot-camera", Icon16x16 ) );
-		Set( "LevelViewport.ToggleActorPilotCameraView.Small", new IMAGE_BRUSH_SVG( "Starship/EditorViewport/actor-pilot-camera", Icon16x16 ) );
+
+		Set( "LevelViewport.ToggleActorPilotCameraView",       new IMAGE_BRUSH_SVG(  "Starship/Common/ExactCameraView_16", Icon16x16 ) );
+		Set( "LevelViewport.ToggleActorPilotCameraView.Small", new IMAGE_BRUSH_SVG( "Starship/Common/ExactCameraView_16", Icon16x16 ) );
 	}
-	
+
 	// Level editor status bar
 	{
 		Set( "TransformSettings.RelativeCoordinateSettings", new IMAGE_BRUSH( "Icons/icon_axis_16px", FVector2D( 16, 16 ) ) );
diff --git a/Engine/Source/Editor/LevelEditor/Private/LevelEditorMenu.cpp b/Engine/Source/Editor/LevelEditor/Private/LevelEditorMenu.cpp
--- a/Engine/Source/Editor/LevelEditor/Private/LevelEditorMenu.cpp
+++ b/Engine/Source/Editor/LevelEditor/Private/LevelEditorMenu.cpp
@@ -760,26 +760,9 @@ void FLevelEditorMenu::RegisterSelectMenu()
 	{
 		FToolMenuSection& UnnamedSection = Menu->FindOrAddSection(NAME_None);
 
-		UnnamedSection.AddMenuEntry(
-			FGenericCommands::Get().SelectAll,
-			FGenericCommands::Get().SelectAll->GetLabel(),
-			FGenericCommands::Get().SelectAll->GetDescription(),
-			FSlateIconFinder::FindIcon("FoliageEditMode.SelectAll")
-		);
-
-		UnnamedSection.AddMenuEntry(
-			FLevelEditorCommands::Get().SelectNone,
-			FLevelEditorCommands::Get().SelectNone->GetLabel(),
-			FLevelEditorCommands::Get().SelectNone->GetDescription(),
-			FSlateIconFinder::FindIcon("Cross")
-		);
-
-		UnnamedSection.AddMenuEntry(
-			FLevelEditorCommands::Get().InvertSelection,
-			FLevelEditorCommands::Get().InvertSelection->GetLabel(),
-			FLevelEditorCommands::Get().InvertSelection->GetDescription(),
-			FSlateIconFinder::FindIcon("FoliageEditMode.DeselectAll")
-		);
+		UnnamedSection.AddMenuEntry(FGenericCommands::Get().SelectAll);
+		UnnamedSection.AddMenuEntry(FLevelEditorCommands::Get().SelectNone);
+		UnnamedSection.AddMenuEntry(FLevelEditorCommands::Get().InvertSelection);
 
 		// Hierarchy based selection
 		{
@@ -794,15 +777,8 @@ void FLevelEditorMenu::RegisterSelectMenu()
 							"SelectAllHierarchy", LOCTEXT("SelectAllHierarchyLabel", "Hierarchy")
 						);
 
-						HierarchySection.AddMenuEntry(
-							FLevelEditorCommands::Get().SelectImmediateChildren,
-							LOCTEXT("HierarchySelectImmediateChildrenLabel", "Immediate Children")
-						);
-
-						HierarchySection.AddMenuEntry(
-							FLevelEditorCommands::Get().SelectAllDescendants,
-							LOCTEXT("HierarchySelectAllDescendantsLabel", "All Descendants")
-						);
+						HierarchySection.AddMenuEntry(FLevelEditorCommands::Get().SelectImmediateChildren);
+						HierarchySection.AddMenuEntry(FLevelEditorCommands::Get().SelectAllDescendants);
 					}
 				),
 				false,
diff --git a/Engine/Source/Editor/LevelEditor/Private/ViewportToolbar/LevelEditorSubmenus.cpp b/Engine/Source/Editor/LevelEditor/Private/ViewportToolbar/LevelEditorSubmenus.cpp
--- a/Engine/Source/Editor/LevelEditor/Private/ViewportToolbar/LevelEditorSubmenus.cpp
+++ b/Engine/Source/Editor/LevelEditor/Private/ViewportToolbar/LevelEditorSubmenus.cpp
@@ -512,7 +512,9 @@ FToolMenuEntry CreateActorSnapCheckboxMenu()
 		FLevelEditorCommands::Get().EnableActorSnap->GetDescription(),
 		MakeMenuDelegate,
 		CheckboxMenuAction,
-		EUserInterfaceActionType::ToggleButton
+		EUserInterfaceActionType::ToggleButton,
+		false,
+		FSlateIcon(FAppStyle::GetAppStyleSetName(), "LevelEditor.EnableActorSnap")
 	);
 
 	return ActorSnapping;
@@ -1718,7 +1720,9 @@ FToolMenuEntry CreateFeatureLevelPreviewSubmenu()
 					);
 				}
 			}
-		)
+		),
+		false,
+		FSlateIcon(FAppStyle::GetAppStyleSetName(), "LevelEditor.PreviewPlatform")
 	);
 }
 
@@ -1745,7 +1749,9 @@ FToolMenuEntry CreateMaterialQualityLevelSubmenu()
 				Section.AddMenuEntry(FLevelEditorCommands::Get().MaterialQualityLevel_High);
 				Section.AddMenuEntry(FLevelEditorCommands::Get().MaterialQualityLevel_Epic);
 			}
-		)
+		),
+		false,
+		FSlateIcon(FAppStyle::GetAppStyleSetName(), "LevelEditor.MaterialQuality")
 	);
 }
 
@@ -2548,7 +2554,7 @@ void ExtendCameraSubmenu(FName InCameraOptionsSubmenuName)
 								}
 							),
 							false,
-							FSlateIcon(FSlateIcon(FAppStyle::GetAppStyleSetName(), "EditorViewport.Speed"))
+							FSlateIcon(FSlateIcon(FAppStyle::GetAppStyleSetName(), "ViewportToolbar.CameraSpeed"))
 						);
 
 						Entry.ToolBarData.LabelOverride = TAttribute<FText>::CreateLambda(
@@ -2576,17 +2582,8 @@ void ExtendCameraSubmenu(FName InCameraOptionsSubmenuName)
 
 					// Move Camera/Object
 					{
-						FToolMenuEntry CameraToObjectMenu =
-							FToolMenuEntry::InitMenuEntry(FLevelEditorCommands::Get().SnapCameraToObject);
-						CameraToObjectMenu.UserInterfaceActionType = EUserInterfaceActionType::ToggleButton;
-						CameraToObjectMenu.Label = LOCTEXT("CameraToObjectLabel", "Move Camera to Object");
-						PositioningSection.AddEntry(CameraToObjectMenu);
-
-						FToolMenuEntry ObjectToCameraMenu =
-							FToolMenuEntry::InitMenuEntry(FLevelEditorCommands::Get().SnapObjectToCamera);
-						ObjectToCameraMenu.UserInterfaceActionType = EUserInterfaceActionType::ToggleButton;
-						ObjectToCameraMenu.Label = LOCTEXT("ObjectToCameraLabel", "Move Object to Camera");
-						PositioningSection.AddEntry(ObjectToCameraMenu);
+						PositioningSection.AddMenuEntry(FLevelEditorCommands::Get().SnapCameraToObject);
+						PositioningSection.AddMenuEntry(FLevelEditorCommands::Get().SnapObjectToCamera);
 					}
 				}
 
@@ -2635,30 +2632,17 @@ void ExtendTransformSubmenu(FName InTransformSubmenuName)
 	// Gizmo
 	{
 		FToolMenuSection& GizmoSection = Submenu->FindOrAddSection("Gizmo", LOCTEXT("GizmoLabel", "Gizmo"));
-		FToolMenuEntry& ShowTransformGizmoEntry = GizmoSection.AddMenuEntry(
-			FLevelEditorCommands::Get().ShowTransformWidget, LOCTEXT("ShowTransformGizmoLabel", "Show Transform Gizmo")
-		);
-
+		FToolMenuEntry& ShowTransformGizmoEntry = GizmoSection.AddMenuEntry(FLevelEditorCommands::Get().ShowTransformWidget);
 		ShowTransformGizmoEntry.InsertPosition.Position = EToolMenuInsertType::First;
 	}
 
 	// Selection
 	{
 		FToolMenuSection& SelectionSection = Submenu->FindOrAddSection("Selection", LOCTEXT("SelectionLabel", "Selection"));
-		SelectionSection.AddMenuEntry(
-			FLevelEditorCommands::Get().AllowTranslucentSelection,
-			LOCTEXT("SelectionAllowTranslucentSelectionLabel", "Translucent Objects")
-		);
-		SelectionSection.AddMenuEntry(
-			FLevelEditorCommands::Get().AllowGroupSelection, LOCTEXT("SelectionAllowGroupSelectionLabel", "Select Groups")
-		);
-		SelectionSection.AddMenuEntry(
-			FLevelEditorCommands::Get().StrictBoxSelect, LOCTEXT("SelectionStrictBoxSelectLabel", "Strict Marquee Selection")
-		);
-		SelectionSection.AddMenuEntry(
-			FLevelEditorCommands::Get().TransparentBoxSelect,
-			LOCTEXT("SelectionTransparentBoxSelectLabel", "Marquee Select Occluded")
-		);
+		SelectionSection.AddMenuEntry(FLevelEditorCommands::Get().AllowTranslucentSelection);
+		SelectionSection.AddMenuEntry(FLevelEditorCommands::Get().AllowGroupSelection);
+		SelectionSection.AddMenuEntry(FLevelEditorCommands::Get().StrictBoxSelect);
+		SelectionSection.AddMenuEntry(FLevelEditorCommands::Get().TransparentBoxSelect);
 	}
 }
 
@@ -2672,17 +2656,10 @@ void ExtendSnappingSubmenu(FName InSnappingSubmenuName)
 		SnappingSection.AddEntry(UE::LevelEditor::Private::CreateActorSnapCheckboxMenu());
 
 		// Socket Snapping
-		FToolMenuEntry SocketSnapping =
-			FToolMenuEntry::InitMenuEntry(FLevelEditorCommands::Get().ToggleSocketSnapping);
-		SocketSnapping.UserInterfaceActionType = EUserInterfaceActionType::ToggleButton;
-		SocketSnapping.Label = LOCTEXT("SocketSnapLabel", "Socket");
-		SnappingSection.AddEntry(SocketSnapping);
+		SnappingSection.AddMenuEntry(FLevelEditorCommands::Get().ToggleSocketSnapping);
 
 		// Vertex Snapping
-		FToolMenuEntry VertexSnapping = FToolMenuEntry::InitMenuEntry(FLevelEditorCommands::Get().EnableVertexSnap);
-		VertexSnapping.UserInterfaceActionType = EUserInterfaceActionType::ToggleButton;
-		VertexSnapping.Label = LOCTEXT("VertexSnapLabel", "Vertex");
-		SnappingSection.AddEntry(VertexSnapping);
+		SnappingSection.AddMenuEntry(FLevelEditorCommands::Get().EnableVertexSnap);
 	}
 }
 
diff --git a/Engine/Source/Editor/UnrealEd/Private/ViewportToolbar/UnrealEdViewportToolbar.cpp b/Engine/Source/Editor/UnrealEd/Private/ViewportToolbar/UnrealEdViewportToolbar.cpp
--- a/Engine/Source/Editor/UnrealEd/Private/ViewportToolbar/UnrealEdViewportToolbar.cpp
+++ b/Engine/Source/Editor/UnrealEd/Private/ViewportToolbar/UnrealEdViewportToolbar.cpp
@@ -199,7 +199,16 @@ FToolMenuEntry CreateSurfaceSnapCheckboxMenu()
 			return FText::AsNumber(Settings->SnapToSurface.SnapOffsetExtent);
 		}
 	);
-	Entry.Icon = FSlateIcon(FAppStyle::GetAppStyleSetName(), "EditorViewport.ToggleSurfaceSnapping");
+	Entry.Icon = TAttribute<FSlateIcon>::CreateLambda(
+		[]()
+		{
+			// todo: dynamic update works in real time for the icon in the raised entry, but not within the menu itself.
+			// In order to se the icon within the menu update, the menu needs to close and re-open
+			return GetDefault<ULevelEditorViewportSettings>()->SnapToSurface.bSnapRotation
+					 ? FSlateIcon(FAppStyle::GetAppStyleSetName(), "ViewportToolbar.SurfaceSnapRotateToNormal")
+					 : FSlateIcon(FAppStyle::GetAppStyleSetName(), "ViewportToolbar.SurfaceSnap");
+		}
+	);
 
 	return Entry;
 }
@@ -344,7 +353,7 @@ FToolMenuEntry CreateLocationSnapCheckboxMenu()
 	}
 
 	Entry.ToolBarData.LabelOverride = TAttribute<FText>::Create(&UE::UnrealEd::GetLocationGridLabel);
-	Entry.Icon = FSlateIcon(FAppStyle::GetAppStyleSetName(), "EditorViewport.LocationGridSnap");
+	Entry.Icon = FSlateIcon(FAppStyle::GetAppStyleSetName(), "ViewportToolbar.SnapLocation");
 
 	return Entry;
 }
@@ -415,7 +424,7 @@ FToolMenuEntry CreateRotationSnapCheckboxMenu()
 	);
 
 	Entry.ToolBarData.LabelOverride = TAttribute<FText>::Create(&UE::UnrealEd::GetRotationGridLabel);
-	Entry.Icon = FSlateIcon(FAppStyle::GetAppStyleSetName(), "EditorViewport.RotationGridSnap");
+	Entry.Icon = FSlateIcon(FAppStyle::GetAppStyleSetName(), "ViewportToolbar.SnapRotation");
 
 	return Entry;
 }
@@ -519,7 +528,7 @@ FToolMenuEntry CreateScaleSnapCheckboxMenu()
 	);
 
 	Entry.ToolBarData.LabelOverride = TAttribute<FText>::Create(&UE::UnrealEd::GetScaleGridLabel);
-	Entry.Icon = FSlateIcon(FAppStyle::GetAppStyleSetName(), "EditorViewport.ScaleGridSnap");
+	Entry.Icon = FSlateIcon(FAppStyle::GetAppStyleSetName(), "ViewportToolbar.SnapScale");
 
 	return Entry;
 }
@@ -852,9 +861,11 @@ FToolMenuEntry CreateTransformsSubmenu()
 						+ SHorizontalBox::Slot()
 						.FillWidth(0.1f);
 					// clang-format on
-					GizmoSection.AddEntry(FToolMenuEntry::InitWidget(
+
+					FToolMenuEntry& GizmoScaleEntry = GizmoSection.AddEntry(FToolMenuEntry::InitWidget(
 						"GizmoScale", GizmoScaleWidget, LOCTEXT("GizmoScaleLabel", "Gizmo Scale")
 					));
+					GizmoScaleEntry.Icon = FSlateIcon(FAppStyle::GetAppStyleSetName(), "ViewportToolbar.GizmoScale");
 				}
 			}
 		)
@@ -1181,7 +1192,7 @@ FToolMenuEntry CreateSnappingSubmenu()
 		)
 	);
 
-	Entry.Icon = FSlateIcon(FAppStyle::GetAppStyleSetName(), "Icons.Snap");
+	Entry.Icon = FSlateIcon(FAppStyle::GetAppStyleSetName(), "ViewportToolbar.Snap");
 	Entry.InsertPosition.Position = EToolMenuInsertType::First;
 	Entry.ToolBarData.LabelOverride = FText();
 	Entry.ToolBarData.ResizeParams.ClippingPriority = 900;
@@ -2626,7 +2637,18 @@ void PopulateCameraMenu(UToolMenu* InMenu)
 				EditorViewportContext, EHidableViewModeMenuSections::Exposure
 			))
 		{
-			AddExposureSection(InMenu, EditorViewport);
+			const FEditorViewportCommands& BaseViewportCommands = FEditorViewportCommands::Get();
+
+			TSharedRef<SWidget> FixedEV100Menu = EditorViewport->BuildFixedEV100Menu();
+			TSharedPtr<FEditorViewportClient> EditorViewportClient = EditorViewport->GetViewportClient();
+			const bool bIsLevelEditor = EditorViewportClient.IsValid() && EditorViewportClient->IsLevelEditorClient();
+
+			FToolMenuSection& Section = InMenu->AddSection("Exposure", LOCTEXT("ExposureHeader", "Exposure"));
+			Section.AddMenuEntry(
+				bIsLevelEditor ? BaseViewportCommands.ToggleInGameExposure : BaseViewportCommands.ToggleAutoExposure
+			);
+			FToolMenuEntry& EVEntry = Section.AddEntry(FToolMenuEntry::InitWidget("FixedEV100", FixedEV100Menu, LOCTEXT("FixedEV100", "EV100")));
+			EVEntry.Icon = FSlateIcon(FAppStyle::GetAppStyleSetName(), "ViewportToolbar.EV100");
 		}
 	}
 }
@@ -3116,7 +3138,7 @@ FToolMenuEntry CreateScalabilitySubmenu()
 					FToolUIAction(),
 					EUserInterfaceActionType::Button,
 					false,
-					 FSlateIcon(FAppStyle::GetAppStyleSetName(), "Icons.ViewportScalability")
+					FSlateIcon(FAppStyle::GetAppStyleSetName(), "Icons.ViewportScalability")
 				);
 
 				ScalabilitySubmenu.SetShowInToolbarTopLevel(TAttribute<bool>::Create(&UE::UnrealEd::IsScalabilityWarningVisible
diff --git a/Engine/Source/Runtime/SlateCore/Private/Styling/StarshipCoreStyle.cpp b/Engine/Source/Runtime/SlateCore/Private/Styling/StarshipCoreStyle.cpp
--- a/Engine/Source/Runtime/SlateCore/Private/Styling/StarshipCoreStyle.cpp
+++ b/Engine/Source/Runtime/SlateCore/Private/Styling/StarshipCoreStyle.cpp
@@ -456,8 +456,10 @@ TSharedRef<ISlateStyle> FStarshipCoreStyle::Create()
 		Style->Set("GenericCommands.Delete", new IMAGE_BRUSH_SVG("Starship/Common/Delete", Icon16x16));
 		Style->Set("GenericCommands.Paste", new IMAGE_BRUSH_SVG("Starship/Common/Paste", Icon16x16));
 		Style->Set("GenericCommands.Duplicate", new IMAGE_BRUSH_SVG("Starship/Common/Duplicate", Icon16x16));
-		
+
 		Style->Set("GenericCommands.Rename", new IMAGE_BRUSH_SVG("Starship/Common/Rename", Icon16x16));
+
+		Style->Set("GenericCommands.SelectAll", new IMAGE_BRUSH_SVG("Starship/Common/SelectAll_16", Icon16x16));
 	}
 
 	// SVerticalBox Drag& Drop icon
