[Viewport Toolbar] Polish: Update layout icons in the "..." menu of the new level editor viewport toolbar

diff --git a/Engine/Source/Editor/EditorStyle/Private/SlateEditorStyle.cpp b/Engine/Source/Editor/EditorStyle/Private/SlateEditorStyle.cpp
--- a/Engine/Source/Editor/EditorStyle/Private/SlateEditorStyle.cpp
+++ b/Engine/Source/Editor/EditorStyle/Private/SlateEditorStyle.cpp
@@ -5019,30 +5019,30 @@ void FSlateEditorStyle::FStyle::SetupLevelEditorStyle()
 		const FVector2D IconLayoutSize(47.0f, 37.0f);
 		const FVector2D IconLayoutSizeSmall(47.0f, 37.0f);		// small version set to same size as these are in their own menu and don't clutter the UI
 
-		Set("EditorViewport.ViewportConfig_OnePane", new IMAGE_BRUSH("Icons/ViewportLayout_OnePane", IconLayoutSize));
-		Set("EditorViewport.ViewportConfig_OnePane.Small", new IMAGE_BRUSH("Icons/ViewportLayout_OnePane", IconLayoutSizeSmall));
-		Set("EditorViewport.ViewportConfig_TwoPanesH", new IMAGE_BRUSH("Icons/ViewportLayout_TwoPanesHoriz", IconLayoutSize));
-		Set("EditorViewport.ViewportConfig_TwoPanesH.Small", new IMAGE_BRUSH("Icons/ViewportLayout_TwoPanesHoriz", IconLayoutSizeSmall));
-		Set("EditorViewport.ViewportConfig_TwoPanesV", new IMAGE_BRUSH("Icons/ViewportLayout_TwoPanesVert", IconLayoutSize));
-		Set("EditorViewport.ViewportConfig_TwoPanesV.Small", new IMAGE_BRUSH("Icons/ViewportLayout_TwoPanesVert", IconLayoutSizeSmall));
-		Set("EditorViewport.ViewportConfig_ThreePanesLeft", new IMAGE_BRUSH("Icons/ViewportLayout_ThreePanesLeft", IconLayoutSize));
-		Set("EditorViewport.ViewportConfig_ThreePanesLeft.Small", new IMAGE_BRUSH("Icons/ViewportLayout_ThreePanesLeft", IconLayoutSizeSmall));
-		Set("EditorViewport.ViewportConfig_ThreePanesRight", new IMAGE_BRUSH("Icons/ViewportLayout_ThreePanesRight", IconLayoutSize));
-		Set("EditorViewport.ViewportConfig_ThreePanesRight.Small", new IMAGE_BRUSH("Icons/ViewportLayout_ThreePanesRight", IconLayoutSizeSmall));
-		Set("EditorViewport.ViewportConfig_ThreePanesTop", new IMAGE_BRUSH("Icons/ViewportLayout_ThreePanesTop", IconLayoutSize));
-		Set("EditorViewport.ViewportConfig_ThreePanesTop.Small", new IMAGE_BRUSH("Icons/ViewportLayout_ThreePanesTop", IconLayoutSizeSmall));
-		Set("EditorViewport.ViewportConfig_ThreePanesBottom", new IMAGE_BRUSH("Icons/ViewportLayout_ThreePanesBottom", IconLayoutSize));
-		Set("EditorViewport.ViewportConfig_ThreePanesBottom.Small", new IMAGE_BRUSH("Icons/ViewportLayout_ThreePanesBottom", IconLayoutSizeSmall));
-		Set("EditorViewport.ViewportConfig_FourPanesLeft", new IMAGE_BRUSH("Icons/ViewportLayout_FourPanesLeft", IconLayoutSize));
-		Set("EditorViewport.ViewportConfig_FourPanesLeft.Small", new IMAGE_BRUSH("Icons/ViewportLayout_FourPanesLeft", IconLayoutSizeSmall));
-		Set("EditorViewport.ViewportConfig_FourPanesRight", new IMAGE_BRUSH("Icons/ViewportLayout_FourPanesRight", IconLayoutSize));
-		Set("EditorViewport.ViewportConfig_FourPanesRight.Small", new IMAGE_BRUSH("Icons/ViewportLayout_FourPanesRight", IconLayoutSizeSmall));
-		Set("EditorViewport.ViewportConfig_FourPanesTop", new IMAGE_BRUSH("Icons/ViewportLayout_FourPanesTop", IconLayoutSize));
-		Set("EditorViewport.ViewportConfig_FourPanesTop.Small", new IMAGE_BRUSH("Icons/ViewportLayout_FourPanesTop", IconLayoutSizeSmall));
-		Set("EditorViewport.ViewportConfig_FourPanesBottom", new IMAGE_BRUSH("Icons/ViewportLayout_FourPanesBottom", IconLayoutSize));
-		Set("EditorViewport.ViewportConfig_FourPanesBottom.Small", new IMAGE_BRUSH("Icons/ViewportLayout_FourPanesBottom", IconLayoutSizeSmall));
-		Set("EditorViewport.ViewportConfig_FourPanes2x2", new IMAGE_BRUSH("Icons/ViewportLayout_FourPanes2x2", IconLayoutSize));
-		Set("EditorViewport.ViewportConfig_FourPanes2x2.Small", new IMAGE_BRUSH("Icons/ViewportLayout_FourPanes2x2", IconLayoutSizeSmall));
+		Set("EditorViewport.ViewportConfig_OnePane", new IMAGE_BRUSH_SVG("Starship/ViewportLayout/Layouts1Pane", IconLayoutSize));
+		Set("EditorViewport.ViewportConfig_OnePane.Small", new IMAGE_BRUSH_SVG("Starship/ViewportLayout/Layouts1Pane", IconLayoutSizeSmall));
+		Set("EditorViewport.ViewportConfig_TwoPanesH", new IMAGE_BRUSH_SVG("Starship/ViewportLayout/Layouts2Panes", IconLayoutSize));
+		Set("EditorViewport.ViewportConfig_TwoPanesH.Small", new IMAGE_BRUSH_SVG("Starship/ViewportLayout/Layouts2Panes", IconLayoutSizeSmall));
+		Set("EditorViewport.ViewportConfig_TwoPanesV", new IMAGE_BRUSH_SVG("Starship/ViewportLayout/Layouts2PanesStacked", IconLayoutSize));
+		Set("EditorViewport.ViewportConfig_TwoPanesV.Small", new IMAGE_BRUSH_SVG("Starship/ViewportLayout/Layouts2PanesStacked", IconLayoutSizeSmall));
+		Set("EditorViewport.ViewportConfig_ThreePanesLeft", new IMAGE_BRUSH_SVG("Starship/ViewportLayout/Layouts3PanesLeft", IconLayoutSize));
+		Set("EditorViewport.ViewportConfig_ThreePanesLeft.Small", new IMAGE_BRUSH_SVG("Starship/ViewportLayout/Layouts3PanesLeft", IconLayoutSizeSmall));
+		Set("EditorViewport.ViewportConfig_ThreePanesRight", new IMAGE_BRUSH_SVG("Starship/ViewportLayout/Layouts3PanesRight", IconLayoutSize));
+		Set("EditorViewport.ViewportConfig_ThreePanesRight.Small", new IMAGE_BRUSH_SVG("Starship/ViewportLayout/Layouts3PanesRight", IconLayoutSizeSmall));
+		Set("EditorViewport.ViewportConfig_ThreePanesTop", new IMAGE_BRUSH_SVG("Starship/ViewportLayout/Layouts3PanesTop", IconLayoutSize));
+		Set("EditorViewport.ViewportConfig_ThreePanesTop.Small", new IMAGE_BRUSH_SVG("Starship/ViewportLayout/Layouts3PanesTop", IconLayoutSizeSmall));
+		Set("EditorViewport.ViewportConfig_ThreePanesBottom", new IMAGE_BRUSH_SVG("Starship/ViewportLayout/Layouts3PanesBottom", IconLayoutSize));
+		Set("EditorViewport.ViewportConfig_ThreePanesBottom.Small", new IMAGE_BRUSH_SVG("Starship/ViewportLayout/Layouts3PanesBottom", IconLayoutSizeSmall));
+		Set("EditorViewport.ViewportConfig_FourPanesLeft", new IMAGE_BRUSH_SVG("Starship/ViewportLayout/Layouts4PanesLeft", IconLayoutSize));
+		Set("EditorViewport.ViewportConfig_FourPanesLeft.Small", new IMAGE_BRUSH_SVG("Starship/ViewportLayout/Layouts4PanesLeft", IconLayoutSizeSmall));
+		Set("EditorViewport.ViewportConfig_FourPanesRight", new IMAGE_BRUSH_SVG("Starship/ViewportLayout/Layouts4PanesRight", IconLayoutSize));
+		Set("EditorViewport.ViewportConfig_FourPanesRight.Small", new IMAGE_BRUSH_SVG("Starship/ViewportLayout/Layouts4PanesRight", IconLayoutSizeSmall));
+		Set("EditorViewport.ViewportConfig_FourPanesTop", new IMAGE_BRUSH_SVG("Starship/ViewportLayout/Layouts4PanesTop", IconLayoutSize));
+		Set("EditorViewport.ViewportConfig_FourPanesTop.Small", new IMAGE_BRUSH_SVG("Starship/ViewportLayout/Layouts4PanesTop", IconLayoutSizeSmall));
+		Set("EditorViewport.ViewportConfig_FourPanesBottom", new IMAGE_BRUSH_SVG("Starship/ViewportLayout/Layouts4PanesBottom", IconLayoutSize));
+		Set("EditorViewport.ViewportConfig_FourPanesBottom.Small", new IMAGE_BRUSH_SVG("Starship/ViewportLayout/Layouts4PanesBottom", IconLayoutSizeSmall));
+		Set("EditorViewport.ViewportConfig_FourPanes2x2", new IMAGE_BRUSH_SVG("Starship/ViewportLayout/Layouts4Panes", IconLayoutSize));
+		Set("EditorViewport.ViewportConfig_FourPanes2x2.Small", new IMAGE_BRUSH_SVG("Starship/ViewportLayout/Layouts4Panes", IconLayoutSizeSmall));
 	}
 
 	// Level viewport layout command icons
@@ -5050,30 +5050,30 @@ void FSlateEditorStyle::FStyle::SetupLevelEditorStyle()
 		const FVector2D IconLayoutSize(47.0f, 37.0f);
 		const FVector2D IconLayoutSizeSmall(47.0f, 37.0f);		// small version set to same size as these are in their own menu and don't clutter the UI
 
-		Set( "LevelViewport.ViewportConfig_OnePane", new IMAGE_BRUSH("Icons/ViewportLayout_OnePane", IconLayoutSize) );
-		Set( "LevelViewport.ViewportConfig_OnePane.Small", new IMAGE_BRUSH("Icons/ViewportLayout_OnePane", IconLayoutSizeSmall) );
-		Set( "LevelViewport.ViewportConfig_TwoPanesH", new IMAGE_BRUSH( "Icons/ViewportLayout_TwoPanesHoriz", IconLayoutSize ) );
-		Set( "LevelViewport.ViewportConfig_TwoPanesH.Small", new IMAGE_BRUSH( "Icons/ViewportLayout_TwoPanesHoriz", IconLayoutSizeSmall ) );
-		Set( "LevelViewport.ViewportConfig_TwoPanesV", new IMAGE_BRUSH( "Icons/ViewportLayout_TwoPanesVert", IconLayoutSize ) );
-		Set( "LevelViewport.ViewportConfig_TwoPanesV.Small", new IMAGE_BRUSH( "Icons/ViewportLayout_TwoPanesVert", IconLayoutSizeSmall ) );
-		Set( "LevelViewport.ViewportConfig_ThreePanesLeft", new IMAGE_BRUSH( "Icons/ViewportLayout_ThreePanesLeft", IconLayoutSize ) );
-		Set( "LevelViewport.ViewportConfig_ThreePanesLeft.Small", new IMAGE_BRUSH( "Icons/ViewportLayout_ThreePanesLeft", IconLayoutSizeSmall ) );
-		Set( "LevelViewport.ViewportConfig_ThreePanesRight", new IMAGE_BRUSH( "Icons/ViewportLayout_ThreePanesRight", IconLayoutSize ) );
-		Set( "LevelViewport.ViewportConfig_ThreePanesRight.Small", new IMAGE_BRUSH( "Icons/ViewportLayout_ThreePanesRight", IconLayoutSizeSmall ) );
-		Set( "LevelViewport.ViewportConfig_ThreePanesTop", new IMAGE_BRUSH( "Icons/ViewportLayout_ThreePanesTop", IconLayoutSize ) );
-		Set( "LevelViewport.ViewportConfig_ThreePanesTop.Small", new IMAGE_BRUSH( "Icons/ViewportLayout_ThreePanesTop", IconLayoutSizeSmall ) );
-		Set( "LevelViewport.ViewportConfig_ThreePanesBottom", new IMAGE_BRUSH( "Icons/ViewportLayout_ThreePanesBottom", IconLayoutSize ) );
-		Set( "LevelViewport.ViewportConfig_ThreePanesBottom.Small", new IMAGE_BRUSH( "Icons/ViewportLayout_ThreePanesBottom", IconLayoutSizeSmall ) );
-		Set( "LevelViewport.ViewportConfig_FourPanesLeft", new IMAGE_BRUSH( "Icons/ViewportLayout_FourPanesLeft", IconLayoutSize ) );
-		Set( "LevelViewport.ViewportConfig_FourPanesLeft.Small", new IMAGE_BRUSH( "Icons/ViewportLayout_FourPanesLeft", IconLayoutSizeSmall ) );
-		Set( "LevelViewport.ViewportConfig_FourPanesRight", new IMAGE_BRUSH( "Icons/ViewportLayout_FourPanesRight", IconLayoutSize ) );
-		Set( "LevelViewport.ViewportConfig_FourPanesRight.Small", new IMAGE_BRUSH( "Icons/ViewportLayout_FourPanesRight", IconLayoutSizeSmall ) );
-		Set( "LevelViewport.ViewportConfig_FourPanesTop", new IMAGE_BRUSH( "Icons/ViewportLayout_FourPanesTop", IconLayoutSize ) );
-		Set( "LevelViewport.ViewportConfig_FourPanesTop.Small", new IMAGE_BRUSH( "Icons/ViewportLayout_FourPanesTop", IconLayoutSizeSmall ) );
-		Set( "LevelViewport.ViewportConfig_FourPanesBottom", new IMAGE_BRUSH( "Icons/ViewportLayout_FourPanesBottom", IconLayoutSize ) );
-		Set( "LevelViewport.ViewportConfig_FourPanesBottom.Small", new IMAGE_BRUSH( "Icons/ViewportLayout_FourPanesBottom", IconLayoutSizeSmall ) );
-		Set( "LevelViewport.ViewportConfig_FourPanes2x2", new IMAGE_BRUSH( "Icons/ViewportLayout_FourPanes2x2", IconLayoutSize ) );
-		Set( "LevelViewport.ViewportConfig_FourPanes2x2.Small", new IMAGE_BRUSH( "Icons/ViewportLayout_FourPanes2x2", IconLayoutSizeSmall ) );
+		Set("LevelViewport.ViewportConfig_OnePane", new IMAGE_BRUSH_SVG("Starship/ViewportLayout/Layouts1Pane", IconLayoutSize));
+		Set("LevelViewport.ViewportConfig_OnePane.Small", new IMAGE_BRUSH_SVG("Starship/ViewportLayout/Layouts1Pane", IconLayoutSizeSmall));
+		Set("LevelViewport.ViewportConfig_TwoPanesH", new IMAGE_BRUSH_SVG("Starship/ViewportLayout/Layouts2Panes", IconLayoutSize));
+		Set("LevelViewport.ViewportConfig_TwoPanesH.Small", new IMAGE_BRUSH_SVG("Starship/ViewportLayout/Layouts2Panes", IconLayoutSizeSmall));
+		Set("LevelViewport.ViewportConfig_TwoPanesV", new IMAGE_BRUSH_SVG("Starship/ViewportLayout/Layouts2PanesStacked", IconLayoutSize));
+		Set("LevelViewport.ViewportConfig_TwoPanesV.Small", new IMAGE_BRUSH_SVG("Starship/ViewportLayout/Layouts2PanesStacked", IconLayoutSizeSmall));
+		Set("LevelViewport.ViewportConfig_ThreePanesLeft", new IMAGE_BRUSH_SVG("Starship/ViewportLayout/Layouts3PanesLeft", IconLayoutSize));
+		Set("LevelViewport.ViewportConfig_ThreePanesLeft.Small", new IMAGE_BRUSH_SVG("Starship/ViewportLayout/Layouts3PanesLeft", IconLayoutSizeSmall));
+		Set("LevelViewport.ViewportConfig_ThreePanesRight", new IMAGE_BRUSH_SVG("Starship/ViewportLayout/Layouts3PanesRight", IconLayoutSize));
+		Set("LevelViewport.ViewportConfig_ThreePanesRight.Small", new IMAGE_BRUSH_SVG("Starship/ViewportLayout/Layouts3PanesRight", IconLayoutSizeSmall));
+		Set("LevelViewport.ViewportConfig_ThreePanesTop", new IMAGE_BRUSH_SVG("Starship/ViewportLayout/Layouts3PanesTop", IconLayoutSize));
+		Set("LevelViewport.ViewportConfig_ThreePanesTop.Small", new IMAGE_BRUSH_SVG("Starship/ViewportLayout/Layouts3PanesTop", IconLayoutSizeSmall));
+		Set("LevelViewport.ViewportConfig_ThreePanesBottom", new IMAGE_BRUSH_SVG("Starship/ViewportLayout/Layouts3PanesBottom", IconLayoutSize));
+		Set("LevelViewport.ViewportConfig_ThreePanesBottom.Small", new IMAGE_BRUSH_SVG("Starship/ViewportLayout/Layouts3PanesBottom", IconLayoutSizeSmall));
+		Set("LevelViewport.ViewportConfig_FourPanesLeft", new IMAGE_BRUSH_SVG("Starship/ViewportLayout/Layouts4PanesLeft", IconLayoutSize));
+		Set("LevelViewport.ViewportConfig_FourPanesLeft.Small", new IMAGE_BRUSH_SVG("Starship/ViewportLayout/Layouts4PanesLeft", IconLayoutSizeSmall));
+		Set("LevelViewport.ViewportConfig_FourPanesRight", new IMAGE_BRUSH_SVG("Starship/ViewportLayout/Layouts4PanesRight", IconLayoutSize));
+		Set("LevelViewport.ViewportConfig_FourPanesRight.Small", new IMAGE_BRUSH_SVG("Starship/ViewportLayout/Layouts4PanesRight", IconLayoutSizeSmall));
+		Set("LevelViewport.ViewportConfig_FourPanesTop", new IMAGE_BRUSH_SVG("Starship/ViewportLayout/Layouts4PanesTop", IconLayoutSize));
+		Set("LevelViewport.ViewportConfig_FourPanesTop.Small", new IMAGE_BRUSH_SVG("Starship/ViewportLayout/Layouts4PanesTop", IconLayoutSizeSmall));
+		Set("LevelViewport.ViewportConfig_FourPanesBottom", new IMAGE_BRUSH_SVG("Starship/ViewportLayout/Layouts4PanesBottom", IconLayoutSize));
+		Set("LevelViewport.ViewportConfig_FourPanesBottom.Small", new IMAGE_BRUSH_SVG("Starship/ViewportLayout/Layouts4PanesBottom", IconLayoutSizeSmall));
+		Set("LevelViewport.ViewportConfig_FourPanes2x2", new IMAGE_BRUSH_SVG("Starship/ViewportLayout/Layouts4Panes", IconLayoutSize));
+		Set("LevelViewport.ViewportConfig_FourPanes2x2.Small", new IMAGE_BRUSH_SVG("Starship/ViewportLayout/Layouts4Panes", IconLayoutSizeSmall));
 	}
 
 	// Level editor status bar
diff --git a/Engine/Source/Editor/EditorStyle/Private/StarshipStyle.cpp b/Engine/Source/Editor/EditorStyle/Private/StarshipStyle.cpp
--- a/Engine/Source/Editor/EditorStyle/Private/StarshipStyle.cpp
+++ b/Engine/Source/Editor/EditorStyle/Private/StarshipStyle.cpp
@@ -4971,30 +4971,30 @@ void FStarshipEditorStyle::FStyle::SetupLevelEditorStyle()
 		const FVector2D IconLayoutSize(47.0f, 37.0f);
 		const FVector2D IconLayoutSizeSmall(47.0f, 37.0f);		// small version set to same size as these are in their own menu and don't clutter the UI
 
-		Set("EditorViewport.ViewportConfig_OnePane", new IMAGE_BRUSH("Icons/ViewportLayout_OnePane", IconLayoutSize));
-		Set("EditorViewport.ViewportConfig_OnePane.Small", new IMAGE_BRUSH("Icons/ViewportLayout_OnePane", IconLayoutSizeSmall));
-		Set("EditorViewport.ViewportConfig_TwoPanesH", new IMAGE_BRUSH("Icons/ViewportLayout_TwoPanesHoriz", IconLayoutSize));
-		Set("EditorViewport.ViewportConfig_TwoPanesH.Small", new IMAGE_BRUSH("Icons/ViewportLayout_TwoPanesHoriz", IconLayoutSizeSmall));
-		Set("EditorViewport.ViewportConfig_TwoPanesV", new IMAGE_BRUSH("Icons/ViewportLayout_TwoPanesVert", IconLayoutSize));
-		Set("EditorViewport.ViewportConfig_TwoPanesV.Small", new IMAGE_BRUSH("Icons/ViewportLayout_TwoPanesVert", IconLayoutSizeSmall));
-		Set("EditorViewport.ViewportConfig_ThreePanesLeft", new IMAGE_BRUSH("Icons/ViewportLayout_ThreePanesLeft", IconLayoutSize));
-		Set("EditorViewport.ViewportConfig_ThreePanesLeft.Small", new IMAGE_BRUSH("Icons/ViewportLayout_ThreePanesLeft", IconLayoutSizeSmall));
-		Set("EditorViewport.ViewportConfig_ThreePanesRight", new IMAGE_BRUSH("Icons/ViewportLayout_ThreePanesRight", IconLayoutSize));
-		Set("EditorViewport.ViewportConfig_ThreePanesRight.Small", new IMAGE_BRUSH("Icons/ViewportLayout_ThreePanesRight", IconLayoutSizeSmall));
-		Set("EditorViewport.ViewportConfig_ThreePanesTop", new IMAGE_BRUSH("Icons/ViewportLayout_ThreePanesTop", IconLayoutSize));
-		Set("EditorViewport.ViewportConfig_ThreePanesTop.Small", new IMAGE_BRUSH("Icons/ViewportLayout_ThreePanesTop", IconLayoutSizeSmall));
-		Set("EditorViewport.ViewportConfig_ThreePanesBottom", new IMAGE_BRUSH("Icons/ViewportLayout_ThreePanesBottom", IconLayoutSize));
-		Set("EditorViewport.ViewportConfig_ThreePanesBottom.Small", new IMAGE_BRUSH("Icons/ViewportLayout_ThreePanesBottom", IconLayoutSizeSmall));
-		Set("EditorViewport.ViewportConfig_FourPanesLeft", new IMAGE_BRUSH("Icons/ViewportLayout_FourPanesLeft", IconLayoutSize));
-		Set("EditorViewport.ViewportConfig_FourPanesLeft.Small", new IMAGE_BRUSH("Icons/ViewportLayout_FourPanesLeft", IconLayoutSizeSmall));
-		Set("EditorViewport.ViewportConfig_FourPanesRight", new IMAGE_BRUSH("Icons/ViewportLayout_FourPanesRight", IconLayoutSize));
-		Set("EditorViewport.ViewportConfig_FourPanesRight.Small", new IMAGE_BRUSH("Icons/ViewportLayout_FourPanesRight", IconLayoutSizeSmall));
-		Set("EditorViewport.ViewportConfig_FourPanesTop", new IMAGE_BRUSH("Icons/ViewportLayout_FourPanesTop", IconLayoutSize));
-		Set("EditorViewport.ViewportConfig_FourPanesTop.Small", new IMAGE_BRUSH("Icons/ViewportLayout_FourPanesTop", IconLayoutSizeSmall));
-		Set("EditorViewport.ViewportConfig_FourPanesBottom", new IMAGE_BRUSH("Icons/ViewportLayout_FourPanesBottom", IconLayoutSize));
-		Set("EditorViewport.ViewportConfig_FourPanesBottom.Small", new IMAGE_BRUSH("Icons/ViewportLayout_FourPanesBottom", IconLayoutSizeSmall));
-		Set("EditorViewport.ViewportConfig_FourPanes2x2", new IMAGE_BRUSH("Icons/ViewportLayout_FourPanes2x2", IconLayoutSize));
-		Set("EditorViewport.ViewportConfig_FourPanes2x2.Small", new IMAGE_BRUSH("Icons/ViewportLayout_FourPanes2x2", IconLayoutSizeSmall));
+		Set("EditorViewport.ViewportConfig_OnePane", new IMAGE_BRUSH_SVG("Starship/ViewportLayout/Layouts1Pane", IconLayoutSize));
+		Set("EditorViewport.ViewportConfig_OnePane.Small", new IMAGE_BRUSH_SVG("Starship/ViewportLayout/Layouts1Pane", IconLayoutSizeSmall));
+		Set("EditorViewport.ViewportConfig_TwoPanesH", new IMAGE_BRUSH_SVG("Starship/ViewportLayout/Layouts2Panes", IconLayoutSize));
+		Set("EditorViewport.ViewportConfig_TwoPanesH.Small", new IMAGE_BRUSH_SVG("Starship/ViewportLayout/Layouts2Panes", IconLayoutSizeSmall));
+		Set("EditorViewport.ViewportConfig_TwoPanesV", new IMAGE_BRUSH_SVG("Starship/ViewportLayout/Layouts2PanesStacked", IconLayoutSize));
+		Set("EditorViewport.ViewportConfig_TwoPanesV.Small", new IMAGE_BRUSH_SVG("Starship/ViewportLayout/Layouts2PanesStacked", IconLayoutSizeSmall));
+		Set("EditorViewport.ViewportConfig_ThreePanesLeft", new IMAGE_BRUSH_SVG("Starship/ViewportLayout/Layouts3PanesLeft", IconLayoutSize));
+		Set("EditorViewport.ViewportConfig_ThreePanesLeft.Small", new IMAGE_BRUSH_SVG("Starship/ViewportLayout/Layouts3PanesLeft", IconLayoutSizeSmall));
+		Set("EditorViewport.ViewportConfig_ThreePanesRight", new IMAGE_BRUSH_SVG("Starship/ViewportLayout/Layouts3PanesRight", IconLayoutSize));
+		Set("EditorViewport.ViewportConfig_ThreePanesRight.Small", new IMAGE_BRUSH_SVG("Starship/ViewportLayout/Layouts3PanesRight", IconLayoutSizeSmall));
+		Set("EditorViewport.ViewportConfig_ThreePanesTop", new IMAGE_BRUSH_SVG("Starship/ViewportLayout/Layouts3PanesTop", IconLayoutSize));
+		Set("EditorViewport.ViewportConfig_ThreePanesTop.Small", new IMAGE_BRUSH_SVG("Starship/ViewportLayout/Layouts3PanesTop", IconLayoutSizeSmall));
+		Set("EditorViewport.ViewportConfig_ThreePanesBottom", new IMAGE_BRUSH_SVG("Starship/ViewportLayout/Layouts3PanesBottom", IconLayoutSize));
+		Set("EditorViewport.ViewportConfig_ThreePanesBottom.Small", new IMAGE_BRUSH_SVG("Starship/ViewportLayout/Layouts3PanesBottom", IconLayoutSizeSmall));
+		Set("EditorViewport.ViewportConfig_FourPanesLeft", new IMAGE_BRUSH_SVG("Starship/ViewportLayout/Layouts4PanesLeft", IconLayoutSize));
+		Set("EditorViewport.ViewportConfig_FourPanesLeft.Small", new IMAGE_BRUSH_SVG("Starship/ViewportLayout/Layouts4PanesLeft", IconLayoutSizeSmall));
+		Set("EditorViewport.ViewportConfig_FourPanesRight", new IMAGE_BRUSH_SVG("Starship/ViewportLayout/Layouts4PanesRight", IconLayoutSize));
+		Set("EditorViewport.ViewportConfig_FourPanesRight.Small", new IMAGE_BRUSH_SVG("Starship/ViewportLayout/Layouts4PanesRight", IconLayoutSizeSmall));
+		Set("EditorViewport.ViewportConfig_FourPanesTop", new IMAGE_BRUSH_SVG("Starship/ViewportLayout/Layouts4PanesTop", IconLayoutSize));
+		Set("EditorViewport.ViewportConfig_FourPanesTop.Small", new IMAGE_BRUSH_SVG("Starship/ViewportLayout/Layouts4PanesTop", IconLayoutSizeSmall));
+		Set("EditorViewport.ViewportConfig_FourPanesBottom", new IMAGE_BRUSH_SVG("Starship/ViewportLayout/Layouts4PanesBottom", IconLayoutSize));
+		Set("EditorViewport.ViewportConfig_FourPanesBottom.Small", new IMAGE_BRUSH_SVG("Starship/ViewportLayout/Layouts4PanesBottom", IconLayoutSizeSmall));
+		Set("EditorViewport.ViewportConfig_FourPanes2x2", new IMAGE_BRUSH_SVG("Starship/ViewportLayout/Layouts4Panes", IconLayoutSize));
+		Set("EditorViewport.ViewportConfig_FourPanes2x2.Small", new IMAGE_BRUSH_SVG("Starship/ViewportLayout/Layouts4Panes", IconLayoutSizeSmall));
 	}
 
 	// Level viewport layout command icons
@@ -5002,30 +5002,30 @@ void FStarshipEditorStyle::FStyle::SetupLevelEditorStyle()
 		const FVector2D IconLayoutSize(47.0f, 37.0f);
 		const FVector2D IconLayoutSizeSmall(47.0f, 37.0f);		// small version set to same size as these are in their own menu and don't clutter the UI
 
-		Set( "LevelViewport.ViewportConfig_OnePane", new IMAGE_BRUSH("Icons/ViewportLayout_OnePane", IconLayoutSize) );
-		Set( "LevelViewport.ViewportConfig_OnePane.Small", new IMAGE_BRUSH("Icons/ViewportLayout_OnePane", IconLayoutSizeSmall) );
-		Set( "LevelViewport.ViewportConfig_TwoPanesH", new IMAGE_BRUSH( "Icons/ViewportLayout_TwoPanesHoriz", IconLayoutSize ) );
-		Set( "LevelViewport.ViewportConfig_TwoPanesH.Small", new IMAGE_BRUSH( "Icons/ViewportLayout_TwoPanesHoriz", IconLayoutSizeSmall ) );
-		Set( "LevelViewport.ViewportConfig_TwoPanesV", new IMAGE_BRUSH( "Icons/ViewportLayout_TwoPanesVert", IconLayoutSize ) );
-		Set( "LevelViewport.ViewportConfig_TwoPanesV.Small", new IMAGE_BRUSH( "Icons/ViewportLayout_TwoPanesVert", IconLayoutSizeSmall ) );
-		Set( "LevelViewport.ViewportConfig_ThreePanesLeft", new IMAGE_BRUSH( "Icons/ViewportLayout_ThreePanesLeft", IconLayoutSize ) );
-		Set( "LevelViewport.ViewportConfig_ThreePanesLeft.Small", new IMAGE_BRUSH( "Icons/ViewportLayout_ThreePanesLeft", IconLayoutSizeSmall ) );
-		Set( "LevelViewport.ViewportConfig_ThreePanesRight", new IMAGE_BRUSH( "Icons/ViewportLayout_ThreePanesRight", IconLayoutSize ) );
-		Set( "LevelViewport.ViewportConfig_ThreePanesRight.Small", new IMAGE_BRUSH( "Icons/ViewportLayout_ThreePanesRight", IconLayoutSizeSmall ) );
-		Set( "LevelViewport.ViewportConfig_ThreePanesTop", new IMAGE_BRUSH( "Icons/ViewportLayout_ThreePanesTop", IconLayoutSize ) );
-		Set( "LevelViewport.ViewportConfig_ThreePanesTop.Small", new IMAGE_BRUSH( "Icons/ViewportLayout_ThreePanesTop", IconLayoutSizeSmall ) );
-		Set( "LevelViewport.ViewportConfig_ThreePanesBottom", new IMAGE_BRUSH( "Icons/ViewportLayout_ThreePanesBottom", IconLayoutSize ) );
-		Set( "LevelViewport.ViewportConfig_ThreePanesBottom.Small", new IMAGE_BRUSH( "Icons/ViewportLayout_ThreePanesBottom", IconLayoutSizeSmall ) );
-		Set( "LevelViewport.ViewportConfig_FourPanesLeft", new IMAGE_BRUSH( "Icons/ViewportLayout_FourPanesLeft", IconLayoutSize ) );
-		Set( "LevelViewport.ViewportConfig_FourPanesLeft.Small", new IMAGE_BRUSH( "Icons/ViewportLayout_FourPanesLeft", IconLayoutSizeSmall ) );
-		Set( "LevelViewport.ViewportConfig_FourPanesRight", new IMAGE_BRUSH( "Icons/ViewportLayout_FourPanesRight", IconLayoutSize ) );
-		Set( "LevelViewport.ViewportConfig_FourPanesRight.Small", new IMAGE_BRUSH( "Icons/ViewportLayout_FourPanesRight", IconLayoutSizeSmall ) );
-		Set( "LevelViewport.ViewportConfig_FourPanesTop", new IMAGE_BRUSH( "Icons/ViewportLayout_FourPanesTop", IconLayoutSize ) );
-		Set( "LevelViewport.ViewportConfig_FourPanesTop.Small", new IMAGE_BRUSH( "Icons/ViewportLayout_FourPanesTop", IconLayoutSizeSmall ) );
-		Set( "LevelViewport.ViewportConfig_FourPanesBottom", new IMAGE_BRUSH( "Icons/ViewportLayout_FourPanesBottom", IconLayoutSize ) );
-		Set( "LevelViewport.ViewportConfig_FourPanesBottom.Small", new IMAGE_BRUSH( "Icons/ViewportLayout_FourPanesBottom", IconLayoutSizeSmall ) );
-		Set( "LevelViewport.ViewportConfig_FourPanes2x2", new IMAGE_BRUSH( "Icons/ViewportLayout_FourPanes2x2", IconLayoutSize ) );
-		Set( "LevelViewport.ViewportConfig_FourPanes2x2.Small", new IMAGE_BRUSH( "Icons/ViewportLayout_FourPanes2x2", IconLayoutSizeSmall ) );
+		Set("LevelViewport.ViewportConfig_OnePane", new IMAGE_BRUSH_SVG("Starship/ViewportLayout/Layouts1Pane", IconLayoutSize));
+		Set("LevelViewport.ViewportConfig_OnePane.Small", new IMAGE_BRUSH_SVG("Starship/ViewportLayout/Layouts1Pane", IconLayoutSizeSmall));
+		Set("LevelViewport.ViewportConfig_TwoPanesH", new IMAGE_BRUSH_SVG("Starship/ViewportLayout/Layouts2Panes", IconLayoutSize));
+		Set("LevelViewport.ViewportConfig_TwoPanesH.Small", new IMAGE_BRUSH_SVG("Starship/ViewportLayout/Layouts2Panes", IconLayoutSizeSmall));
+		Set("LevelViewport.ViewportConfig_TwoPanesV", new IMAGE_BRUSH_SVG("Starship/ViewportLayout/Layouts2PanesStacked", IconLayoutSize));
+		Set("LevelViewport.ViewportConfig_TwoPanesV.Small", new IMAGE_BRUSH_SVG("Starship/ViewportLayout/Layouts2PanesStacked", IconLayoutSizeSmall));
+		Set("LevelViewport.ViewportConfig_ThreePanesLeft", new IMAGE_BRUSH_SVG("Starship/ViewportLayout/Layouts3PanesLeft", IconLayoutSize));
+		Set("LevelViewport.ViewportConfig_ThreePanesLeft.Small", new IMAGE_BRUSH_SVG("Starship/ViewportLayout/Layouts3PanesLeft", IconLayoutSizeSmall));
+		Set("LevelViewport.ViewportConfig_ThreePanesRight", new IMAGE_BRUSH_SVG("Starship/ViewportLayout/Layouts3PanesRight", IconLayoutSize));
+		Set("LevelViewport.ViewportConfig_ThreePanesRight.Small", new IMAGE_BRUSH_SVG("Starship/ViewportLayout/Layouts3PanesRight", IconLayoutSizeSmall));
+		Set("LevelViewport.ViewportConfig_ThreePanesTop", new IMAGE_BRUSH_SVG("Starship/ViewportLayout/Layouts3PanesTop", IconLayoutSize));
+		Set("LevelViewport.ViewportConfig_ThreePanesTop.Small", new IMAGE_BRUSH_SVG("Starship/ViewportLayout/Layouts3PanesTop", IconLayoutSizeSmall));
+		Set("LevelViewport.ViewportConfig_ThreePanesBottom", new IMAGE_BRUSH_SVG("Starship/ViewportLayout/Layouts3PanesBottom", IconLayoutSize));
+		Set("LevelViewport.ViewportConfig_ThreePanesBottom.Small", new IMAGE_BRUSH_SVG("Starship/ViewportLayout/Layouts3PanesBottom", IconLayoutSizeSmall));
+		Set("LevelViewport.ViewportConfig_FourPanesLeft", new IMAGE_BRUSH_SVG("Starship/ViewportLayout/Layouts4PanesLeft", IconLayoutSize));
+		Set("LevelViewport.ViewportConfig_FourPanesLeft.Small", new IMAGE_BRUSH_SVG("Starship/ViewportLayout/Layouts4PanesLeft", IconLayoutSizeSmall));
+		Set("LevelViewport.ViewportConfig_FourPanesRight", new IMAGE_BRUSH_SVG("Starship/ViewportLayout/Layouts4PanesRight", IconLayoutSize));
+		Set("LevelViewport.ViewportConfig_FourPanesRight.Small", new IMAGE_BRUSH_SVG("Starship/ViewportLayout/Layouts4PanesRight", IconLayoutSizeSmall));
+		Set("LevelViewport.ViewportConfig_FourPanesTop", new IMAGE_BRUSH_SVG("Starship/ViewportLayout/Layouts4PanesTop", IconLayoutSize));
+		Set("LevelViewport.ViewportConfig_FourPanesTop.Small", new IMAGE_BRUSH_SVG("Starship/ViewportLayout/Layouts4PanesTop", IconLayoutSizeSmall));
+		Set("LevelViewport.ViewportConfig_FourPanesBottom", new IMAGE_BRUSH_SVG("Starship/ViewportLayout/Layouts4PanesBottom", IconLayoutSize));
+		Set("LevelViewport.ViewportConfig_FourPanesBottom.Small", new IMAGE_BRUSH_SVG("Starship/ViewportLayout/Layouts4PanesBottom", IconLayoutSizeSmall));
+		Set("LevelViewport.ViewportConfig_FourPanes2x2", new IMAGE_BRUSH_SVG("Starship/ViewportLayout/Layouts4Panes", IconLayoutSize));
+		Set("LevelViewport.ViewportConfig_FourPanes2x2.Small", new IMAGE_BRUSH_SVG("Starship/ViewportLayout/Layouts4Panes", IconLayoutSizeSmall));
 
 		Set( "LevelViewport.EjectActorPilot", new IMAGE_BRUSH_SVG( "Starship/Common/StopPiloting_16", Icon16x16 ) );
 		Set( "LevelViewport.EjectActorPilot.Small", new IMAGE_BRUSH_SVG( "Starship/Common/StopPiloting_16", Icon16x16 ) );
