[Viewport Toolbar] Level Editor View Mode Groom icon update

diff --git a/Engine/Source/Editor/EditorStyle/Private/StarshipStyle.cpp b/Engine/Source/Editor/EditorStyle/Private/StarshipStyle.cpp
--- a/Engine/Source/Editor/EditorStyle/Private/StarshipStyle.cpp
+++ b/Engine/Source/Editor/EditorStyle/Private/StarshipStyle.cpp
@@ -4630,7 +4630,7 @@ void FStarshipEditorStyle::FStyle::SetupLevelEditorStyle()
 		Set( "EditorViewport.VisualizeNaniteMode",    new IMAGE_BRUSH_SVG("Starship/Common/Nanite_16", Icon16x16 ) );
 		Set( "EditorViewport.VisualizeLumenMode",     new IMAGE_BRUSH_SVG("Starship/Common/Lumen_16", Icon16x16 ) );
 		Set( "EditorViewport.VisualizeSubstrateMode", new IMAGE_BRUSH_SVG("Starship/Common/BufferVisualization", Icon16x16 ) );
-		Set( "EditorViewport.VisualizeGroomMode",     new IMAGE_BRUSH_SVG("Starship/Common/BufferVisualization", Icon16x16 ) );
+		Set( "EditorViewport.VisualizeGroomMode",     new IMAGE_BRUSH("Common/icon_ShowHair_16x", Icon16x16 ) );
 		Set( "EditorViewport.VisualizeVirtualShadowMapMode", new IMAGE_BRUSH_SVG("Starship/Common/VirtualShadowMap_16", Icon16x16 ) );
 		Set( "EditorViewport.VisualizeActorColorationMode", new IMAGE_BRUSH_SVG("Starship/Common/ActorColoration_16", Icon16x16));
 		Set( "EditorViewport.CollisionPawn",          new IMAGE_BRUSH_SVG("Starship/Common/PlayerCollision", Icon16x16 ) );
