Change documentation tooltip style back to Black, and apply local override to the text in the documentation control box instead.

diff --git a/Engine/Source/Editor/Documentation/Private/SDocumentationToolTip.cpp b/Engine/Source/Editor/Documentation/Private/SDocumentationToolTip.cpp
--- a/Engine/Source/Editor/Documentation/Private/SDocumentationToolTip.cpp
+++ b/Engine/Source/Editor/Documentation/Private/SDocumentationToolTip.cpp
@@ -329,7 +329,7 @@ void SDocumentationToolTip::AddDocumentation(TSharedPtr< SVerticalBox > Vertical
 			[
 				SNew(STextBlock)
 					.Text(FText::FromString(DocumentationLink + OptionalExcerptName))
-					.TextStyle(FAppStyle::Get(), "Documentation.SDocumentationTooltip")
+					.TextStyle(FAppStyle::Get(), "Documentation.Text")
 			];
 
 			if (!DocumentationPage->HasExcerpt(ExcerptName) && FSlateApplication::Get().SupportsSourceAccess())
diff --git a/Engine/Source/Editor/EditorStyle/Private/StarshipStyle.cpp b/Engine/Source/Editor/EditorStyle/Private/StarshipStyle.cpp
--- a/Engine/Source/Editor/EditorStyle/Private/StarshipStyle.cpp
+++ b/Engine/Source/Editor/EditorStyle/Private/StarshipStyle.cpp
@@ -3110,7 +3110,7 @@ void FStarshipEditorStyle::FStyle::SetupTutorialStyles()
 	{
 		const FTextBlockStyle DocumentationTooltipText = FTextBlockStyle( NormalText )
 			.SetFont( DEFAULT_FONT( "Regular", 9 ) )
-			.SetColorAndOpacity( FLinearColor::White );
+			.SetColorAndOpacity( FLinearColor::Black );
 		Set("Documentation.SDocumentationTooltip", FTextBlockStyle(DocumentationTooltipText));
 
 		const FTextBlockStyle DocumentationTooltipTextSubdued = FTextBlockStyle( NormalText )
