Fix hover and pressed colors of the level editor viewport toolbar "..." menu (ViewportSizing)

diff --git a/Engine/Source/Editor/EditorStyle/Private/StarshipStyle.cpp b/Engine/Source/Editor/EditorStyle/Private/StarshipStyle.cpp
--- a/Engine/Source/Editor/EditorStyle/Private/StarshipStyle.cpp
+++ b/Engine/Source/Editor/EditorStyle/Private/StarshipStyle.cpp
@@ -2463,6 +2463,9 @@ void FStarshipEditorStyle::FStyle::SetupViewportStyles()
 {
 	// New viewport toolbar.
 	{
+		FLinearColor ToolbarBackgroundColor = FStyleColors::Dropdown.GetSpecifiedColor();
+		ToolbarBackgroundColor.A = .90f;
+		const FSlateRoundedBoxBrush BackgroundBrush(ToolbarBackgroundColor, 5.f, ToolbarBackgroundColor, 1.0);
 
 		FLinearColor ToolbarHoveredColor = FStyleColors::Dropdown.GetSpecifiedColor();
 		ToolbarHoveredColor.R = FMath::Min(1.0f, ToolbarHoveredColor.R * 1.5);
@@ -2476,10 +2479,6 @@ void FStarshipEditorStyle::FStyle::SetupViewportStyles()
 
 		FButtonStyle ButtonStyle;
 		{
-			FLinearColor ToolbarBackgroundColor = FStyleColors::Dropdown.GetSpecifiedColor();
-			ToolbarBackgroundColor.A = .90f;
-			const FSlateRoundedBoxBrush BackgroundBrush(ToolbarBackgroundColor, 5.f, ToolbarBackgroundColor, 1.0);
-
 			ButtonStyle.SetNormal(BackgroundBrush)
 				.SetHovered(HoveredButtonBrush)
 				.SetPressed(PressedButtonBrush)
@@ -2622,9 +2621,9 @@ void FStarshipEditorStyle::FStyle::SetupViewportStyles()
 		{
 			const FButtonStyle ViewportSizingButtonStyle =
 				FButtonStyle(ButtonStyle)
-					.SetNormal(FSlateRoundedBoxBrush(FStyleColors::Dropdown, FVector4(4.0f, 0.0f, 0.0f, 4.0f)))
-					.SetHovered(FSlateRoundedBoxBrush(FStyleColors::Dropdown, FVector4(4.0f, 0.0f, 0.0f, 4.0f)))
-					.SetPressed(FSlateRoundedBoxBrush(FStyleColors::Dropdown, FVector4(4.0f, 0.0f, 0.0f, 4.0f)))
+					.SetNormal(FSlateRoundedBoxBrush(ToolbarBackgroundColor, FVector4(4.0f, 0.0f, 0.0f, 4.0f)))
+					.SetHovered(FSlateRoundedBoxBrush(ToolbarHoveredColor, FVector4(4.0f, 0.0f, 0.0f, 4.0f)))
+					.SetPressed(FSlateRoundedBoxBrush(ToolbarPressedColor, FVector4(4.0f, 0.0f, 0.0f, 4.0f)))
 					.SetDisabled(FSlateRoundedBoxBrush(FStyleColors::Dropdown, FVector4(4.0f, 0.0f, 0.0f, 4.0f)))
 					.SetNormalPadding(FMargin(0, 0, 2.f, 0))
 					.SetPressedPadding(FMargin(0, 0, 2.f, 0));
@@ -2653,9 +2652,9 @@ void FStarshipEditorStyle::FStyle::SetupViewportStyles()
 
 					ButtonStyleRaised.SetNormal(BackgroundBrushRaised)
 						.SetNormalForeground(FStyleColors::Foreground)
-						.SetNormal(FSlateRoundedBoxBrush(FStyleColors::Dropdown, FVector4(0, 4.0f, 4.0f, 0)))
-						.SetHovered(FSlateRoundedBoxBrush(FStyleColors::Dropdown, FVector4(0, 4.0f, 4.0f, 0)))
-						.SetPressed(FSlateRoundedBoxBrush(FStyleColors::Dropdown, FVector4(0, 4.0f, 4.0f, 0)))
+						.SetNormal(FSlateRoundedBoxBrush(ToolbarBackgroundColor, FVector4(0, 4.0f, 4.0f, 0)))
+						.SetHovered(FSlateRoundedBoxBrush(ToolbarHoveredColor, FVector4(0, 4.0f, 4.0f, 0)))
+						.SetPressed(FSlateRoundedBoxBrush(ToolbarPressedColor, FVector4(0, 4.0f, 4.0f, 0)))
 						.SetDisabled(FSlateRoundedBoxBrush(FStyleColors::Dropdown, FVector4(0, 4.0f, 4.0f, 0)));
 				}
 
