[Viewport Toolbar] Add Game Settings icon to camera menu Exposure section

diff --git a/Engine/Source/Editor/EditorStyle/Private/StarshipStyle.cpp b/Engine/Source/Editor/EditorStyle/Private/StarshipStyle.cpp
--- a/Engine/Source/Editor/EditorStyle/Private/StarshipStyle.cpp
+++ b/Engine/Source/Editor/EditorStyle/Private/StarshipStyle.cpp
@@ -4577,6 +4577,7 @@ void FStarshipEditorStyle::FStyle::SetupLevelEditorStyle()
 		Set("EditorViewport.ScaleGridSnap", new IMAGE_BRUSH_SVG( "Starship/EditorViewport/scale-grid-snap", Icon16x16 ) );
 		Set("EditorViewport.ToggleSurfaceSnapping", new IMAGE_BRUSH_SVG( "Starship/EditorViewport/surface-snap", Icon16x16 ) );
 		Set("EditorViewport.ToggleSurfaceSnapping", new IMAGE_BRUSH_SVG("Starship/EditorViewport/surface-snap", Icon16x16));
+		Set("EditorViewport.ToggleInGameExposure", new IMAGE_BRUSH_SVG("Starship/Common/GameSettings_16", Icon16x16));
 
 		Set( "EditorViewport.RelativeCoordinateSystem_World", new IMAGE_BRUSH_SVG( "Starship/EditorViewport/globe", Icon16x16 ) );
 		Set("EditorViewport.RelativeCoordinateSystem_Local",
