[EditorStyle] Fix Graph.ConnectorFeedback.OKWarn Brush
Fix wrong Macro for that style entry

diff --git a/Engine/Source/Editor/EditorStyle/Private/StarshipStyle.cpp b/Engine/Source/Editor/EditorStyle/Private/StarshipStyle.cpp
--- a/Engine/Source/Editor/EditorStyle/Private/StarshipStyle.cpp
+++ b/Engine/Source/Editor/EditorStyle/Private/StarshipStyle.cpp
@@ -4194,7 +4194,7 @@ void FStarshipEditorStyle::FStyle::SetupGraphEditorStyles()
 
 		Set( "Graph.ConnectorFeedback.Border", new BOX_BRUSH( "Old/Menu_Background", FMargin(8.0f/64.0f) ) );
 		Set( "Graph.ConnectorFeedback.OK", new CORE_IMAGE_BRUSH_SVG("Starship/Common/check-circle", Icon16x16, FStyleColors::AccentGreen));
-		Set( "Graph.ConnectorFeedback.OKWarn", new IMAGE_BRUSH_SVG("Starship/Common/check-circle", Icon16x16, FStyleColors::AccentYellow));
+		Set( "Graph.ConnectorFeedback.OKWarn", new CORE_IMAGE_BRUSH_SVG("Starship/Common/check-circle", Icon16x16, FStyleColors::AccentYellow));
 		Set( "Graph.ConnectorFeedback.Error", new IMAGE_BRUSH( "Old/Graph/Feedback_Error", Icon16x16 ) );
 		Set( "Graph.ConnectorFeedback.NewNode", new IMAGE_BRUSH( "Old/Graph/Feedback_NewNode", Icon16x16 ) );
 		Set( "Graph.ConnectorFeedback.ViaCast", new IMAGE_BRUSH( "Old/Graph/Feedback_ConnectViaCast", Icon16x16 ) );
