[UE] PhAT Update Body Icon Set
- Replace old physics body icons in PhAT Skeleton tree with new, monochrom icons that differentiate between default, simulated and kinematic bodies that contain a single primitive or multiple primitives

diff --git a/Engine/Source/Editor/EditorStyle/Private/SlateEditorStyle.cpp b/Engine/Source/Editor/EditorStyle/Private/SlateEditorStyle.cpp
--- a/Engine/Source/Editor/EditorStyle/Private/SlateEditorStyle.cpp
+++ b/Engine/Source/Editor/EditorStyle/Private/SlateEditorStyle.cpp
@@ -7639,8 +7639,14 @@ void FSlateEditorStyle::FStyle::SetupToolkitStyles()
 		Set("PhysicsAssetEditor.AddConstraintToCurrentConstraintProfile", new IMAGE_BRUSH("PhysicsAssetEditor/icon_PHatMode_Joint_40x", Icon20x20));
 		Set("PhysicsAssetEditor.RemoveConstraintFromCurrentConstraintProfile", new IMAGE_BRUSH("PhysicsAssetEditor/icon_PhAT_DeleteConstraint_40x", Icon20x20));
 
-		Set("PhysicsAssetEditor.Tree.Body", new IMAGE_BRUSH("PhysicsAssetEditor/BoneGroup_16x", Icon16x16));
-		Set("PhysicsAssetEditor.Tree.KinematicBody", new IMAGE_BRUSH("PhysicsAssetEditor/KinematicBody_16x", Icon16x16));
+		Set("PhysicsAssetEditor.Tree.BodyMultipleDefault", new IMAGE_BRUSH("PhysicsAssetEditor/BodyMultipleDefault_16x", Icon16x16));
+		Set("PhysicsAssetEditor.Tree.BodyMultipleKinematic", new IMAGE_BRUSH("PhysicsAssetEditor/BodyMultipleKinematic_16x", Icon16x16));
+		Set("PhysicsAssetEditor.Tree.BodyMultipleSimulated", new IMAGE_BRUSH("PhysicsAssetEditor/BodyMultipleSimulated_16x", Icon16x16));
+		Set("PhysicsAssetEditor.Tree.BodySingleDefault", new IMAGE_BRUSH("PhysicsAssetEditor/BodySingleDefault_16x", Icon16x16));
+		Set("PhysicsAssetEditor.Tree.BodySingleKinematic", new IMAGE_BRUSH("PhysicsAssetEditor/BodySingleKinematic_16x", Icon16x16));
+		Set("PhysicsAssetEditor.Tree.BodySingleSimulated", new IMAGE_BRUSH("PhysicsAssetEditor/BodySingleSimulated_16x", Icon16x16));
+		Set("PhysicsAssetEditor.Tree.Body", new IMAGE_BRUSH("PhysicsAssetEditor/BodyMultipleDefault_16x", Icon16x16));
+		Set("PhysicsAssetEditor.Tree.KinematicBody", new IMAGE_BRUSH("PhysicsAssetEditor/BodyMultipleKinematic_16x", Icon16x16));
 		Set("PhysicsAssetEditor.Tree.EmptyBody", new IMAGE_BRUSH("PhysicsAssetEditor/EmptyBody_16x", Icon16x16));
 		Set("PhysicsAssetEditor.Tree.Bone", new IMAGE_BRUSH("PhysicsAssetEditor/Bone_16x", Icon16x16));
 		Set("PhysicsAssetEditor.Tree.Sphere", new IMAGE_BRUSH("PhysicsAssetEditor/Sphere_16x", Icon16x16));
diff --git a/Engine/Source/Editor/EditorStyle/Private/StarshipStyle.cpp b/Engine/Source/Editor/EditorStyle/Private/StarshipStyle.cpp
--- a/Engine/Source/Editor/EditorStyle/Private/StarshipStyle.cpp
+++ b/Engine/Source/Editor/EditorStyle/Private/StarshipStyle.cpp
@@ -7616,8 +7616,14 @@ void FStarshipEditorStyle::FStyle::SetupToolkitStyles()
 		Set("PhysicsAssetEditor.AddConstraintToCurrentConstraintProfile", new IMAGE_BRUSH("PhysicsAssetEditor/icon_PHatMode_Joint_40x", Icon20x20));
 		Set("PhysicsAssetEditor.RemoveConstraintFromCurrentConstraintProfile", new IMAGE_BRUSH("PhysicsAssetEditor/icon_PhAT_DeleteConstraint_40x", Icon20x20));
 
-		Set("PhysicsAssetEditor.Tree.Body", new IMAGE_BRUSH("PhysicsAssetEditor/BoneGroup_16x", Icon16x16));
-		Set("PhysicsAssetEditor.Tree.KinematicBody", new IMAGE_BRUSH("PhysicsAssetEditor/KinematicBody_16x", Icon16x16));
+		Set("PhysicsAssetEditor.Tree.BodyMultipleDefault", new IMAGE_BRUSH("PhysicsAssetEditor/BodyMultipleDefault_16x", Icon16x16));
+		Set("PhysicsAssetEditor.Tree.BodyMultipleKinematic", new IMAGE_BRUSH("PhysicsAssetEditor/BodyMultipleKinematic_16x", Icon16x16));
+		Set("PhysicsAssetEditor.Tree.BodyMultipleSimulated", new IMAGE_BRUSH("PhysicsAssetEditor/BodyMultipleSimulated_16x", Icon16x16));
+		Set("PhysicsAssetEditor.Tree.BodySingleDefault", new IMAGE_BRUSH("PhysicsAssetEditor/BodySingleDefault_16x", Icon16x16));
+		Set("PhysicsAssetEditor.Tree.BodySingleKinematic", new IMAGE_BRUSH("PhysicsAssetEditor/BodySingleKinematic_16x", Icon16x16));
+		Set("PhysicsAssetEditor.Tree.BodySingleSimulated", new IMAGE_BRUSH("PhysicsAssetEditor/BodySingleSimulated_16x", Icon16x16));
+		Set("PhysicsAssetEditor.Tree.Body", new IMAGE_BRUSH("PhysicsAssetEditor/BodyMultipleDefault_16x", Icon16x16));
+		Set("PhysicsAssetEditor.Tree.KinematicBody", new IMAGE_BRUSH("PhysicsAssetEditor/BodyMultipleKinematic_16x", Icon16x16));
 		Set("PhysicsAssetEditor.Tree.EmptyBody", new IMAGE_BRUSH("PhysicsAssetEditor/EmptyBody_16x", Icon16x16));
 		Set("PhysicsAssetEditor.Tree.Bone", new IMAGE_BRUSH("PhysicsAssetEditor/Bone_16x", Icon16x16));
 		Set("PhysicsAssetEditor.Tree.Sphere", new IMAGE_BRUSH("PhysicsAssetEditor/Sphere_16x", Icon16x16));
diff --git a/Engine/Source/Editor/PhysicsAssetEditor/Private/SkeletonTreePhysicsBodyItem.cpp b/Engine/Source/Editor/PhysicsAssetEditor/Private/SkeletonTreePhysicsBodyItem.cpp
--- a/Engine/Source/Editor/PhysicsAssetEditor/Private/SkeletonTreePhysicsBodyItem.cpp
+++ b/Engine/Source/Editor/PhysicsAssetEditor/Private/SkeletonTreePhysicsBodyItem.cpp
@@ -43,7 +43,36 @@ ECheckBoxState FSkeletonTreePhysicsBodyItem::IsItemDisplayed() const
 
 const FSlateBrush* FSkeletonTreePhysicsBodyItem::GetBrush() const
 {
-	return BodySetup->PhysicsType == EPhysicsType::PhysType_Kinematic ? FAppStyle::GetBrush("PhysicsAssetEditor.Tree.KinematicBody") : FAppStyle::GetBrush("PhysicsAssetEditor.Tree.Body");
+	if (BodySetup)
+	{
+		if (BodySetup->AggGeom.GetElementCount() > 1)
+		{
+			if (BodySetup->PhysicsType == EPhysicsType::PhysType_Kinematic)
+			{
+				return FAppStyle::GetBrush("PhysicsAssetEditor.Tree.BodyMultipleKinematic");
+			}
+			else if (BodySetup->PhysicsType == EPhysicsType::PhysType_Simulated)
+			{
+				return FAppStyle::GetBrush("PhysicsAssetEditor.Tree.BodyMultipleSimulated");
+			}
+
+			check(BodySetup->PhysicsType == EPhysicsType::PhysType_Default);
+			return FAppStyle::GetBrush("PhysicsAssetEditor.Tree.BodyMultipleDefault");
+		}
+
+		if (BodySetup->PhysicsType == EPhysicsType::PhysType_Kinematic)
+		{
+			return FAppStyle::GetBrush("PhysicsAssetEditor.Tree.BodySingleKinematic");
+		}
+		else if (BodySetup->PhysicsType == EPhysicsType::PhysType_Simulated)
+		{
+			return FAppStyle::GetBrush("PhysicsAssetEditor.Tree.BodySingleSimulated");
+		}
+
+		check(BodySetup->PhysicsType == EPhysicsType::PhysType_Default);
+	}
+
+	return 	FAppStyle::GetBrush("PhysicsAssetEditor.Tree.BodySingleDefault");
 }
 
 FSlateColor FSkeletonTreePhysicsBodyItem::GetTextColor() const
