Improve viewport toolbar styling & item differentiation.
- Add a small chevron for toolbar combo buttons.
- Tweak element padding for better readability.
- Add additional toolbar-specific customization on FToolMenuEntry.

diff --git a/Engine/Source/Developer/ToolMenus/Private/ToolMenus.cpp b/Engine/Source/Developer/ToolMenus/Private/ToolMenus.cpp
--- a/Engine/Source/Developer/ToolMenus/Private/ToolMenus.cpp
+++ b/Engine/Source/Developer/ToolMenus/Private/ToolMenus.cpp
@@ -1716,14 +1716,18 @@ void UToolMenus::PopulateToolBarBuilderWithEntry(
 
 	// Override the style name.
 	{
-		FName OverrideStyleName = Block.StyleNameOverride;
+		FName OverrideStyleName = Block.ToolBarData.StyleNameOverride;
+		if (OverrideStyleName == NAME_None)
+		{
+			OverrideStyleName = Block.StyleNameOverride;
+		}
 
 		// Add the .Raised suffix for menu entries raised to the top level.
 		if (bIsRaisingToTopLevel)
 		{
-			if (Block.StyleNameOverride != NAME_None)
+			if (OverrideStyleName != NAME_None)
 			{
-				OverrideStyleName = ISlateStyle::Join(Block.StyleNameOverride, ".Raised");
+				OverrideStyleName = ISlateStyle::Join(OverrideStyleName, ".Raised");
 			}
 			else
 			{
diff --git a/Engine/Source/Developer/ToolMenus/Public/ToolMenuEntry.h b/Engine/Source/Developer/ToolMenus/Public/ToolMenuEntry.h
--- a/Engine/Source/Developer/ToolMenus/Public/ToolMenuEntry.h
+++ b/Engine/Source/Developer/ToolMenus/Public/ToolMenuEntry.h
@@ -60,6 +60,9 @@ public:
 	/** Optional override label to use when the entry appears in a toolbar. */
 	TAttribute<FText> LabelOverride;
 
+	/** Optional style override used when this entry appears in a toolbar. */
+	FName StyleNameOverride = NAME_None;
+
 	/** Delegate that generates a widget for this combo button's menu content. Called when the menu is summoned. */
 	FNewToolMenuChoice ComboButtonContextMenuGenerator;
 
diff --git a/Engine/Source/Editor/EditorStyle/Private/StarshipStyle.cpp b/Engine/Source/Editor/EditorStyle/Private/StarshipStyle.cpp
--- a/Engine/Source/Editor/EditorStyle/Private/StarshipStyle.cpp
+++ b/Engine/Source/Editor/EditorStyle/Private/StarshipStyle.cpp
@@ -2496,8 +2496,8 @@ void FStarshipEditorStyle::FStyle::SetupViewportStyles()
 				.SetHoveredForeground(FStyleColors::ForegroundHover)
 				.SetPressedForeground(FStyleColors::ForegroundHover)
 				.SetDisabledForeground(FStyleColors::Foreground)
-				.SetNormalPadding(FMargin(4.0f))
-				.SetPressedPadding(FMargin(4.0f));
+				.SetNormalPadding(FMargin(6.0f, 4.0f))
+				.SetPressedPadding(FMargin(6.0f, 4.0f));
 		}
 
 		const FToolBarStyle SlimToolBarStyle =
@@ -2521,9 +2521,10 @@ void FStarshipEditorStyle::FStyle::SetupViewportStyles()
 		{
 			FButtonStyle InnerButtonStyle = FButtonStyle(ButtonStyle).SetNormalPadding(0).SetPressedPadding(0);
 
-			ComboButtonStyle.SetContentPadding(FMargin(2.0f, 0, 0, 0))
-				.SetDownArrowPadding(FMargin(1.0f, 0, 0, 0))
-				.SetButtonStyle(InnerButtonStyle);
+			ComboButtonStyle
+				.SetButtonStyle(InnerButtonStyle)
+				.SetDownArrowImage(IMAGE_BRUSH_SVG("Starship/EditorViewport/small-chevron-down", FVector2f(6.f, 6.f)))
+				.SetDownArrowPadding(FMargin(4.0f, 0, 0, 0));
 		}
 
 		FToolBarStyle ViewportToolbarStyle = SlimToolBarStyle;
@@ -2535,14 +2536,13 @@ void FStarshipEditorStyle::FStyle::SetupViewportStyles()
 			ViewportToolbarStyle.SetIconSize(Icon16x16)
 				.SetButtonStyle(ButtonStyle)
 				.SetToggleButtonStyle(ToggleButtonStyle)
-				.SetButtonPadding(FMargin(2.0f, 0))
+				.SetButtonPadding(0)
 				.SetComboButtonStyle(ComboButtonStyle)
 				.SetComboButtonPadding(FMargin(4.0f, 0))
-				.SetButtonPadding(0)
 				.SetSeparatorBrush(SeparatorBrush)
 				.SetSeparatorThickness(1.0f)
 				.SetSeparatorPadding(FMargin(4.0f, 0))
-				.SetBackgroundPadding(FMargin(0, 4.0f, 4.0f, 4.0f))
+				.SetBackgroundPadding(FMargin(2.0f))
 				.SetWrapButtonPadding(FMargin(4.0f, 0))
 				.SetWrapButtonIndex(-1)
 				.SetRaisedChildrenRightPadding(16.0f);
@@ -2559,11 +2559,13 @@ void FStarshipEditorStyle::FStyle::SetupViewportStyles()
 				BackgroundColorRaised.A = .0f;
 				const FSlateRoundedBoxBrush BackgroundBrushRaised(BackgroundColorRaised, 5.f, BackgroundColorRaised, 1.0);
 
-				ButtonStyleRaised.SetNormal(BackgroundBrushRaised)
+				ButtonStyleRaised
+					.SetNormal(BackgroundBrushRaised)
 					.SetNormalForeground(FStyleColors::Foreground)
-					.SetPressed(BackgroundBrush);
+					.SetNormalPadding(4.0f)
+					.SetPressedPadding(4.0f);
 			}
-
+			
 			FCheckBoxStyle ToggleButtonStyleRaised = ToggleButtonStyle;
 			{
 				ToggleButtonStyleRaised.SetUncheckedHoveredImage(HoveredButtonBrush)
@@ -2571,11 +2573,19 @@ void FStarshipEditorStyle::FStyle::SetupViewportStyles()
 					.SetUncheckedPressedImage(BackgroundBrush);
 			}
 
-			ViewportToolbarRaisedStyle.SetButtonStyle(ButtonStyleRaised)
-				.SetToggleButtonStyle(ToggleButtonStyleRaised)
-				.SetComboButtonPadding(0)
-				.SetBlockHovered(BackgroundBrush);
+			FComboButtonStyle ComboStyleRaised = ViewportToolbarStyle.ComboButtonStyle;
+			{
+				ComboButtonStyle.SetDownArrowPadding(FMargin(2.0f, 0, 0, 0));
+			}
 
+			ViewportToolbarRaisedStyle
+				.SetButtonStyle(ButtonStyleRaised)
+				.SetToggleButtonStyle(ToggleButtonStyleRaised)
+				.SetComboButtonStyle(ComboStyleRaised)
+				.SetBlockHovered(BackgroundBrush)
+				.SetButtonPadding(FMargin(3.0f, 0))
+				.SetComboButtonPadding(FMargin(3.0f, 0));
+			
 			Set("ViewportToolbar.Raised", ViewportToolbarRaisedStyle);
 		}
 
@@ -2588,6 +2598,13 @@ void FStarshipEditorStyle::FStyle::SetupViewportStyles()
 			Set("LevelEditorViewportToolbar", LevelEditorViewportToolbarStyle);
 			Set("LevelEditorViewportToolbar.Raised", ViewportToolbarRaisedStyle);
 		}
+		
+		// Special style for raised transform buttons
+		{
+			FToolBarStyle RaisedTransformTools = FToolBarStyle(GetWidgetStyle<FToolBarStyle>("ViewportToolbar.Raised"));
+			RaisedTransformTools.SetButtonPadding(0);
+			Set("ViewportToolbar.TransformTools.Raised", RaisedTransformTools);
+		}
 
 		// Special styling for warnings.
 		{
diff --git a/Engine/Source/Editor/UnrealEd/Private/ViewportToolbar/UnrealEdViewportToolbar.cpp b/Engine/Source/Editor/UnrealEd/Private/ViewportToolbar/UnrealEdViewportToolbar.cpp
--- a/Engine/Source/Editor/UnrealEd/Private/ViewportToolbar/UnrealEdViewportToolbar.cpp
+++ b/Engine/Source/Editor/UnrealEd/Private/ViewportToolbar/UnrealEdViewportToolbar.cpp
@@ -706,29 +706,36 @@ FToolMenuEntry CreateTransformsSubmenu()
 
 					FToolMenuEntry SelectMode = FToolMenuEntry::InitMenuEntry(FEditorViewportCommands::Get().SelectMode);
 					SelectMode.SetShowInToolbarTopLevel(true);
+					SelectMode.ToolBarData.StyleNameOverride = "ViewportToolbar.TransformTools";
+					
 					TransformToolsSection.AddEntry(SelectMode);
 
 					FToolMenuEntry TranslateMode =
 						FToolMenuEntry::InitMenuEntry(FEditorViewportCommands::Get().TranslateMode);
 					TranslateMode.SetShowInToolbarTopLevel(true);
+					TranslateMode.ToolBarData.StyleNameOverride = "ViewportToolbar.TransformTools";
 					TransformToolsSection.AddEntry(TranslateMode);
 
 					FToolMenuEntry RotateMode = FToolMenuEntry::InitMenuEntry(FEditorViewportCommands::Get().RotateMode);
 					RotateMode.SetShowInToolbarTopLevel(true);
+					RotateMode.ToolBarData.StyleNameOverride = "ViewportToolbar.TransformTools";
 					TransformToolsSection.AddEntry(RotateMode);
 
 					FToolMenuEntry ScaleMode = FToolMenuEntry::InitMenuEntry(FEditorViewportCommands::Get().ScaleMode);
 					ScaleMode.SetShowInToolbarTopLevel(true);
+					ScaleMode.ToolBarData.StyleNameOverride = "ViewportToolbar.TransformTools";
 					TransformToolsSection.AddEntry(ScaleMode);
 
 					FToolMenuEntry TranslateRotateMode =
 						FToolMenuEntry::InitMenuEntry(FEditorViewportCommands::Get().TranslateRotateMode);
 					TranslateRotateMode.SetShowInToolbarTopLevel(true);
+					TranslateRotateMode.ToolBarData.StyleNameOverride = "ViewportToolbar.TransformTools";
 					TransformToolsSection.AddEntry(TranslateRotateMode);
 
 					FToolMenuEntry TranslateRotate2DMode =
 						FToolMenuEntry::InitMenuEntry(FEditorViewportCommands::Get().TranslateRotate2DMode);
 					TranslateRotate2DMode.SetShowInToolbarTopLevel(true);
+					TranslateRotate2DMode.ToolBarData.StyleNameOverride = "ViewportToolbar.TransformTools";
 					TransformToolsSection.AddEntry(TranslateRotate2DMode);
 
 					// Build a submenu for selecting the coordinate system to use.
diff --git a/Engine/Source/Runtime/Slate/Private/Framework/MultiBox/SToolBarComboButtonBlock.cpp b/Engine/Source/Runtime/Slate/Private/Framework/MultiBox/SToolBarComboButtonBlock.cpp
--- a/Engine/Source/Runtime/Slate/Private/Framework/MultiBox/SToolBarComboButtonBlock.cpp
+++ b/Engine/Source/Runtime/Slate/Private/Framework/MultiBox/SToolBarComboButtonBlock.cpp
@@ -344,13 +344,13 @@ void SToolBarComboButtonBlock::BuildMultiBlockWidget(const ISlateStyle* StyleSet
 				))
 				+ SHorizontalBox::Slot()
 				.AutoWidth()
-				.VAlign(VAlign_Center)
+				.VAlign(VAlign_Fill)
 				[
 					CheckBoxWidget.ToSharedRef()
 				]
 				+ SHorizontalBox::Slot()
 				.AutoWidth()
-				.VAlign(VAlign_Center)
+				.VAlign(VAlign_Fill)
 				[
 					ComboButtonWidget.ToSharedRef()
 				]
