Removing Grey Background from Viewport layout icons and decreased icon size from UX feedback

diff --git a/Engine/Source/Editor/EditorStyle/Private/StarshipStyle.cpp b/Engine/Source/Editor/EditorStyle/Private/StarshipStyle.cpp
--- a/Engine/Source/Editor/EditorStyle/Private/StarshipStyle.cpp
+++ b/Engine/Source/Editor/EditorStyle/Private/StarshipStyle.cpp
@@ -1901,7 +1901,8 @@ void FStarshipEditorStyle::FStyle::SetupGeneralStyles()
 	// ViewportLayoutToolbar
 	{
 		FToolBarStyle ViewportLayoutToolbar = FStarshipCoreStyle::GetCoreStyle().GetWidgetStyle<FToolBarStyle>("SlimToolBar");
-		ViewportLayoutToolbar.SetIconSize(Icon40x40);
+		ViewportLayoutToolbar.SetIconSize(Icon32x32);
+		ViewportLayoutToolbar.SetBackground(FSlateColorBrush(FStyleColors::Dropdown));
 
 		Set("ViewportLayoutToolbar", ViewportLayoutToolbar);
 	}
