[Truncation Warnings] Fix and enable truncation warnings for EditorConfig and EditorStyle and EditorWidgets modules

diff --git a/Engine/Source/Editor/EditorStyle/Private/SlateEditorStyle.cpp b/Engine/Source/Editor/EditorStyle/Private/SlateEditorStyle.cpp
--- a/Engine/Source/Editor/EditorStyle/Private/SlateEditorStyle.cpp
+++ b/Engine/Source/Editor/EditorStyle/Private/SlateEditorStyle.cpp
@@ -124,8 +124,8 @@ void FSlateEditorStyle::FStyle::SyncSettings()
 
 		// Sync the window background settings
 		FSlateColor WindowBackgroundColor(FLinearColor::White);
-		FSlateBrush WindowBackgroundMain(IMAGE_BRUSH("Old/Window/WindowBackground", FVector2D(74, 74), FLinearColor::White, ESlateBrushTileType::Both));
-		FSlateBrush WindowBackgroundChild(IMAGE_BRUSH("Common/NoiseBackground", FVector2D(64, 64), FLinearColor::White, ESlateBrushTileType::Both));
+		FSlateBrush WindowBackgroundMain(IMAGE_BRUSH("Old/Window/WindowBackground", FVector2f(74.f, 74.f), FLinearColor::White, ESlateBrushTileType::Both));
+		FSlateBrush WindowBackgroundChild(IMAGE_BRUSH("Common/NoiseBackground", Icon64x64, FLinearColor::White, ESlateBrushTileType::Both));
 
 		WindowBackgroundColor = Settings->EditorWindowBackgroundColor;
 
@@ -202,7 +202,7 @@ void FSlateEditorStyle::FStyle::SetupGeneralStyles()
 	NormalText = FTextBlockStyle()
 		.SetFont(DEFAULT_FONT("Regular", FCoreStyle::RegularTextSize))
 		.SetColorAndOpacity(FSlateColor::UseForeground())
-		.SetShadowOffset(FVector2D::ZeroVector)
+		.SetShadowOffset(FVector2f::ZeroVector)
 		.SetShadowColorAndOpacity(FLinearColor::Black)
 		.SetHighlightColor( FLinearColor( 0.02f, 0.3f, 0.0f ) )
 		.SetHighlightShape( BOX_BRUSH( "Common/TextBlockHighlightShape", FMargin(3.f/8.f) ) );
@@ -244,12 +244,12 @@ void FSlateEditorStyle::FStyle::SetupGeneralStyles()
 		Set( "TextBlock.HighlighColor", FLinearColor( 0.02f, 0.3f, 0.0f ) );
 
 		Set("TextBlock.ShadowedText", FTextBlockStyle(NormalText)
-			.SetShadowOffset(FVector2D(1.0f, 1.0f))
+			.SetShadowOffset(FVector2f::UnitVector)
 			.SetShadowColorAndOpacity(FLinearColor(0.0f, 0.0f, 0.0f)));
 
 		Set("TextBlock.ShadowedTextWarning", FTextBlockStyle(NormalText)
 			.SetColorAndOpacity(FLinearColor(1.0f, 0.0f, 0.0f))
-			.SetShadowOffset(FVector2D(1.0f, 1.0f))
+			.SetShadowOffset(FVector2f::UnitVector)
 			.SetShadowColorAndOpacity(FLinearColor(0.0f, 0.0f, 0.0f)));
 
 		Set("NormalText", NormalText);
@@ -261,7 +261,7 @@ void FSlateEditorStyle::FStyle::SetupGeneralStyles()
 			.SetFont(DEFAULT_FONT("Bold", FCoreStyle::RegularTextSize))
 			.SetColorAndOpacity(FLinearColor(1.0f, 1.0f, 1.0f))
 			.SetHighlightColor(FLinearColor(1.0f, 1.0f, 1.0f))
-			.SetShadowOffset(FVector2D(1, 1))
+			.SetShadowOffset(FVector2f::UnitVector)
 			.SetShadowColorAndOpacity(FLinearColor(0, 0, 0, 0.9f)));
 
 		Set("SmallText", FTextBlockStyle(NormalText)
@@ -282,7 +282,7 @@ void FSlateEditorStyle::FStyle::SetupGeneralStyles()
 			.SetFont(DEFAULT_FONT("Bold", 11))
 			.SetColorAndOpacity(FLinearColor(1.0f, 1.0f, 1.0f))
 			.SetHighlightColor(FLinearColor(1.0f, 1.0f, 1.0f))
-			.SetShadowOffset(FVector2D(1, 1))
+			.SetShadowOffset(FVector2f::UnitVector)
 			.SetShadowColorAndOpacity(FLinearColor(0, 0, 0, 0.9f)));
 	}
 
@@ -298,10 +298,10 @@ void FSlateEditorStyle::FStyle::SetupGeneralStyles()
 
 	Set( "PlainBorder", new BORDER_BRUSH( "Common/PlainBorder", 2.f/8.f) );
 
-	Set( "WideDash.Horizontal", new IMAGE_BRUSH("Common/WideDash_Horizontal", FVector2D(22, 4), FLinearColor::White, ESlateBrushTileType::Horizontal));
-	Set( "WideDash.Vertical", new IMAGE_BRUSH("Common/WideDash_Vertical", FVector2D(4, 22), FLinearColor::White, ESlateBrushTileType::Vertical));
+	Set( "WideDash.Horizontal", new IMAGE_BRUSH("Common/WideDash_Horizontal", FVector2f(22.f, 4.f), FLinearColor::White, ESlateBrushTileType::Horizontal));
+	Set( "WideDash.Vertical", new IMAGE_BRUSH("Common/WideDash_Vertical", FVector2f(4.f, 22.f), FLinearColor::White, ESlateBrushTileType::Vertical));
 
-	Set("ThinLine.Horizontal", new IMAGE_BRUSH("Common/ThinLine_Horizontal", FVector2D(11, 2), FLinearColor::White, ESlateBrushTileType::Horizontal));
+	Set("ThinLine.Horizontal", new IMAGE_BRUSH("Common/ThinLine_Horizontal", FVector2f(11.f, 2.f), FLinearColor::White, ESlateBrushTileType::Horizontal));
 
 	// Debug Colors
 	Set( "MultiboxHookColor", FLinearColor(0.f, 1.f, 0.f, 1.f) );
@@ -384,15 +384,15 @@ void FSlateEditorStyle::FStyle::SetupGeneralStyles()
 		Set( "Icons.Contact", new IMAGE_BRUSH( "Icons/icon_mail_16x", Icon16x16 ) );
 	}
 
-	Set( "WarningStripe", new IMAGE_BRUSH( "Common/WarningStripe", FVector2D(20,6), FLinearColor::White, ESlateBrushTileType::Horizontal ) );
+	Set( "WarningStripe", new IMAGE_BRUSH( "Common/WarningStripe", FVector2f(20.f,6.f), FLinearColor::White, ESlateBrushTileType::Horizontal ) );
 
 	// Normal button
 	Button = FButtonStyle()
-		.SetNormal( BOX_BRUSH( "Common/Button", FVector2D(32,32), 8.0f/32.0f ) )
-		.SetHovered( BOX_BRUSH( "Common/Button_Hovered", FVector2D(32,32), 8.0f/32.0f ) )
-		.SetPressed( BOX_BRUSH( "Common/Button_Pressed", FVector2D(32,32), 8.0f/32.0f ) )
-		.SetNormalPadding( FMargin( 2,2,2,2 ) )
-		.SetPressedPadding( FMargin( 2,3,2,1 ) )
+		.SetNormal( BOX_BRUSH( "Common/Button", Icon32x32, 8.0f/32.0f ) )
+		.SetHovered( BOX_BRUSH( "Common/Button_Hovered", Icon32x32, 8.0f/32.0f ) )
+		.SetPressed( BOX_BRUSH( "Common/Button_Pressed", Icon32x32, 8.0f/32.0f ) )
+		.SetNormalPadding( FMargin( 2.f ) )
+		.SetPressedPadding( FMargin( 2.f,3.f,2.f,1.f) )
 		.SetNormalForeground(InvertedForeground)
 		.SetPressedForeground(InvertedForeground)
 		.SetHoveredForeground(InvertedForeground)
@@ -436,9 +436,9 @@ void FSlateEditorStyle::FStyle::SetupGeneralStyles()
 		);
 
 		Set("FlatButton.Light", FButtonStyle(Button)
-			.SetNormal(BOX_BRUSH("Common/FlatButton", 2.0f / 8.0f, FLinearColor(0.72267, 0.72267, 0.72267, 1)))
-			.SetHovered(BOX_BRUSH("Common/FlatButton", 2.0f / 8.0f, FLinearColor(0.85, 0.85, 0.85, 1)))
-			.SetPressed(BOX_BRUSH("Common/FlatButton", 2.0f / 8.0f, FLinearColor(0.58597, 0.58597, 0.58597, 1)))
+			.SetNormal(BOX_BRUSH("Common/FlatButton", 2.0f / 8.0f, FLinearColor(0.72267f, 0.72267f, 0.72267f, 1.f)))
+			.SetHovered(BOX_BRUSH("Common/FlatButton", 2.0f / 8.0f, FLinearColor(0.85f, 0.85f, 0.85f, 1.f)))
+			.SetPressed(BOX_BRUSH("Common/FlatButton", 2.0f / 8.0f, FLinearColor(0.58597f, 0.58597f, 0.58597f, 1.f)))
 			);
 
 		Set("FlatButton.Default", GetWidgetStyle<FButtonStyle>("FlatButton.Dark"));
@@ -447,8 +447,8 @@ void FSlateEditorStyle::FStyle::SetupGeneralStyles()
 			.SetFont(DEFAULT_FONT("Bold", 10))
 			.SetColorAndOpacity(FLinearColor(1.0f, 1.0f, 1.0f))
 			.SetHighlightColor(FLinearColor(1.0f, 1.0f, 1.0f))
-			.SetShadowOffset(FVector2D(1, 1))
-			.SetShadowColorAndOpacity(FLinearColor(0, 0, 0, 0.9f)));
+			.SetShadowOffset(FVector2f::UnitVector)
+			.SetShadowColorAndOpacity(FLinearColor(0.f, 0.f, 0.f, 0.9f)));
 
 		struct ButtonColor
 		{
@@ -470,11 +470,11 @@ void FSlateEditorStyle::FStyle::SetupGeneralStyles()
 		};
 
 		TArray< ButtonColor > FlatButtons;
-		FlatButtons.Add(ButtonColor("FlatButton.Primary", FLinearColor(0.02899, 0.19752, 0.48195)));
-		FlatButtons.Add(ButtonColor("FlatButton.Success", FLinearColor(0.10616, 0.48777, 0.10616)));
-		FlatButtons.Add(ButtonColor("FlatButton.Info", FLinearColor(0.10363, 0.53564, 0.7372)));
-		FlatButtons.Add(ButtonColor("FlatButton.Warning", FLinearColor(0.87514, 0.42591, 0.07383)));
-		FlatButtons.Add(ButtonColor("FlatButton.Danger", FLinearColor(0.70117, 0.08464, 0.07593)));
+		FlatButtons.Add(ButtonColor("FlatButton.Primary", FLinearColor(0.02899f, 0.19752f, 0.48195f)));
+		FlatButtons.Add(ButtonColor("FlatButton.Success", FLinearColor(0.10616f, 0.48777f, 0.10616f)));
+		FlatButtons.Add(ButtonColor("FlatButton.Info", FLinearColor(0.10363f, 0.53564f, 0.7372f)));
+		FlatButtons.Add(ButtonColor("FlatButton.Warning", FLinearColor(0.87514f, 0.42591f, 0.07383f)));
+		FlatButtons.Add(ButtonColor("FlatButton.Danger", FLinearColor(0.70117f, 0.08464f, 0.07593f)));
 
 		for ( const ButtonColor& Entry : FlatButtons )
 		{
@@ -498,7 +498,7 @@ void FSlateEditorStyle::FStyle::SetupGeneralStyles()
 		/* Create a checkbox style for "ToggleButton" ... */
 		const FCheckBoxStyle ToggleButtonStyle = FCheckBoxStyle()
 			.SetCheckBoxType(ESlateCheckBoxType::ToggleButton)
-			.SetUncheckedImage( BOX_BRUSH( "Common/RoundedSelection_16x", 4.0f/16.0f, FLinearColor( 1, 1, 1, 0.1f ) ) )
+			.SetUncheckedImage( BOX_BRUSH( "Common/RoundedSelection_16x", 4.0f/16.0f, FLinearColor( 1.f, 1.f, 1.f, 0.1f ) ) )
 			.SetUncheckedHoveredImage( BOX_BRUSH( "Common/RoundedSelection_16x", 4.0f/16.0f, SelectionColor ) )
 			.SetUncheckedPressedImage( BOX_BRUSH( "Common/RoundedSelection_16x", 4.0f/16.0f, SelectionColor_Pressed ) )
 			.SetCheckedImage( BOX_BRUSH( "Common/RoundedSelection_16x",  4.0f/16.0f, SelectionColor_Pressed ) )
@@ -552,10 +552,10 @@ void FSlateEditorStyle::FStyle::SetupGeneralStyles()
 		Set("GenericFilters.TextStyle", FTextBlockStyle(NormalText)
 			.SetFont(DEFAULT_FONT("Bold", 9))
 			.SetColorAndOpacity(FLinearColor(1.0f, 1.0f, 1.0f, 0.9f))
-			.SetShadowOffset(FVector2D(1, 1))
-			.SetShadowColorAndOpacity(FLinearColor(0, 0, 0, 0.9f)));
+			.SetShadowOffset(FVector2f::UnitVector)
+			.SetShadowColorAndOpacity(FLinearColor(0.f, 0.f, 0.f, 0.9f)));
 
-		ComboButton.SetMenuBorderPadding(FMargin(1.0));
+		ComboButton.SetMenuBorderPadding(FMargin(1.0f));
 
 		FComboBoxStyle ComboBox = FComboBoxStyle()
 			.SetComboButtonStyle(ComboButton);
@@ -629,7 +629,7 @@ void FSlateEditorStyle::FStyle::SetupGeneralStyles()
 	{
 		Set( "ErrorReporting.Box",  new BOX_BRUSH( "Common/TextBlockHighlightShape", FMargin(3.f/8.f) ));
 		Set( "ErrorReporting.EmptyBox",  new BOX_BRUSH( "Common/TextBlockHighlightShape_Empty", FMargin(3.f/8.f) ));
-		Set( "ErrorReporting.BackgroundColor", FLinearColor(0.35f,0,0) );
+		Set( "ErrorReporting.BackgroundColor", FLinearColor(0.35f,0.f,0.f) );
 		Set( "ErrorReporting.WarningBackgroundColor", FLinearColor(0.828f, 0.364f, 0.003f) );
 		Set( "InfoReporting.BackgroundColor", FLinearColor(0.1f, 0.33f, 1.0f));
 		Set( "ErrorReporting.ForegroundColor", FLinearColor::White );
@@ -638,10 +638,10 @@ void FSlateEditorStyle::FStyle::SetupGeneralStyles()
 
 	// Scrollbar
 	const FScrollBarStyle ScrollBar = FScrollBarStyle()
-		.SetVerticalTopSlotImage(IMAGE_BRUSH("Common/Scrollbar_Background_Vertical", FVector2D(8, 8)))
-		.SetVerticalBottomSlotImage(IMAGE_BRUSH("Common/Scrollbar_Background_Vertical", FVector2D(8, 8)))
-		.SetHorizontalTopSlotImage(IMAGE_BRUSH("Common/Scrollbar_Background_Horizontal", FVector2D(8, 8)))
-		.SetHorizontalBottomSlotImage(IMAGE_BRUSH("Common/Scrollbar_Background_Horizontal", FVector2D(8, 8)))
+		.SetVerticalTopSlotImage(IMAGE_BRUSH("Common/Scrollbar_Background_Vertical", Icon8x8))
+		.SetVerticalBottomSlotImage(IMAGE_BRUSH("Common/Scrollbar_Background_Vertical", Icon8x8))
+		.SetHorizontalTopSlotImage(IMAGE_BRUSH("Common/Scrollbar_Background_Horizontal", Icon8x8))
+		.SetHorizontalBottomSlotImage(IMAGE_BRUSH("Common/Scrollbar_Background_Horizontal", Icon8x8))
 		.SetNormalThumbImage( BOX_BRUSH( "Common/Scrollbar_Thumb", FMargin(4.f/16.f) ) )
 		.SetDraggedThumbImage( BOX_BRUSH( "Common/Scrollbar_Thumb", FMargin(4.f/16.f) ) )
 		.SetHoveredThumbImage( BOX_BRUSH( "Common/Scrollbar_Thumb", FMargin(4.f/16.f) ) );
@@ -687,7 +687,7 @@ void FSlateEditorStyle::FStyle::SetupGeneralStyles()
 		Set( "ProgressBar", FProgressBarStyle()
 			.SetBackgroundImage( BOX_BRUSH( "Common/ProgressBar_Background", FMargin(5.f/12.f) ) )
 			.SetFillImage( BOX_BRUSH( "Common/ProgressBar_Fill", FMargin(5.f/12.f), FLinearColor( 1.0f, 0.22f, 0.0f ) ) )
-			.SetMarqueeImage( IMAGE_BRUSH( "Common/ProgressBar_Marquee", FVector2D(20,12), FLinearColor::White, ESlateBrushTileType::Horizontal ) )
+			.SetMarqueeImage( IMAGE_BRUSH( "Common/ProgressBar_Marquee", FVector2f(20.f,12.f), FLinearColor::White, ESlateBrushTileType::Horizontal ) )
 			);
 
 		Set( "ProgressBar.ThinBackground", new BOX_BRUSH( "Common/ProgressBar_Thin_Background", FMargin(5.f/12.f) ) );
@@ -696,8 +696,8 @@ void FSlateEditorStyle::FStyle::SetupGeneralStyles()
 		// Legacy ProgressBar styles; kept here because other FAppStyle controls (mis)use them
 		// todo: jdale - Widgets using these styles should be updated to use SlateStyle types once FAppStyle has been obliterated from the Slate core
 		Set( "ProgressBar.Background", new BOX_BRUSH( "Common/ProgressBar_Background", FMargin(5.f/12.f) ) );
-		Set( "ProgressBar.Marquee", new IMAGE_BRUSH( "Common/ProgressBar_Marquee", FVector2D(20,12), FLinearColor::White, ESlateBrushTileType::Horizontal ) );
-		Set( "ProgressBar.BorderPadding", FVector2D(1,0) );
+		Set( "ProgressBar.Marquee", new IMAGE_BRUSH( "Common/ProgressBar_Marquee", FVector2f(20.f,12.f), FLinearColor::White, ESlateBrushTileType::Horizontal ) );
+		Set( "ProgressBar.BorderPadding", FVector2f(1.f,0.f) );
 	}
 
 	// WorkingBar
@@ -705,7 +705,7 @@ void FSlateEditorStyle::FStyle::SetupGeneralStyles()
 		Set( "WorkingBar", FProgressBarStyle()
 			.SetBackgroundImage( FSlateNoResource() )
 			.SetFillImage( BOX_BRUSH( "Common/ProgressBar_Fill", FMargin(5.f/12.f), FLinearColor( 1.0f, 0.22f, 0.0f ) ) )
-			.SetMarqueeImage( IMAGE_BRUSH( "Common/WorkingBar_Marquee", FVector2D(20,2), FLinearColor::White, ESlateBrushTileType::Horizontal ) )
+			.SetMarqueeImage( IMAGE_BRUSH( "Common/WorkingBar_Marquee", FVector2f(20.f,2.f), FLinearColor::White, ESlateBrushTileType::Horizontal ) )
 			);
 	}
 
@@ -719,15 +719,15 @@ void FSlateEditorStyle::FStyle::SetupGeneralStyles()
 	// Filtering/Searching feedback
 	{
 		const FLinearColor ActiveFilterColor = FLinearColor(1.0f,0.55f,0.0f,1.0f);
-		Set("Searching.SearchActiveTab", new BOX_BRUSH("Common/SearchPseudoTab", FVector2D(16,16), FMargin(0.49f), ActiveFilterColor));
-		Set("Searching.SearchActiveBorder", new BOX_BRUSH("Common/SearchActiveBorder", FVector2D(8,8), FMargin(0.49f), ActiveFilterColor));
+		Set("Searching.SearchActiveTab", new BOX_BRUSH("Common/SearchPseudoTab", Icon16x16, FMargin(0.49f), ActiveFilterColor));
+		Set("Searching.SearchActiveBorder", new BOX_BRUSH("Common/SearchActiveBorder", Icon8x8, FMargin(0.49f), ActiveFilterColor));
 	}
 
 	// Inline Editable Text Block
 	{
 		FTextBlockStyle InlineEditableTextBlockReadOnly = FTextBlockStyle(NormalText)
 			.SetColorAndOpacity( FSlateColor::UseForeground() )
-			.SetShadowOffset( FVector2D::ZeroVector )
+			.SetShadowOffset( FVector2f::ZeroVector )
 			.SetShadowColorAndOpacity( FLinearColor::Black );
 
 		FEditableTextBoxStyle InlineEditableTextBlockEditable = FEditableTextBoxStyle()
@@ -746,7 +746,7 @@ void FSlateEditorStyle::FStyle::SetupGeneralStyles()
 	}
 
 	// Images sizes are specified in Slate Screen Units. These do not necessarily correspond to pixels!
-	// An IMAGE_BRUSH( "SomeImage", FVector2D(32,32)) will have a desired size of 16x16 Slate Screen Units.
+	// An IMAGE_BRUSH( "SomeImage", Icon32x32) will have a desired size of 16x16 Slate Screen Units.
 	// This allows the original resource to be scaled up or down as needed.
 
 	Set( "WhiteTexture", new IMAGE_BRUSH( "Old/White", Icon16x16 ) );
@@ -789,9 +789,9 @@ void FSlateEditorStyle::FStyle::SetupGeneralStyles()
 	{
 		// Dark Hyperlink - for use on light backgrounds
 		FButtonStyle DarkHyperlinkButton = FButtonStyle()
-			.SetNormal ( BORDER_BRUSH( "Old/HyperlinkDotted", FMargin(0,0,0,3/16.0f), FLinearColor::Black ) )
+			.SetNormal ( BORDER_BRUSH( "Old/HyperlinkDotted", FMargin(0.f,0.f,0.f,3.f/16.0f), FLinearColor::Black ) )
 			.SetPressed( FSlateNoResource() )
-			.SetHovered( BORDER_BRUSH( "Old/HyperlinkUnderline", FMargin(0,0,0,3/16.0f), FLinearColor::Black ) );
+			.SetHovered( BORDER_BRUSH( "Old/HyperlinkUnderline", FMargin(0.f,0.f,0.f,3.f/16.0f), FLinearColor::Black ) );
 		FHyperlinkStyle DarkHyperlink = FHyperlinkStyle()
 			.SetUnderlineStyle(DarkHyperlinkButton)
 			.SetTextStyle(NormalText)
@@ -802,7 +802,7 @@ void FSlateEditorStyle::FStyle::SetupGeneralStyles()
 		FButtonStyle HoverOnlyHyperlinkButton = FButtonStyle()
 			.SetNormal(FSlateNoResource() )
 			.SetPressed(FSlateNoResource() )
-			.SetHovered(BORDER_BRUSH( "Old/HyperlinkUnderline", FMargin(0,0,0,3/16.0f) ) );
+			.SetHovered(BORDER_BRUSH( "Old/HyperlinkUnderline", FMargin(0.f,0.f,0.f,3.f/16.0f) ) );
 		Set("HoverOnlyHyperlinkButton", HoverOnlyHyperlinkButton);
 
 		FHyperlinkStyle HoverOnlyHyperlink = FHyperlinkStyle()
@@ -862,12 +862,12 @@ void FSlateEditorStyle::FStyle::SetupGeneralStyles()
 	{
 		Set( "MaterialList.DragDropBorder", new BOX_BRUSH( "Old/Window/ViewportDebugBorder", 0.8f ) );
 		Set( "MaterialList.HyperlinkStyle", FTextBlockStyle(NormalText) .SetFont( DEFAULT_FONT( "Regular", 8 ) ) );
-		Set( "MaterialList.HyperlinkStyle.ShadowOffset", FVector2D::ZeroVector );
+		Set( "MaterialList.HyperlinkStyle.ShadowOffset", FVector2f::ZeroVector );
 	}
 
 	// Dialogue Wave Details
 	{
-		Set( "DialogueWaveDetails.SpeakerToTarget", new IMAGE_BRUSH( "PropertyView/SpeakerToTarget", FVector2D(30.0f, 30.0f) ) );
+		Set( "DialogueWaveDetails.SpeakerToTarget", new IMAGE_BRUSH( "PropertyView/SpeakerToTarget", FVector2f(30.0f, 30.0f) ) );
 		Set( "DialogueWaveDetails.HeaderBorder", new BOX_BRUSH( "Common/MenuBarBorder", FMargin(4.0f/16.0f) ) );
 		Set( "DialogueWaveDetails.PropertyEditorMenu", new BOX_BRUSH( "Old/Menu_Background", FMargin(8.0f/64.0f) ) );
 	}
@@ -902,8 +902,8 @@ void FSlateEditorStyle::FStyle::SetupGeneralStyles()
 	// Scroll Box
 	{
 		Set( "ScrollBox", FScrollBoxStyle()
-			.SetTopShadowBrush( IMAGE_BRUSH( "Common/ScrollBoxShadowTop", FVector2D(64,8) ) )
-			.SetBottomShadowBrush( IMAGE_BRUSH( "Common/ScrollBoxShadowBottom", FVector2D(64,8) ) )
+			.SetTopShadowBrush( IMAGE_BRUSH( "Common/ScrollBoxShadowTop", FVector2f(64.f,8.f) ) )
+			.SetBottomShadowBrush( IMAGE_BRUSH( "Common/ScrollBoxShadowBottom", FVector2f(64.f,8.f) ) )
 			);
 	}//
 
@@ -911,10 +911,10 @@ void FSlateEditorStyle::FStyle::SetupGeneralStyles()
 	{
 		Set( "TableView.Row", FTableRowStyle( NormalTableRowStyle) );
 		Set( "TableView.DarkRow",FTableRowStyle( NormalTableRowStyle)
-			.SetEvenRowBackgroundBrush(IMAGE_BRUSH("PropertyView/DetailCategoryMiddle", FVector2D(16, 16)))
-			.SetEvenRowBackgroundHoveredBrush(IMAGE_BRUSH("PropertyView/DetailCategoryMiddle_Hovered", FVector2D(16, 16)))
-			.SetOddRowBackgroundBrush(IMAGE_BRUSH("PropertyView/DetailCategoryMiddle", FVector2D(16, 16)))
-			.SetOddRowBackgroundHoveredBrush(IMAGE_BRUSH("PropertyView/DetailCategoryMiddle_Hovered", FVector2D(16, 16)))
+			.SetEvenRowBackgroundBrush(IMAGE_BRUSH("PropertyView/DetailCategoryMiddle", Icon16x16))
+			.SetEvenRowBackgroundHoveredBrush(IMAGE_BRUSH("PropertyView/DetailCategoryMiddle_Hovered", Icon16x16))
+			.SetOddRowBackgroundBrush(IMAGE_BRUSH("PropertyView/DetailCategoryMiddle", Icon16x16))
+			.SetOddRowBackgroundHoveredBrush(IMAGE_BRUSH("PropertyView/DetailCategoryMiddle_Hovered", Icon16x16))
 			.SetSelectorFocusedBrush(BORDER_BRUSH("Common/Selector", FMargin(4.f / 16.f), SelectorColor))
 			.SetActiveBrush(IMAGE_BRUSH("Common/Selection", Icon8x8, SelectionColor))
 			.SetActiveHoveredBrush(IMAGE_BRUSH("Common/Selection", Icon8x8, SelectionColor))
@@ -1006,21 +1006,21 @@ void FSlateEditorStyle::FStyle::SetupGeneralStyles()
 
 	// Throbber
 	{
-		Set( "Throbber.Chunk", new IMAGE_BRUSH( "Old/Throbber/Throbber_Piece", FVector2D(16,16) ) );
-		Set( "Throbber.CircleChunk", new IMAGE_BRUSH( "Old/Throbber/Throbber_Piece", FVector2D(8,8) ) );
-		Set( "SmallThrobber.Chunk", new IMAGE_BRUSH( "Common/ThrobberPiece_Small", FVector2D(8,16) ) );
+		Set( "Throbber.Chunk", new IMAGE_BRUSH( "Old/Throbber/Throbber_Piece", Icon16x16 ) );
+		Set( "Throbber.CircleChunk", new IMAGE_BRUSH( "Old/Throbber/Throbber_Piece", Icon8x8 ) );
+		Set( "SmallThrobber.Chunk", new IMAGE_BRUSH( "Common/ThrobberPiece_Small", FVector2f(8.f,16.f) ) );
 	}
 
 	{
-		Set("CurveEd.TimelineArea", new IMAGE_BRUSH("Old/White", Icon16x16, FLinearColor(1, 1, 1, 0.25f)));
+		Set("CurveEd.TimelineArea", new IMAGE_BRUSH("Old/White", Icon16x16, FLinearColor(1.f, 1.f, 1.f, 0.25f)));
 		Set("CurveEd.FitHorizontal", new IMAGE_BRUSH("Icons/FitHorz_16x", Icon16x16));
 		Set("CurveEd.FitVertical", new IMAGE_BRUSH("Icons/FitVert_16x", Icon16x16));
-		Set("CurveEd.CurveKey", new IMAGE_BRUSH("Common/Key", FVector2D(11.0f, 11.0f)));
-		Set("CurveEd.CurveKeySelected", new IMAGE_BRUSH("Common/Key", FVector2D(11.0f, 11.0f), SelectionColor));
+		Set("CurveEd.CurveKey", new IMAGE_BRUSH("Common/Key", FVector2f(11.0f, 11.0f)));
+		Set("CurveEd.CurveKeySelected", new IMAGE_BRUSH("Common/Key", FVector2f(11.0f, 11.0f), SelectionColor));
 		Set("CurveEd.InfoFont", DEFAULT_FONT("Regular", 8));
 		Set("CurveEd.LabelFont", DEFAULT_FONT("Bold", 10));
-		Set("CurveEd.Tangent", new IMAGE_BRUSH("Common/Tangent", FVector2D(7.0f, 7.0f), FLinearColor(0.0f, 0.66f, 0.7f)));
-		Set("CurveEd.TangentSelected", new IMAGE_BRUSH("Common/Tangent", FVector2D(7.0f, 7.0f), FLinearColor(1.0f, 1.0f, 0.0f)));
+		Set("CurveEd.Tangent", new IMAGE_BRUSH("Common/Tangent", FVector2f(7.0f, 7.0f), FLinearColor(0.0f, 0.66f, 0.7f)));
+		Set("CurveEd.TangentSelected", new IMAGE_BRUSH("Common/Tangent", FVector2f(7.0f, 7.0f), FLinearColor(1.0f, 1.0f, 0.0f)));
 		Set("CurveEd.TangentColor", FLinearColor(0.0f, 0.66f, 0.7f));
 		Set("CurveEd.TangentColorSelected", FLinearColor(1.0f, 1.0f, 0.0f));
 		Set("CurveEd.Visible", new IMAGE_BRUSH("Icons/icon_levels_visible_16px", Icon16x16));
@@ -1112,7 +1112,7 @@ void FSlateEditorStyle::FStyle::SetupGeneralStyles()
 	{
 		Set( "MessageLog", FTextBlockStyle(NormalText)
 			.SetFont( DEFAULT_FONT( "Regular", 8 ) )
-			.SetShadowOffset( FVector2D::ZeroVector )
+			.SetShadowOffset( FVector2f::ZeroVector )
 		);
 
 		Set( "MessageLog.Action", new IMAGE_BRUSH( "Icons/icon_file_choosepackages_16px", Icon16x16) );
@@ -1138,7 +1138,7 @@ void FSlateEditorStyle::FStyle::SetupGeneralStyles()
 		Set( "AnimEditor.FilterSearch", new IMAGE_BRUSH( "Old/FilterSearch", Icon16x16 ) );
 		Set( "AnimEditor.FilterCancel", new IMAGE_BRUSH( "Old/FilterCancel", Icon16x16 ) );
 
-		Set( "AnimEditor.NotifyGraphBackground", new IMAGE_BRUSH( "Old/AnimEditor/NotifyTrackBackground", FVector2D(64, 64), FLinearColor::White, ESlateBrushTileType::Both) );
+		Set( "AnimEditor.NotifyGraphBackground", new IMAGE_BRUSH( "Old/AnimEditor/NotifyTrackBackground", Icon64x64, FLinearColor::White, ESlateBrushTileType::Both) );
 
 		Set( "BlendSpace.SamplePoint", new IMAGE_BRUSH( "Old/AnimEditor/BlendSpace_Sample", Icon16x16 ) );
 		Set( "BlendSpace.SamplePoint_Highlight", new IMAGE_BRUSH( "Old/AnimEditor/BlendSpace_Sample_Highlight", Icon16x16 ) );
@@ -1150,10 +1150,10 @@ void FSlateEditorStyle::FStyle::SetupGeneralStyles()
 
 
 	// Gamma reference.
-	Set("GammaReference",new IMAGE_BRUSH( "Old/GammaReference",FVector2D(256,128)));
+	Set("GammaReference",new IMAGE_BRUSH( "Old/GammaReference",FVector2f(256.f,128.f)));
 	
-	Set("TrashCan", new IMAGE_BRUSH( "Old/TrashCan", FVector2D(64, 64)));
-	Set("TrashCan_Small", new IMAGE_BRUSH( "Old/TrashCan_Small", FVector2D(18, 18)));
+	Set("TrashCan", new IMAGE_BRUSH( "Old/TrashCan", Icon64x64));
+	Set("TrashCan_Small", new IMAGE_BRUSH( "Old/TrashCan_Small", FVector2f(18.f, 18.f)));
 
 	// Spline component controls
 	{
@@ -1206,7 +1206,7 @@ void FSlateEditorStyle::FStyle::SetupGeneralStyles()
 	Set( "EmbossedText", FTextBlockStyle(NormalText)
 		.SetFont( DEFAULT_FONT( "Regular", 24 ) )
 		.SetColorAndOpacity( FLinearColor::Black )
-		.SetShadowOffset( FVector2D( 0,1 ) )
+		.SetShadowOffset( FVector2f( 0.f,1.f) )
 		.SetShadowColorAndOpacity(FLinearColor(0.8f,0.8f,0.8f, 0.5) )
 	);
 
@@ -1258,7 +1258,7 @@ void FSlateEditorStyle::FStyle::SetupGeneralStyles()
 			.SetFont(DEFAULT_FONT( "Bold", 12))
 			.SetColorAndOpacity(FLinearColor(1.0f, 1.0f, 1.0f))
 			.SetHighlightColor(FLinearColor(1.0f, 1.0f, 1.0f))
-			.SetShadowOffset(FVector2D(1.0f, 1.0f))
+			.SetShadowOffset(FVector2f::UnitVector)
 			.SetShadowColorAndOpacity(FLinearColor(0.0f, 0.0f, 0.0f, 0.9f)));
 	}
 
@@ -1296,29 +1296,29 @@ void FSlateEditorStyle::FStyle::SetupGeneralStyles()
 		Set( "InputBindingEditor.SmallFont", DEFAULT_FONT( "Regular", 8 ) );
 
 		Set( "InputBindingEditor.HeaderButton", FButtonStyle(NoBorder)
-			.SetNormalPadding(FMargin( 1,1,2,2 ))
-			.SetPressedPadding(FMargin( 2,2,2,2 )) );
+			.SetNormalPadding(FMargin( 1.f,1.f,2.f,2.f))
+			.SetPressedPadding(FMargin( 2.f,2.f,2.f,2.f)) );
 
 		Set( "InputBindingEditor.HeaderButton.Disabled", new FSlateNoResource() );
 
 
-		Set( "InputBindingEditor.Tab",  new IMAGE_BRUSH( "Icons/icon_tab_KeyBindings_16px", FVector2D( 16, 16 ) ) );
-		Set( "InputBindingEditor.AssetEditor",  new IMAGE_BRUSH( "Icons/icon_keyb_AssetEditor_16px", FVector2D( 16, 16 ) ) );
-		Set( "InputBindingEditor.AssetEditor",  new IMAGE_BRUSH( "Icons/icon_keyb_AssetEditor_16px", FVector2D( 16, 16 ) ) );
-		Set( "InputBindingEditor.GenericCommands",  new IMAGE_BRUSH( "Icons/icon_keyb_CommonCommands_16px", FVector2D( 16, 16 ) ) );
-		Set( "InputBindingEditor.FoliageEditMode",  new IMAGE_BRUSH( "Icons/icon_keyb_FoliageEditMode_16px", FVector2D( 16, 16 ) ) );
-		Set( "InputBindingEditor.LandscapeEditor",  new IMAGE_BRUSH( "Icons/icon_keyb_LandscapeEditor_16px", FVector2D( 16, 16 ) ) );
-		Set( "InputBindingEditor.LayersView",  new IMAGE_BRUSH( "Icons/icon_keyb_Layers_16px", FVector2D( 16, 16 ) ) );
-		Set( "InputBindingEditor.LevelEditor",  new IMAGE_BRUSH( "Icons/icon_keyb_LevelEditor_16px", FVector2D( 16, 16 ) ) );
-		Set( "InputBindingEditor.LevelViewport",  new IMAGE_BRUSH( "Icons/icon_keyb_LevelViewports_16px", FVector2D( 16, 16 ) ) );
-		Set( "InputBindingEditor.MainFrame",  new IMAGE_BRUSH( "Icons/icon_keyb_MainFrame_16px", FVector2D( 16, 16 ) ) );
-		Set( "InputBindingEditor.OutputLog",  new IMAGE_BRUSH( "Icons/icon_keyb_OutputLog_16px", FVector2D( 16, 16 ) ) );
-		Set( "InputBindingEditor.PlayWorld",  new IMAGE_BRUSH( "Icons/icon_keyb_PlayWorld_16px", FVector2D( 16, 16 ) ) );		
+		Set( "InputBindingEditor.Tab",  new IMAGE_BRUSH( "Icons/icon_tab_KeyBindings_16px", Icon16x16 ) );
+		Set( "InputBindingEditor.AssetEditor",  new IMAGE_BRUSH( "Icons/icon_keyb_AssetEditor_16px", Icon16x16 ) );
+		Set( "InputBindingEditor.AssetEditor",  new IMAGE_BRUSH( "Icons/icon_keyb_AssetEditor_16px", Icon16x16 ) );
+		Set( "InputBindingEditor.GenericCommands",  new IMAGE_BRUSH( "Icons/icon_keyb_CommonCommands_16px", Icon16x16 ) );
+		Set( "InputBindingEditor.FoliageEditMode",  new IMAGE_BRUSH( "Icons/icon_keyb_FoliageEditMode_16px", Icon16x16 ) );
+		Set( "InputBindingEditor.LandscapeEditor",  new IMAGE_BRUSH( "Icons/icon_keyb_LandscapeEditor_16px", Icon16x16 ) );
+		Set( "InputBindingEditor.LayersView",  new IMAGE_BRUSH( "Icons/icon_keyb_Layers_16px", Icon16x16 ) );
+		Set( "InputBindingEditor.LevelEditor",  new IMAGE_BRUSH( "Icons/icon_keyb_LevelEditor_16px", Icon16x16 ) );
+		Set( "InputBindingEditor.LevelViewport",  new IMAGE_BRUSH( "Icons/icon_keyb_LevelViewports_16px", Icon16x16 ) );
+		Set( "InputBindingEditor.MainFrame",  new IMAGE_BRUSH( "Icons/icon_keyb_MainFrame_16px", Icon16x16 ) );
+		Set( "InputBindingEditor.OutputLog",  new IMAGE_BRUSH( "Icons/icon_keyb_OutputLog_16px", Icon16x16 ) );
+		Set( "InputBindingEditor.PlayWorld",  new IMAGE_BRUSH( "Icons/icon_keyb_PlayWorld_16px", Icon16x16 ) );		
 	}
 
 	// Package restore
 	{
-		Set( "PackageRestore.FolderOpen", new IMAGE_BRUSH( "Icons/FolderOpen", FVector2D(18, 16) ) );
+		Set( "PackageRestore.FolderOpen", new IMAGE_BRUSH( "Icons/FolderOpen", FVector2f(18.f, 16.f) ) );
 	}
 #endif // WITH_EDITOR || (IS_PROGRAM && WITH_UNREAL_DEVELOPER_TOOLS)
 
@@ -1339,9 +1339,9 @@ void FSlateEditorStyle::FStyle::SetupGeneralStyles()
 	// Package Dialog
 	{
 		Set( "PackageDialog.ListHeader", new BOX_BRUSH( "Old/SavePackages/ListHeader", 4.0f/32.0f ) );
-		Set( "SavePackages.SCC_DlgCheckedOutOther", new IMAGE_BRUSH( "Old/SavePackages/SCC_DlgCheckedOutOther", FVector2D( 18, 16 ) ) );
-		Set( "SavePackages.SCC_DlgNotCurrent", new IMAGE_BRUSH( "Old/SavePackages/SCC_DlgNotCurrent", FVector2D( 18, 16 ) ) );
-		Set( "SavePackages.SCC_DlgReadOnly", new IMAGE_BRUSH( "Old/SavePackages/SCC_DlgReadOnly", FVector2D( 18, 16 ) ) );
+		Set( "SavePackages.SCC_DlgCheckedOutOther", new IMAGE_BRUSH( "Old/SavePackages/SCC_DlgCheckedOutOther", FVector2f( 18.f, 16.f) ) );
+		Set( "SavePackages.SCC_DlgNotCurrent", new IMAGE_BRUSH( "Old/SavePackages/SCC_DlgNotCurrent", FVector2f( 18.f, 16.f) ) );
+		Set( "SavePackages.SCC_DlgReadOnly", new IMAGE_BRUSH( "Old/SavePackages/SCC_DlgReadOnly", FVector2f( 18.f, 16.f) ) );
 		Set( "SavePackages.SCC_DlgNoIcon", new FSlateNoResource() );
 	}
 #endif // WITH_EDITOR || (IS_PROGRAM && WITH_UNREAL_DEVELOPER_TOOLS)
@@ -1417,14 +1417,14 @@ void FSlateEditorStyle::FStyle::SetupGeneralStyles()
 			.SetNormal        ( FSlateNoResource() )
 			.SetPressed       ( BOX_BRUSH( "Common/Button_Pressed", 8.0f/32.0f, SelectionColor_Pressed ) )
 			.SetHovered       ( BOX_BRUSH( "Common/Button_Hovered", 8.0f/32.0f, SelectionColor ) )
-			.SetNormalPadding ( FMargin( 2,2,4,4 ) )
-			.SetPressedPadding( FMargin( 3,3,3,3 ) )
+			.SetNormalPadding ( FMargin( 2.f,2.f,4.f,4.f) )
+			.SetPressedPadding( FMargin( 3.f,3.f,3.f,3.f) )
 		);
 
 		Set( "GraphBreadcrumbButtonText", FTextBlockStyle(NormalText)
 			.SetFont( DEFAULT_FONT( "Regular", 14 ) )
-			.SetColorAndOpacity( FLinearColor(1,1,1,0.5) )
-			.SetShadowOffset( FVector2D::ZeroVector )
+			.SetColorAndOpacity( FLinearColor(1.f,1.f,1.f,0.5f) )
+			.SetShadowOffset( FVector2f::ZeroVector )
 		);
 
 		Set( "GraphBreadcrumb.BrowseBack", new IMAGE_BRUSH( "Icons/icon_BlueprintBrowserL_24x", Icon24x24) );
@@ -1447,7 +1447,7 @@ void FSlateEditorStyle::FStyle::SetupGeneralStyles()
 		Set("GraphBookmarkButtonText", FTextBlockStyle(NormalText)
 			.SetFont(REGULAR_ICON_FONT(16))
 			.SetColorAndOpacity(FLinearColor(0.5f, 0.5f, 0.5f))
-			.SetShadowOffset(FVector2D::ZeroVector)
+			.SetShadowOffset(FVector2f::ZeroVector)
 		);
 
 		Set("GraphBookmarkMenuImage.Button_Add", new IMAGE_BRUSH("Icons/PlusSymbol_12x", Icon12x12));
@@ -1497,19 +1497,19 @@ void FSlateEditorStyle::FStyle::SetupGeneralStyles()
 		
 	// Asset Thumbnail
 	{
-		Set( "AssetThumbnail.AssetBackground", new IMAGE_BRUSH( "Common/AssetBackground", FVector2D(64.f, 64.f), FLinearColor(1.f, 1.f, 1.f, 1.0f) ) );
-		Set( "AssetThumbnail.ClassBackground", new IMAGE_BRUSH( "Common/ClassBackground_64x", FVector2D(64.f, 64.f), FLinearColor(0.75f, 0.75f, 0.75f, 1.0f) ) );
-		Set( "AssetThumbnail.DataOnlyBPAssetBackground", new IMAGE_BRUSH( "Common/DataOnlyBPAssetBackground_64x", FVector2D(64.f, 64.f), FLinearColor(1, 1, 1, 1) ) );
+		Set( "AssetThumbnail.AssetBackground", new IMAGE_BRUSH( "Common/AssetBackground", Icon64x64, FLinearColor(1.f, 1.f, 1.f, 1.0f) ) );
+		Set( "AssetThumbnail.ClassBackground", new IMAGE_BRUSH( "Common/ClassBackground_64x", Icon64x64, FLinearColor(0.75f, 0.75f, 0.75f, 1.0f) ) );
+		Set( "AssetThumbnail.DataOnlyBPAssetBackground", new IMAGE_BRUSH( "Common/DataOnlyBPAssetBackground_64x", Icon64x64, FLinearColor::White ) );
 		Set( "AssetThumbnail.Font", DEFAULT_FONT( "Regular", 10 ) );
 		Set( "AssetThumbnail.FontSmall", DEFAULT_FONT( "Regular", 7 ) );
 		Set( "AssetThumbnail.ColorAndOpacity", FLinearColor(0.75f, 0.75f, 0.75f, 1) );
-		Set( "AssetThumbnail.ShadowOffset", FVector2D(1,1) );
-		Set( "AssetThumbnail.ShadowColorAndOpacity", FLinearColor(0, 0, 0, 0.5) );
+		Set( "AssetThumbnail.ShadowOffset", FVector2f::UnitVector );
+		Set( "AssetThumbnail.ShadowColorAndOpacity", FLinearColor(0.f, 0.f, 0.f, 0.5f) );
 		Set( "AssetThumbnail.HintFont", DEFAULT_FONT( "Regular", 8 ) );
 		Set( "AssetThumbnail.HintFontSmall", DEFAULT_FONT( "Regular", 6 ) );
-		Set( "AssetThumbnail.HintColorAndOpacity", FLinearColor(0.75f, 0.75f, 0.75f, 1) );
-		Set( "AssetThumbnail.HintShadowOffset", FVector2D(1,1) );
-		Set( "AssetThumbnail.HintShadowColorAndOpacity", FLinearColor(0, 0, 0, 0.5) );
+		Set( "AssetThumbnail.HintColorAndOpacity", FLinearColor(0.75f, 0.75f, 0.75f, 1.f) );
+		Set( "AssetThumbnail.HintShadowOffset", FVector2f::UnitVector );
+		Set( "AssetThumbnail.HintShadowColorAndOpacity", FLinearColor(0.f, 0.f, 0.f, 0.5f) );
 		Set( "AssetThumbnail.HintBackground", new BOX_BRUSH( "Common/TableViewHeader", FMargin(8.0f/32.0f) ) );
 		Set( "AssetThumbnail.Border", new FSlateColorBrush( FColor::White ) );
 	}
@@ -1519,12 +1519,12 @@ void FSlateEditorStyle::FStyle::SetupGeneralStyles()
 		Set( "SystemWideCommands.SummonOpenAssetDialog", new IMAGE_BRUSH( "Icons/icon_asset_open_16px", Icon16x16 ) );
 	
 		Set( "GlobalAssetPicker.Background", new BOX_BRUSH( "Old/Menu_Background", FMargin(8.0f/64.0f) ) );
-		Set( "GlobalAssetPicker.OutermostMargin", FMargin(4, 4, 4, 4) );
+		Set( "GlobalAssetPicker.OutermostMargin", FMargin(4.f, 4.f, 4.f, 4.f) );
 
 		Set( "GlobalAssetPicker.TitleFont", FTextBlockStyle(NormalText)
 				.SetFont( DEFAULT_FONT( "Regular", 9 ) )
 			.SetColorAndOpacity( FLinearColor::White )
-				.SetShadowOffset( FVector2D( 1,1 ) )
+				.SetShadowOffset( FVector2f::UnitVector )
 				.SetShadowColorAndOpacity( FLinearColor::Black )
 			);
 	}
@@ -1568,8 +1568,8 @@ void FSlateEditorStyle::FStyle::SetupGeneralStyles()
 			.SetNormal( IMAGE_BRUSH( "Icons/StatusInfo_16x", Icon16x16 ) )
 			.SetHovered( IMAGE_BRUSH( "Icons/StatusInfo_16x", Icon16x16 ) )
 			.SetPressed( IMAGE_BRUSH( "Icons/StatusInfo_16x", Icon16x16 ) )
-			.SetNormalPadding(0)
-			.SetPressedPadding(0)
+			.SetNormalPadding(0.f)
+			.SetPressedPadding(0.f)
 		);
 	}
 
@@ -1591,13 +1591,13 @@ void FSlateEditorStyle::FStyle::SetupGeneralStyles()
 	// About screen
 	if (IncludeEditorSpecificStyles())
 	{
-		Set( "AboutScreen.Background", new IMAGE_BRUSH( "About/Background", FVector2D(600,332), FLinearColor::White, ESlateBrushTileType::Both) );
-		Set( "AboutScreen.Facebook", new IMAGE_BRUSH( "About/FacebookIcon", FVector2D(35,35) ) );
-		Set( "AboutScreen.FacebookHovered", new IMAGE_BRUSH( "About/FacebookIcon_Hovered", FVector2D(35,35) ) );
-		Set( "AboutScreen.UE4", new IMAGE_BRUSH( "About/UE4Icon", FVector2D(50,50), FLinearColor::Gray) );
-		Set( "AboutScreen.UE4Hovered", new IMAGE_BRUSH( "About/UE4Icon", FVector2D(50,50), FLinearColor::White) );
-		Set( "AboutScreen.EpicGames", new IMAGE_BRUSH( "About/EpicGamesIcon", FVector2D(50,50) ) );
-		Set( "AboutScreen.EpicGamesHovered", new IMAGE_BRUSH( "About/EpicGamesIcon_Hovered", FVector2D(50,50) ) );
+		Set( "AboutScreen.Background", new IMAGE_BRUSH( "About/Background", FVector2f(600.f,332.f), FLinearColor::White, ESlateBrushTileType::Both) );
+		Set( "AboutScreen.Facebook", new IMAGE_BRUSH( "About/FacebookIcon", FVector2f(35.f,35.f) ) );
+		Set( "AboutScreen.FacebookHovered", new IMAGE_BRUSH( "About/FacebookIcon_Hovered", FVector2f(35.f,35.f) ) );
+		Set( "AboutScreen.UE4", new IMAGE_BRUSH( "About/UE4Icon", FVector2f(50.f,50.f), FLinearColor::Gray) );
+		Set( "AboutScreen.UE4Hovered", new IMAGE_BRUSH( "About/UE4Icon", FVector2f(50.f,50.f), FLinearColor::White) );
+		Set( "AboutScreen.EpicGames", new IMAGE_BRUSH( "About/EpicGamesIcon", FVector2f(50.f,50.f) ) );
+		Set( "AboutScreen.EpicGamesHovered", new IMAGE_BRUSH( "About/EpicGamesIcon_Hovered", FVector2f(50.f,50.f) ) );
 	}
 #endif // WITH_EDITOR || (IS_PROGRAM && WITH_UNREAL_DEVELOPER_TOOLS)
 
@@ -1618,43 +1618,43 @@ void FSlateEditorStyle::FStyle::SetupGeneralStyles()
 
 		FTextBlockStyle CreditsNormal = FTextBlockStyle(NormalText)
 			.SetFont(DEFAULT_FONT("Regular", 16))
-			.SetShadowOffset(FVector2D::UnitVector);
+			.SetShadowOffset(FVector2f::UnitVector);
 
 		Set("Credits.Normal", CreditsNormal);
 
 		Set("Credits.Strong", FTextBlockStyle(CreditsNormal)
 			.SetFont(DEFAULT_FONT("Bold", 16))
-			.SetShadowOffset(FVector2D::UnitVector));
+			.SetShadowOffset(FVector2f::UnitVector));
 
 		Set("Credits.H1", FTextBlockStyle(CreditsNormal)
 			.SetColorAndOpacity(EditorOrange)
 			.SetFont(DEFAULT_FONT("Bold", 36))
-			.SetShadowOffset(FVector2D::UnitVector));
+			.SetShadowOffset(FVector2f::UnitVector));
 
 		Set("Credits.H2", FTextBlockStyle(CreditsNormal)
 			.SetColorAndOpacity(EditorOrange)
 			.SetFont(DEFAULT_FONT("Bold", 30))
-			.SetShadowOffset(FVector2D::UnitVector));
+			.SetShadowOffset(FVector2f::UnitVector));
 
 		Set("Credits.H3", FTextBlockStyle(CreditsNormal)
 			.SetFont(DEFAULT_FONT("Bold", 24))
-			.SetShadowOffset(FVector2D::UnitVector));
+			.SetShadowOffset(FVector2f::UnitVector));
 
 		Set("Credits.H4", FTextBlockStyle(CreditsNormal)
 			.SetFont(DEFAULT_FONT("Bold", 18))
-			.SetShadowOffset(FVector2D::UnitVector));
+			.SetShadowOffset(FVector2f::UnitVector));
 
 		Set("Credits.H5", FTextBlockStyle(CreditsNormal)
 			.SetFont(DEFAULT_FONT("Bold", 12))
-			.SetShadowOffset(FVector2D::UnitVector));
+			.SetShadowOffset(FVector2f::UnitVector));
 
 		Set("Credits.H6", FTextBlockStyle(CreditsNormal)
 			.SetFont(DEFAULT_FONT("Bold", 6))
-			.SetShadowOffset(FVector2D::UnitVector));
+			.SetShadowOffset(FVector2f::UnitVector));
 
 		FTextBlockStyle LinkText = FTextBlockStyle(NormalText)
 			.SetColorAndOpacity(EditorOrange)
-			.SetShadowOffset(FVector2D::UnitVector);
+			.SetShadowOffset(FVector2f::UnitVector);
 		FButtonStyle HoverOnlyHyperlinkButton = FButtonStyle()
 			.SetNormal(FSlateNoResource())
 			.SetPressed(FSlateNoResource())
@@ -1681,7 +1681,7 @@ void FSlateEditorStyle::FStyle::SetupGeneralStyles()
 		Set("HardwareTargets.Strong", FTextBlockStyle(TargetSettingsNormal)
 			.SetFont(DEFAULT_FONT("Bold", 8))
 			.SetColorAndOpacity(EditorOrange)
-			.SetShadowOffset(FVector2D::UnitVector));
+			.SetShadowOffset(FVector2f::UnitVector));
 	}
 #endif
 
@@ -1689,9 +1689,9 @@ void FSlateEditorStyle::FStyle::SetupGeneralStyles()
 #if WITH_EDITOR || (IS_PROGRAM && WITH_UNREAL_DEVELOPER_TOOLS)
 	{
 		Set( "NewLevelDialog.BlackBorder", new FSlateColorBrush( FColor(0, 0, 0, 100) ) );
-		Set( "NewLevelDialog.Blank", new IMAGE_BRUSH( "NewLevels/NewLevelBlank", FVector2D(256,256) ) );
-		Set( "NewLevelDialog.BlankWP", new IMAGE_BRUSH("NewLevels/NewLevelBlankWP", FVector2D(256, 256) ) );
-		Set( "NewLevelDialog.Default", new IMAGE_BRUSH( "NewLevels/NewLevelDefault", FVector2D(256,256) ) );
+		Set( "NewLevelDialog.Blank", new IMAGE_BRUSH( "NewLevels/NewLevelBlank", FVector2f(256.f,256.f) ) );
+		Set( "NewLevelDialog.BlankWP", new IMAGE_BRUSH("NewLevels/NewLevelBlankWP", FVector2f(256.f, 256.f) ) );
+		Set( "NewLevelDialog.Default", new IMAGE_BRUSH( "NewLevels/NewLevelDefault", FVector2f(256.f,256.f) ) );
 	}
 
 	// Build and Submit
@@ -1739,7 +1739,7 @@ void FSlateEditorStyle::FStyle::SetupGeneralStyles()
 
 		Set("FoliageEditToolBar.Heading", FTextBlockStyle(NormalText)
 			.SetFont(DEFAULT_FONT("Regular", 8))
-			.SetColorAndOpacity(FLinearColor(0.4f, 0.4, 0.4f, 1.0f)));
+			.SetColorAndOpacity(FLinearColor(0.4f, 0.4f, 0.4f, 1.0f)));
 
 		Set("FoliageEditMode.SetSelect",                 new IMAGE_BRUSH("Icons/GeneralTools/Select_40x", Icon20x20));
 		Set("FoliageEditMode.SetSelect.Small",           new IMAGE_BRUSH("Icons/GeneralTools/Select_40x", Icon20x20));
@@ -1788,7 +1788,7 @@ void FSlateEditorStyle::FStyle::SetupGeneralStyles()
 
 		Set( "FoliageEditMode.TreeView.ScrollBorder", FScrollBorderStyle()
 			.SetTopShadowBrush(FSlateNoResource())
-			.SetBottomShadowBrush(BOX_BRUSH("Common/ScrollBorderShadowBottom", FVector2D(16, 8), FMargin(0.5, 0, 0.5, 1)))
+			.SetBottomShadowBrush(BOX_BRUSH("Common/ScrollBorderShadowBottom", FVector2f(16.f, 8.f), FMargin(0.5f, 0.f, 0.5f, 1.f)))
 			);
 
 		Set("FoliageEditMode.Splitter", FSplitterStyle()
@@ -1798,15 +1798,15 @@ void FSlateEditorStyle::FStyle::SetupGeneralStyles()
 
 		Set("FoliageEditMode.ActiveToolName.Text", FTextBlockStyle(NormalText)
 			.SetFont(DEFAULT_FONT("Bold", 11))
-			.SetShadowOffset(FVector2D(1, 1))
+			.SetShadowOffset(FVector2f::UnitVector)
 			);
 
 		Set("FoliageEditMode.AddFoliageType.Text", FTextBlockStyle(NormalText)
 			.SetFont(DEFAULT_FONT("Bold", 10))
 			.SetColorAndOpacity(FLinearColor(1.0f, 1.0f, 1.0f))
 			.SetHighlightColor(FLinearColor(1.0f, 1.0f, 1.0f))
-			.SetShadowOffset(FVector2D(1, 1))
-			.SetShadowColorAndOpacity(FLinearColor(0, 0, 0, 0.9f)));
+			.SetShadowOffset(FVector2f::UnitVector)
+			.SetShadowColorAndOpacity(FLinearColor(0.f, 0.f, 0.f, 0.9f)));
 	}
 #endif // WITH_EDITOR || (IS_PROGRAM && WITH_UNREAL_DEVELOPER_TOOLS)
 
@@ -1830,26 +1830,26 @@ void FSlateEditorStyle::FStyle::SetupGeneralStyles()
 	{
 		Set( "GameProjectDialog.NewProjectTitle", FTextBlockStyle(NormalText) 
 			.SetFont( DEFAULT_FONT( "BoldCondensed", 28 ) )
-			.SetShadowOffset( FVector2D( 1,1 ) )
-			.SetShadowColorAndOpacity( FLinearColor(0,0,0,0.9f) )
+			.SetShadowOffset( FVector2f::UnitVector )
+			.SetShadowColorAndOpacity( FLinearColor(0.f,0.f,0.f,0.9f) )
 		);
 
 		Set( "GameProjectDialog.RecentProjectsTitle", FTextBlockStyle(NormalText) 
 			.SetFont( DEFAULT_FONT( "BoldCondensed", 16 ) )
-			.SetColorAndOpacity( FLinearColor(0.8f,0.8f,0.8f,1) )
-			.SetShadowOffset( FVector2D( 0,1 ) )
-			.SetShadowColorAndOpacity( FLinearColor(0,0,0,0.9f) )
+			.SetColorAndOpacity( FLinearColor(0.8f,0.8f,0.8f,1.f) )
+			.SetShadowOffset( FVector2f( 0.f,1.f) )
+			.SetShadowColorAndOpacity( FLinearColor(0.f,0.f,0.f,0.9f) )
 		);
 
 		Set( "GameProjectDialog.ProjectNamePathLabels", FTextBlockStyle(NormalText)
 			.SetFont( DEFAULT_FONT( "Regular", 12 ) )
-			.SetShadowOffset( FVector2D( 0,1 ) )
-			.SetShadowColorAndOpacity( FLinearColor(0,0,0,0.9f) )
+			.SetShadowOffset( FVector2f( 0.f,1.f) )
+			.SetShadowColorAndOpacity( FLinearColor(0.f,0.f,0.f,0.9f) )
 		);
 
 		Set( "GameProjectDialog.ErrorLabelFont", FTextBlockStyle(NormalText)
 			.SetFont( DEFAULT_FONT( "Regular", 10 ) )
-			.SetShadowOffset( FVector2D::ZeroVector )
+			.SetShadowOffset( FVector2f::ZeroVector )
 		);
 
 		Set( "GameProjectDialog.ErrorLabelBorder", new FSlateColorBrush( FLinearColor(0.2f, 0.f, 0.f, 0.7f) ) );
@@ -1872,43 +1872,43 @@ void FSlateEditorStyle::FStyle::SetupGeneralStyles()
 		Set( "GameProjectDialog.DefaultGameThumbnail", new IMAGE_BRUSH( "GameProjectDialog/default_game_thumbnail", Icon128x128 ) );
 		Set( "GameProjectDialog.DefaultGameThumbnail.Small", new IMAGE_BRUSH( "GameProjectDialog/default_game_thumbnail", Icon128x128 ) );
 		Set( "GameProjectDialog.BlankProjectThumbnail", new IMAGE_BRUSH( "GameProjectDialog/blank_project_thumbnail", Icon128x128 ) );
-		Set( "GameProjectDialog.BlankProjectPreview", new IMAGE_BRUSH( "GameProjectDialog/blank_project_preview", FVector2D(400, 200) ) );
+		Set( "GameProjectDialog.BlankProjectPreview", new IMAGE_BRUSH( "GameProjectDialog/blank_project_preview", FVector2f(400.f, 200.f) ) );
 		Set( "GameProjectDialog.BasicCodeThumbnail", new IMAGE_BRUSH( "GameProjectDialog/basic_code_thumbnail", Icon128x128 ) );
-		Set( "GameProjectDialog.CodeIcon", new IMAGE_BRUSH( "GameProjectDialog/feature_code_32x", FVector2D(32,32) ) );
-		Set( "GameProjectDialog.CodeImage", new IMAGE_BRUSH( "GameProjectDialog/feature_code", FVector2D(96,96) ) );
+		Set( "GameProjectDialog.CodeIcon", new IMAGE_BRUSH( "GameProjectDialog/feature_code_32x", Icon32x32 ) );
+		Set( "GameProjectDialog.CodeImage", new IMAGE_BRUSH( "GameProjectDialog/feature_code", FVector2f(96.f,96.f) ) );
 		Set( "GameProjectDialog.CodeImage_64", new IMAGE_BRUSH( "GameProjectDialog/feature_code", Icon64x64) );
 		Set( "GameProjectDialog.CodeImage_48x", new IMAGE_BRUSH("GameProjectDialog/feature_code", Icon40x40));
-		Set( "GameProjectDialog.BlueprintIcon", new IMAGE_BRUSH( "GameProjectDialog/feature_blueprint_32x", FVector2D(32,32) ) );
-		Set( "GameProjectDialog.BlueprintImage", new IMAGE_BRUSH( "GameProjectDialog/feature_blueprint", FVector2D(96,96) ) );
+		Set( "GameProjectDialog.BlueprintIcon", new IMAGE_BRUSH( "GameProjectDialog/feature_blueprint_32x", Icon32x32 ) );
+		Set( "GameProjectDialog.BlueprintImage", new IMAGE_BRUSH( "GameProjectDialog/feature_blueprint", FVector2f(96.f,96.f) ) );
 		Set( "GameProjectDialog.BlueprintImage_64", new IMAGE_BRUSH( "GameProjectDialog/feature_blueprint", Icon64x64) );
-		Set( "GameProjectDialog.CodeBorder", new BOX_BRUSH( "GameProjectDialog/feature_border", FMargin(4.0f/16.0f), FLinearColor(0.570, 0.359, 0.081, 1.f) ) );
+		Set( "GameProjectDialog.CodeBorder", new BOX_BRUSH( "GameProjectDialog/feature_border", FMargin(4.0f/16.0f), FLinearColor(0.570f, 0.359f, 0.081f, 1.f) ) );
 		Set( "GameProjectDialog.FeatureText", FTextBlockStyle(NormalText)
 			.SetFont( DEFAULT_FONT( "Regular", 14 ) )
-			.SetShadowOffset( FVector2D( 0,1 ) )
-			.SetShadowColorAndOpacity( FLinearColor(0,0,0,0.9f) ) );
+			.SetShadowOffset( FVector2f( 0.f,1.f) )
+			.SetShadowColorAndOpacity( FLinearColor(0.f,0.f,0.f,0.9f) ) );
 		Set( "GameProjectDialog.TemplateItemTitle", FTextBlockStyle(NormalText)
 			.SetFont( DEFAULT_FONT( "Regular", 10 ) )
-			.SetShadowOffset( FVector2D( 0,1 ) )
-			.SetShadowColorAndOpacity( FLinearColor(0,0,0,0.9f) ) );
+			.SetShadowOffset( FVector2f( 0.f,1.f) )
+			.SetShadowColorAndOpacity( FLinearColor(0.f,0.f,0.f,0.9f) ) );
 		
 		Set( "GameProjectDialog.Tab", FCheckBoxStyle()
 			.SetCheckBoxType( ESlateCheckBoxType::ToggleButton )
-			.SetUncheckedImage( BOX_BRUSH( "/GameProjectDialog/Tab_Inactive", 4 / 16.0f ) )
-			.SetUncheckedPressedImage( BOX_BRUSH( "/GameProjectDialog/Tab_Active", 4 / 16.0f ) )
-			.SetUncheckedHoveredImage( BOX_BRUSH( "/GameProjectDialog/Tab_Active", 4 / 16.0f ) )
-			.SetCheckedHoveredImage( BOX_BRUSH( "/GameProjectDialog/Tab_Active", 4 / 16.0f ) )
-			.SetCheckedPressedImage( BOX_BRUSH( "/GameProjectDialog/Tab_Active", 4 / 16.0f ) )
-			.SetCheckedImage( BOX_BRUSH( "/GameProjectDialog/Tab_Active", 4 / 16.0f ) )
+			.SetUncheckedImage( BOX_BRUSH( "/GameProjectDialog/Tab_Inactive", 4.f / 16.0f ) )
+			.SetUncheckedPressedImage( BOX_BRUSH( "/GameProjectDialog/Tab_Active", 4.f / 16.0f ) )
+			.SetUncheckedHoveredImage( BOX_BRUSH( "/GameProjectDialog/Tab_Active", 4.f / 16.0f ) )
+			.SetCheckedHoveredImage( BOX_BRUSH( "/GameProjectDialog/Tab_Active", 4.f / 16.0f ) )
+			.SetCheckedPressedImage( BOX_BRUSH( "/GameProjectDialog/Tab_Active", 4.f / 16.0f ) )
+			.SetCheckedImage( BOX_BRUSH( "/GameProjectDialog/Tab_Active", 4.f / 16.0f ) )
 			);
 
-		Set( "GameProjectDialog.TabBackground", new BOX_BRUSH( "/GameProjectDialog/tab_background", 4 / 16.0f ) );
+		Set( "GameProjectDialog.TabBackground", new BOX_BRUSH( "/GameProjectDialog/tab_background", 4.f / 16.0f ) );
 
-		Set( "GameProjectDialog.FolderIconClosed", new IMAGE_BRUSH( "Icons/FolderClosed", FVector2D(18, 16) ) );
-		Set( "GameProjectDialog.FolderIconOpen", new IMAGE_BRUSH( "Icons/FolderOpen", FVector2D(18, 16) ) );
-		Set( "GameProjectDialog.ProjectFileIcon", new IMAGE_BRUSH( "Icons/doc_16x", FVector2D(18, 16) ) );
+		Set( "GameProjectDialog.FolderIconClosed", new IMAGE_BRUSH( "Icons/FolderClosed", FVector2f(18.f, 16.f) ) );
+		Set( "GameProjectDialog.FolderIconOpen", new IMAGE_BRUSH( "Icons/FolderOpen", FVector2f(18.f, 16.f) ) );
+		Set( "GameProjectDialog.ProjectFileIcon", new IMAGE_BRUSH( "Icons/doc_16x", FVector2f(18.f, 16.f) ) );
 
-		Set( "GameProjectDialog.IncludeStarterContent", new IMAGE_BRUSH( "/GameProjectDialog/IncludeStarterContent", FVector2D(64, 64) ) );
-		Set( "GameProjectDialog.NoStarterContent", new IMAGE_BRUSH( "/GameProjectDialog/NoStarterContent", FVector2D(64, 64) ) );
+		Set( "GameProjectDialog.IncludeStarterContent", new IMAGE_BRUSH( "/GameProjectDialog/IncludeStarterContent", Icon64x64 ) );
+		Set( "GameProjectDialog.NoStarterContent", new IMAGE_BRUSH( "/GameProjectDialog/NoStarterContent", Icon64x64 ) );
 		
 		Set("GameProjectDialog.XRDisabled", new IMAGE_BRUSH("Icons/AssetIcons/SkeletalMesh_64x", Icon64x64));
 		Set("GameProjectDialog.XREnabled", new IMAGE_BRUSH("Icons/AssetIcons/SkeletalMesh_64x", Icon64x64));
@@ -1918,11 +1918,11 @@ void FSlateEditorStyle::FStyle::SetupGeneralStyles()
 
 		Set( "FilePath.FolderButton",
 			FButtonStyle(HoverHintOnly)
-			.SetNormal(  BOX_BRUSH( "Common/ButtonHoverHint", FMargin(4/16.0f), FLinearColor(0.5f,0.5f,0.5f) ) )
-			.SetHovered( BOX_BRUSH( "Common/ButtonHoverHint", FMargin(4/16.0f), FLinearColor(0.6f,0.6f,0.6f) ) )
-			.SetPressed( BOX_BRUSH( "Common/ButtonHoverHint", FMargin(4/16.0f), FLinearColor(0.65f,0.65f,0.65f) ) )
-			.SetNormalPadding( FMargin(0,0,0,1) )
-			.SetPressedPadding( FMargin(0,1,0,0) )
+			.SetNormal(  BOX_BRUSH( "Common/ButtonHoverHint", FMargin(4.f/16.0f), FLinearColor(0.5f,0.5f,0.5f) ) )
+			.SetHovered( BOX_BRUSH( "Common/ButtonHoverHint", FMargin(4.f/16.0f), FLinearColor(0.6f,0.6f,0.6f) ) )
+			.SetPressed( BOX_BRUSH( "Common/ButtonHoverHint", FMargin(4.f/16.0f), FLinearColor(0.65f,0.65f,0.65f) ) )
+			.SetNormalPadding( FMargin(0.f,0.f,0.f,1.f) )
+			.SetPressedPadding( FMargin(0.f,1.f,0.f,0.f) )
 			);
 		Set( "FilePath.GroupIndicator", new BOX_BRUSH( "GameProjectDialog/filepath_group_indicator", FMargin(4.0f/16.f) ) );
 	}
@@ -1932,14 +1932,14 @@ void FSlateEditorStyle::FStyle::SetupGeneralStyles()
 	{
 		Set( "NewClassDialog.PageTitle", FTextBlockStyle(NormalText)
 			.SetFont( DEFAULT_FONT( "BoldCondensed", 28 ) )
-			.SetShadowOffset( FVector2D( 1,1 ) )
-			.SetShadowColorAndOpacity( FLinearColor(0,0,0,0.9f) )
+			.SetShadowOffset( FVector2f::UnitVector )
+			.SetShadowColorAndOpacity( FLinearColor(0.f,0.f,0.f,0.9f) )
 		);
 
 		Set( "NewClassDialog.SelectedParentClassLabel", FTextBlockStyle(NormalText)
 			.SetFont( DEFAULT_FONT( "Regular", 12 ) )
-			.SetShadowOffset( FVector2D( 0,1 ) )
-			.SetShadowColorAndOpacity( FLinearColor(0,0,0,0.9f) )
+			.SetShadowOffset( FVector2f( 0.f,1.f) )
+			.SetShadowColorAndOpacity( FLinearColor(0.f,0.f,0.f,0.9f) )
 		);
 
 		Set( "NewClassDialog.ErrorLabelFont", FTextBlockStyle(NormalText)
@@ -1965,8 +1965,8 @@ void FSlateEditorStyle::FStyle::SetupGeneralStyles()
 
 		Set( "NewClassDialog.ParentClassItemTitle", FTextBlockStyle(NormalText)
 			.SetFont( DEFAULT_FONT( "Regular", 14 ) )
-			.SetShadowOffset( FVector2D( 0,1 ) )
-			.SetShadowColorAndOpacity( FLinearColor(0,0,0,0.9f) )
+			.SetShadowOffset( FVector2f( 0.f,1.f) )
+			.SetShadowColorAndOpacity( FLinearColor(0.f,0.f,0.f,0.9f) )
 		);
 	}
 
@@ -1979,13 +1979,13 @@ void FSlateEditorStyle::FStyle::SetupGeneralStyles()
 
 	// Hardware Targeting
 	{
-		Set( "HardwareTargeting.MobilePlatform", new IMAGE_BRUSH( "/Icons/HardwareTargeting/Mobile", FVector2D(64, 64) ) );
-		Set( "HardwareTargeting.DesktopPlatform", new IMAGE_BRUSH( "/Icons/HardwareTargeting/Desktop", FVector2D(64, 64) ) );
-		Set( "HardwareTargeting.HardwareUnspecified", new IMAGE_BRUSH( "/Icons/HardwareTargeting/HardwareUnspecified", FVector2D(64, 64) ) );
+		Set( "HardwareTargeting.MobilePlatform", new IMAGE_BRUSH( "/Icons/HardwareTargeting/Mobile", Icon64x64 ) );
+		Set( "HardwareTargeting.DesktopPlatform", new IMAGE_BRUSH( "/Icons/HardwareTargeting/Desktop", Icon64x64 ) );
+		Set( "HardwareTargeting.HardwareUnspecified", new IMAGE_BRUSH( "/Icons/HardwareTargeting/HardwareUnspecified", Icon64x64 ) );
 
-		Set( "HardwareTargeting.MaximumQuality", new IMAGE_BRUSH( "/Icons/HardwareTargeting/MaximumQuality", FVector2D(64, 64) ) );
-		Set( "HardwareTargeting.ScalableQuality", new IMAGE_BRUSH( "/Icons/HardwareTargeting/ScalableQuality", FVector2D(64, 64) ) );
-		Set( "HardwareTargeting.GraphicsUnspecified", new IMAGE_BRUSH( "/Icons/HardwareTargeting/GraphicsUnspecified", FVector2D(64, 64) ) );
+		Set( "HardwareTargeting.MaximumQuality", new IMAGE_BRUSH( "/Icons/HardwareTargeting/MaximumQuality", Icon64x64 ) );
+		Set( "HardwareTargeting.ScalableQuality", new IMAGE_BRUSH( "/Icons/HardwareTargeting/ScalableQuality", Icon64x64 ) );
+		Set( "HardwareTargeting.GraphicsUnspecified", new IMAGE_BRUSH( "/Icons/HardwareTargeting/GraphicsUnspecified", Icon64x64 ) );
 	}
 
 #endif // WITH_EDITOR
@@ -2137,7 +2137,7 @@ void FSlateEditorStyle::FStyle::SetupGeneralStyles()
 
 	// Highres Screenshot
 	{
-		Set("HighresScreenshot.WarningStrip", new IMAGE_BRUSH( "Common/WarningStripe", FVector2D(20,6), FLinearColor::White, ESlateBrushTileType::Horizontal ) );
+		Set("HighresScreenshot.WarningStrip", new IMAGE_BRUSH( "Common/WarningStripe", FVector2f(20.f,6.f), FLinearColor::White, ESlateBrushTileType::Horizontal ) );
 		Set("HighresScreenshot.SpecifyCaptureRectangle", new IMAGE_BRUSH( "Icons/icon_CaptureRegion_24x", Icon24x24 ) );
 		Set("HighresScreenshot.FullViewportCaptureRegion", new IMAGE_BRUSH( "Icons/icon_CaptureRegion_FullViewport_24x", Icon24x24 ) );
 		Set("HighresScreenshot.CameraSafeAreaCaptureRegion", new IMAGE_BRUSH("Icons/icon_CaptureRegion_Camera_Safe_24x", Icon24x24));
@@ -2280,8 +2280,8 @@ void FSlateEditorStyle::FStyle::SetupWorldBrowserStyles()
 
 		Set("WorldBrowser.StatusBarText", FTextBlockStyle(NormalText)
 			.SetFont(DEFAULT_FONT("BoldCondensed", 12))
-			.SetColorAndOpacity(FLinearColor(0.9, 0.9f, 0.9f, 0.5f))
-			.SetShadowOffset(FVector2D::ZeroVector)
+			.SetColorAndOpacity(FLinearColor(0.9f, 0.9f, 0.9f, 0.5f))
+			.SetShadowOffset(FVector2f::ZeroVector)
 		);
 
 		Set("WorldBrowser.LabelFont", DEFAULT_FONT("Regular", 9));
@@ -2316,11 +2316,11 @@ void FSlateEditorStyle::FStyle::SetupSequencerStyles()
 		Set("Sequencer.KeyTriangle", new IMAGE_BRUSH("Sequencer/KeyTriangle", Icon12x12));
 		Set("Sequencer.KeyLeft", new IMAGE_BRUSH("Sequencer/KeyLeft", Icon12x12));
 		Set("Sequencer.KeyRight", new IMAGE_BRUSH("Sequencer/KeyRight", Icon12x12));
-		Set("Sequencer.PartialKey", new IMAGE_BRUSH("Sequencer/PartialKey", FVector2D(11.f, 11.f)));
+		Set("Sequencer.PartialKey", new IMAGE_BRUSH("Sequencer/PartialKey", FVector2f(11.f, 11.f)));
 		Set("Sequencer.Star", new IMAGE_BRUSH("Sequencer/Star", Icon12x12));
 		Set("Sequencer.Empty", new IMAGE_BRUSH("Sequencer/Empty", Icon12x12));
-		Set("Sequencer.TangentHandle", new IMAGE_BRUSH("Sequencer/TangentHandle", FVector2D(7, 7)));
-		Set("Sequencer.GenericDivider", new IMAGE_BRUSH("Sequencer/GenericDivider", FVector2D(2.f, 2.f), FLinearColor::White, ESlateBrushTileType::Vertical));
+		Set("Sequencer.TangentHandle", new IMAGE_BRUSH("Sequencer/TangentHandle", FVector2f(7.f, 7.f)));
+		Set("Sequencer.GenericDivider", new IMAGE_BRUSH("Sequencer/GenericDivider", FVector2f(2.f, 2.f), FLinearColor::White, ESlateBrushTileType::Vertical));
 
 		Set("Sequencer.Timeline.ScrubHandleDown", new BOX_BRUSH("Sequencer/ScrubHandleDown", FMargin(6.f / 13.f, 5 / 12.f, 6 / 13.f, 8 / 12.f)));
 		Set("Sequencer.Timeline.ScrubHandleUp", new BOX_BRUSH("Sequencer/ScrubHandleUp", FMargin(6.f / 13.f, 8 / 12.f, 6 / 13.f, 5 / 12.f)));
@@ -2333,7 +2333,7 @@ void FSlateEditorStyle::FStyle::SetupSequencerStyles()
 		Set("Sequencer.Timeline.RangeHandleLeft", new BOX_BRUSH("Sequencer/GenericGripLeft", FMargin(5.f / 16.f)));
 		Set("Sequencer.Timeline.RangeHandleRight", new BOX_BRUSH("Sequencer/GenericGripRight", FMargin(5.f / 16.f)));
 		Set("Sequencer.Timeline.RangeHandle", new BOX_BRUSH("Sequencer/GenericSectionBackground", FMargin(5.f / 16.f)));
-		Set("Sequencer.Timeline.NotifyAlignmentMarker", new IMAGE_BRUSH("Sequencer/NotifyAlignmentMarker", FVector2D(10, 19)));
+		Set("Sequencer.Timeline.NotifyAlignmentMarker", new IMAGE_BRUSH("Sequencer/NotifyAlignmentMarker", FVector2f(10.f, 19.f)));
 		Set("Sequencer.Timeline.PlayRange_Top_L", new BOX_BRUSH("Sequencer/PlayRange_Top_L", FMargin(1.f, 0.5f, 0.f, 0.5f)));
 		Set("Sequencer.Timeline.PlayRange_Top_R", new BOX_BRUSH("Sequencer/PlayRange_Top_R", FMargin(0.f, 0.5f, 1.f, 0.5f)));
 		Set("Sequencer.Timeline.PlayRange_L", new BOX_BRUSH("Sequencer/PlayRange_L", FMargin(1.f, 0.5f, 0.f, 0.5f)));
@@ -2343,7 +2343,7 @@ void FSlateEditorStyle::FStyle::SetupSequencerStyles()
 
 		Set("Sequencer.Timeline.SubSequenceRangeHashL", new BORDER_BRUSH("Sequencer/SubSequenceRangeHashL", FMargin(1.f, 0.f, 0.f, 0.f)));
 		Set("Sequencer.Timeline.SubSequenceRangeHashR", new BORDER_BRUSH("Sequencer/SubSequenceRangeHashR", FMargin(1.f, 0.f, 0.f, 0.f)));
-		Set("Sequencer.Timeline.EaseInOut", new IMAGE_BRUSH("Sequencer/EaseInOut", FVector2D(128, 128)));
+		Set("Sequencer.Timeline.EaseInOut", new IMAGE_BRUSH("Sequencer/EaseInOut", Icon128x128));
 		Set("Sequencer.InterpLine", new BOX_BRUSH("Sequencer/InterpLine", FMargin(5.f / 7.f, 0.f, 0.f, 0.f)));
 
 		Set("Sequencer.Transport.JumpToPreviousKey", FButtonStyle()
@@ -2355,13 +2355,13 @@ void FSlateEditorStyle::FStyle::SetupSequencerStyles()
 			.SetPressed(IMAGE_BRUSH("/Sequencer/Transport_Bar/Next_Frame_24x", Icon24x24))
 			.SetHovered(IMAGE_BRUSH("/Sequencer/Transport_Bar/Next_Frame_24x_OFF", Icon24x24)));
 		Set("Sequencer.Transport.SetPlayStart", FButtonStyle()
-			.SetNormal(IMAGE_BRUSH("/Sequencer/Transport_Bar/Bracket_In_16x24_OFF", FVector2D(16, 24)))
-			.SetPressed(IMAGE_BRUSH("/Sequencer/Transport_Bar/Bracket_In_16x24", FVector2D(16, 24)))
-			.SetHovered(IMAGE_BRUSH("/Sequencer/Transport_Bar/Bracket_In_16x24_OFF", FVector2D(16, 24))));
+			.SetNormal(IMAGE_BRUSH("/Sequencer/Transport_Bar/Bracket_In_16x24_OFF", FVector2f(16.f, 24.f)))
+			.SetPressed(IMAGE_BRUSH("/Sequencer/Transport_Bar/Bracket_In_16x24", FVector2f(16.f, 24.f)))
+			.SetHovered(IMAGE_BRUSH("/Sequencer/Transport_Bar/Bracket_In_16x24_OFF", FVector2f(16.f, 24.f))));
 		Set("Sequencer.Transport.SetPlayEnd", FButtonStyle()
-			.SetNormal(IMAGE_BRUSH("/Sequencer/Transport_Bar/Bracket_Out_16x24_OFF", FVector2D(16, 24)))
-			.SetPressed(IMAGE_BRUSH("/Sequencer/Transport_Bar/Bracket_Out_16x24", FVector2D(16, 24)))
-			.SetHovered(IMAGE_BRUSH("/Sequencer/Transport_Bar/Bracket_Out_16x24_OFF", FVector2D(16, 24))));
+			.SetNormal(IMAGE_BRUSH("/Sequencer/Transport_Bar/Bracket_Out_16x24_OFF", FVector2f(16.f, 24.f)))
+			.SetPressed(IMAGE_BRUSH("/Sequencer/Transport_Bar/Bracket_Out_16x24", FVector2f(16.f, 24.f)))
+			.SetHovered(IMAGE_BRUSH("/Sequencer/Transport_Bar/Bracket_Out_16x24_OFF", FVector2f(16.f, 24.f))));
 
 		Set("Sequencer.Transport.CloseButton", FButtonStyle()
 			.SetNormal(IMAGE_BRUSH("/Docking/CloseApp_Normal", Icon16x16))
@@ -2389,10 +2389,10 @@ void FSlateEditorStyle::FStyle::SetupSequencerStyles()
 		Set("Sequencer.Section.SelectionBorder", new BORDER_BRUSH(TEXT("Sequencer/SectionHighlight"), FMargin(7.f / 16.f)));
 		Set("Sequencer.Section.LockedBorder", new BORDER_BRUSH(TEXT("Sequencer/SectionLocked"), FMargin(7.f / 16.f)));
 		Set("Sequencer.Section.SelectedSectionOverlay", new IMAGE_BRUSH(TEXT("Sequencer/SelectedSectionOverlay"), Icon16x16, FLinearColor::White, ESlateBrushTileType::Both));
-		Set("Sequencer.Section.FilmBorder", new IMAGE_BRUSH(TEXT("Sequencer/SectionFilmBorder"), FVector2D(10, 7), FLinearColor::White, ESlateBrushTileType::Horizontal));
+		Set("Sequencer.Section.FilmBorder", new IMAGE_BRUSH(TEXT("Sequencer/SectionFilmBorder"), FVector2f(10.f, 7.f), FLinearColor::White, ESlateBrushTileType::Horizontal));
 		Set("Sequencer.Section.GripLeft", new BOX_BRUSH("Sequencer/SectionGripLeft", FMargin(5.f / 16.f)));
 		Set("Sequencer.Section.GripRight", new BOX_BRUSH("Sequencer/SectionGripRight", FMargin(5.f / 16.f)));
-		Set("Sequencer.Section.EasingHandle", new IMAGE_BRUSH("Sequencer/EasingHandle", FVector2D(10.f, 10.f)));
+		Set("Sequencer.Section.EasingHandle", new IMAGE_BRUSH("Sequencer/EasingHandle", FVector2f(10.f, 10.f)));
 
 		Set("Sequencer.Section.PreRoll", new BORDER_BRUSH(TEXT("Sequencer/PreRoll"), FMargin(0.f, .5f, 0.f, .5f)));
 
@@ -2426,8 +2426,8 @@ void FSlateEditorStyle::FStyle::SetupSequencerStyles()
 		Set("Sequencer.AnimationOutliner.BoldFont", DEFAULT_FONT("Bold", 11));
 		Set("Sequencer.AnimationOutliner.RegularFont", DEFAULT_FONT("Regular", 9));
 		Set("Sequencer.AnimationOutliner.ItalicFont", DEFAULT_FONT("Italic", 10));
-		Set("Sequencer.ShotFilter", new IMAGE_BRUSH("Sequencer/FilteredArea", FVector2D(74, 74), FLinearColor::White, ESlateBrushTileType::Both));
-		Set("Sequencer.KeyMark", new IMAGE_BRUSH("Sequencer/KeyMark", FVector2D(3, 21), FLinearColor::White, ESlateBrushTileType::NoTile));
+		Set("Sequencer.ShotFilter", new IMAGE_BRUSH("Sequencer/FilteredArea", FVector2f(74.f, 74.f), FLinearColor::White, ESlateBrushTileType::Both));
+		Set("Sequencer.KeyMark", new IMAGE_BRUSH("Sequencer/KeyMark", FVector2f(3.f, 21.f), FLinearColor::White, ESlateBrushTileType::NoTile));
 		Set("Sequencer.ToggleAutoKeyEnabled", new IMAGE_BRUSH("Sequencer/Main_Icons/Icon_Sequencer_Auto_Key_24x", Icon48x48));
 		Set("Sequencer.ToggleAutoKeyEnabled.Small", new IMAGE_BRUSH("Sequencer/Main_Icons/Icon_Sequencer_Auto_Key_24x", Icon24x24));
 		Set("Sequencer.SetAutoKey", new IMAGE_BRUSH("Sequencer/Main_Icons/Icon_Sequencer_Auto_Key_24x", Icon48x48));
@@ -2464,7 +2464,7 @@ void FSlateEditorStyle::FStyle::SetupSequencerStyles()
 			.SetFont(DEFAULT_FONT("Bold", 11))
 			.SetColorAndOpacity(FLinearColor(1.0f, 1.0f, 1.0f))
 			.SetHighlightColor(FLinearColor(1.0f, 1.0f, 1.0f))
-			.SetShadowOffset(FVector2D(1, 1))
+			.SetShadowOffset(FVector2f::UnitVector)
 			.SetShadowColorAndOpacity(FLinearColor(0, 0, 0, 0.9f)));
 
 		Set("Sequencer.ToggleAutoScroll", new IMAGE_BRUSH("Icons/icon_Sequencer_ToggleAutoScroll_40x", Icon48x48));
@@ -2477,12 +2477,12 @@ void FSlateEditorStyle::FStyle::SetupSequencerStyles()
 		Set("Sequencer.LockCamera", new IMAGE_BRUSH("Sequencer/Main_Icons/Icon_Sequencer_Look_Thru_24x", Icon16x16));
 		Set("Sequencer.UnlockCamera", new IMAGE_BRUSH("Sequencer/Main_Icons/Icon_Sequencer_Look_Thru_24x", Icon16x16, FLinearColor(1.f, 1.f, 1.f, 0.5f)));
 		Set("Sequencer.Thumbnail.SectionHandle", new IMAGE_BRUSH("Old/White", Icon16x16, FLinearColor::Black));
-		Set("Sequencer.TrackHoverHighlight_Top", new IMAGE_BRUSH(TEXT("Sequencer/TrackHoverHighlight_Top"), FVector2D(4, 4)));
-		Set("Sequencer.TrackHoverHighlight_Bottom", new IMAGE_BRUSH(TEXT("Sequencer/TrackHoverHighlight_Bottom"), FVector2D(4, 4)));
-		Set("Sequencer.SpawnableIconOverlay", new IMAGE_BRUSH(TEXT("Sequencer/SpawnableIconOverlay"), FVector2D(13, 13)));
-		Set("Sequencer.ReplaceableIconOverlay", new IMAGE_BRUSH(TEXT("Sequencer/ReplaceableIconOverlay"), FVector2D(13, 13)));
+		Set("Sequencer.TrackHoverHighlight_Top", new IMAGE_BRUSH(TEXT("Sequencer/TrackHoverHighlight_Top"), FVector2f(4.f, 4.f)));
+		Set("Sequencer.TrackHoverHighlight_Bottom", new IMAGE_BRUSH(TEXT("Sequencer/TrackHoverHighlight_Bottom"), FVector2f(4.f, 4.f)));
+		Set("Sequencer.SpawnableIconOverlay", new IMAGE_BRUSH(TEXT("Sequencer/SpawnableIconOverlay"), FVector2f(13.f, 13.f)));
+		Set("Sequencer.ReplaceableIconOverlay", new IMAGE_BRUSH(TEXT("Sequencer/ReplaceableIconOverlay"), FVector2f(13.f, 13.f)));
 		Set("Sequencer.InvalidSpawnableIcon", new IMAGE_BRUSH(TEXT("Sequencer/InvalidSpawnableIcon"), Icon16x16));
-		Set("Sequencer.MultipleIconOverlay", new IMAGE_BRUSH(TEXT("Sequencer/MultipleIconOverlay"), FVector2D(13, 13)));
+		Set("Sequencer.MultipleIconOverlay", new IMAGE_BRUSH(TEXT("Sequencer/MultipleIconOverlay"), FVector2f(13.f, 13.f)));
 		Set("Sequencer.DynamicBindingIconOverlay", new IMAGE_BRUSH(TEXT("Sequencer/DynamicBindingIconOverlay"), Icon16x16));
 		Set("Sequencer.SpawnableDynamicBindingIconOverlay", new IMAGE_BRUSH(TEXT("Sequencer/SpawnableDynamicBindingIconOverlay"), Icon16x16));
 		Set("Sequencer.LockSequence", new IMAGE_BRUSH("Sequencer/Main_Icons/Icon_Sequencer_Locked_16x", Icon16x16));
@@ -2532,23 +2532,23 @@ void FSlateEditorStyle::FStyle::SetupSequencerStyles()
 		.SetFont(DEFAULT_FONT("Bold", 11))
 		.SetColorAndOpacity(FLinearColor(1.0f, 1.0f, 1.0f))
 		.SetHighlightColor(FLinearColor(1.0f, 1.0f, 1.0f))
-		.SetShadowOffset(FVector2D(1, 1))
-		.SetShadowColorAndOpacity(FLinearColor(0, 0, 0, 0.9f)));
+		.SetShadowOffset(FVector2f::UnitVector)
+		.SetShadowColorAndOpacity(FLinearColor(0.f, 0.f, 0.f, 0.9f)));
 	Set("Sequencer.BreadcrumbIcon", new IMAGE_BRUSH("Common/SmallArrowRight", Icon10x10));
 
 	const FButtonStyle NoBorder = FButtonStyle()
 	.SetNormal(FSlateNoResource())
 	.SetHovered(FSlateNoResource())
 	.SetPressed(FSlateNoResource())
-	.SetNormalPadding(FMargin(0, 0, 0, 1))
-	.SetPressedPadding(FMargin(0, 1, 0, 0));
+	.SetNormalPadding(FMargin(0.f, 0.f, 0.f, 1.f))
+	.SetPressedPadding(FMargin(0.f, 1.f, 0.f, 0.f));
 
 	const FButtonStyle DetailsKeyButton = FButtonStyle(NoBorder)
-	.SetNormal(IMAGE_BRUSH("Sequencer/AddKey_Details", FVector2D(11, 11)))
-	.SetHovered(IMAGE_BRUSH("Sequencer/AddKey_Details", FVector2D(11, 11), SelectionColor))
-	.SetPressed(IMAGE_BRUSH("Sequencer/AddKey_Details", FVector2D(11, 11), SelectionColor_Pressed))
-	.SetNormalPadding(FMargin(0, 1))
-	.SetPressedPadding(FMargin(0, 2, 0, 0));
+	.SetNormal(IMAGE_BRUSH("Sequencer/AddKey_Details", FVector2f(11.f, 11.f)))
+	.SetHovered(IMAGE_BRUSH("Sequencer/AddKey_Details", FVector2f(11.f, 11.f), SelectionColor))
+	.SetPressed(IMAGE_BRUSH("Sequencer/AddKey_Details", FVector2f(11.f, 11.f), SelectionColor_Pressed))
+	.SetNormalPadding(FMargin(0.f, 1.f))
+	.SetPressedPadding(FMargin(0.f, 2.f, 0.f, 0.f));
 	Set("Sequencer.AddKey.Details", DetailsKeyButton);
 
 	const FSplitterStyle OutlinerSplitterStyle = FSplitterStyle()
@@ -2557,8 +2557,8 @@ void FSlateEditorStyle::FStyle::SetupSequencerStyles()
 	Set("Sequencer.AnimationOutliner.Splitter", OutlinerSplitterStyle);
 
 	Set("Sequencer.HyperlinkSpinBox", FSpinBoxStyle(GetWidgetStyle<FSpinBoxStyle>("SpinBox"))
-		.SetTextPadding(FMargin(0))
-		.SetBackgroundBrush(BORDER_BRUSH("Old/HyperlinkDotted", FMargin(0, 0, 0, 3 / 16.0f), FSlateColor::UseSubduedForeground()))
+		.SetTextPadding(FMargin(0.f))
+		.SetBackgroundBrush(BORDER_BRUSH("Old/HyperlinkDotted", FMargin(0.f, 0.f, 0.f, 3.f / 16.0f), FSlateColor::UseSubduedForeground()))
 		.SetHoveredBackgroundBrush(FSlateNoResource())
 		.SetInactiveFillBrush(FSlateNoResource())
 		.SetActiveFillBrush(FSlateNoResource())
@@ -2609,10 +2609,10 @@ void FSlateEditorStyle::FStyle::SetupSequencerStyles()
 	Set("Sequencer.UnboundEvent", new IMAGE_BRUSH("Sequencer/UnboundEvent", Icon16x16));
 
 	// Sequencer Blending Iconography
-	Set("EMovieSceneBlendType::Absolute", new IMAGE_BRUSH("Sequencer/EMovieSceneBlendType_Absolute", FVector2D(32, 16)));
-	Set("EMovieSceneBlendType::Relative", new IMAGE_BRUSH("Sequencer/EMovieSceneBlendType_Relative", FVector2D(32, 16)));
-	Set("EMovieSceneBlendType::Additive", new IMAGE_BRUSH("Sequencer/EMovieSceneBlendType_Additive", FVector2D(32, 16)));
-	Set("EMovieSceneBlendType::AdditiveFromBase", new IMAGE_BRUSH("Sequencer/EMovieSceneBlendType_AdditiveFromBase", FVector2D(32, 16)));
+	Set("EMovieSceneBlendType::Absolute", new IMAGE_BRUSH("Sequencer/EMovieSceneBlendType_Absolute", FVector2f(32.f, 16.f)));
+	Set("EMovieSceneBlendType::Relative", new IMAGE_BRUSH("Sequencer/EMovieSceneBlendType_Relative", FVector2f(32.f, 16.f)));
+	Set("EMovieSceneBlendType::Additive", new IMAGE_BRUSH("Sequencer/EMovieSceneBlendType_Additive", FVector2f(32.f, 16.f)));
+	Set("EMovieSceneBlendType::AdditiveFromBase", new IMAGE_BRUSH("Sequencer/EMovieSceneBlendType_AdditiveFromBase", FVector2f(32.f, 16.f)));
 
 	// Sequencer & Curve Editor Toolbar Style ( Grabs core default, copies brushes, changes separator and block spacing )
 	Set( "Sequencer.ToolBar.Background", 		new CORE_BOX_BRUSH( "Common/GroupBorder", FMargin(4.0f/16.0f) ) );
@@ -2704,7 +2704,7 @@ void FSlateEditorStyle::FStyle::SetupViewportStyles()
 		Set("ViewportMenu.Block.Padding", FMargin(0));
 
 		Set("ViewportMenu.Heading.Font", DEFAULT_FONT("Regular", 8));
-		Set("ViewportMenu.Heading.ColorAndOpacity", FLinearColor(0.4f, 0.4, 0.4f, 1.0f));
+		Set("ViewportMenu.Heading.ColorAndOpacity", FLinearColor(0.4f, 0.4f, 0.4f, 1.0f));
 
 		const FCheckBoxStyle ViewportMenuCheckBoxCheckBoxStyle = FCheckBoxStyle()
 			.SetUncheckedImage(IMAGE_BRUSH("Common/SmallCheckBox", Icon14x14))
@@ -2853,7 +2853,7 @@ void FSlateEditorStyle::FStyle::SetupNotificationBarStyles()
 
 		Set("NotificationBar.Heading", FTextBlockStyle(NormalText)
 			.SetFont(DEFAULT_FONT("Regular", 8))
-			.SetColorAndOpacity(FLinearColor(0.4f, 0.4, 0.4f, 1.0f)));
+			.SetColorAndOpacity(FLinearColor(0.4f, 0.4f, 0.4f, 1.0f)));
 
 		const FCheckBoxStyle NotificationBarCheckBoxCheckBoxStyle = FCheckBoxStyle()
 			.SetUncheckedImage(IMAGE_BRUSH("Common/SmallCheckBox", Icon14x14))
@@ -2941,41 +2941,41 @@ void FSlateEditorStyle::FStyle::SetupWindowStyles()
 {
 	// Window styling
 	{
-		Set( "Window.Background",                 new IMAGE_BRUSH( "Old/Window/WindowBackground", FVector2D(74, 74), FLinearColor::White, ESlateBrushTileType::Both) );
+		Set( "Window.Background",                 new IMAGE_BRUSH( "Old/Window/WindowBackground", FVector2f(74.f, 74.f), FLinearColor::White, ESlateBrushTileType::Both) );
 		Set( "Window.Border",                     new BOX_BRUSH( "Old/Window/WindowBorder", 0.48f ) );
-		Set( "Window.Title.Active",               new IMAGE_BRUSH( "Old/Window/WindowTitle", Icon32x32, FLinearColor(1,1,1,1), ESlateBrushTileType::Horizontal  ) );
-		Set( "Window.Title.Inactive",             new IMAGE_BRUSH( "Old/Window/WindowTitle_Inactive", Icon32x32, FLinearColor(1,1,1,1), ESlateBrushTileType::Horizontal  ) );
-		Set( "Window.Title.Flash",				  new IMAGE_BRUSH( "Old/Window/WindowTitle_Flashing", Icon24x24, FLinearColor(1,1,1,1), ESlateBrushTileType::Horizontal  ) );
+		Set( "Window.Title.Active",               new IMAGE_BRUSH( "Old/Window/WindowTitle", Icon32x32, FLinearColor::White, ESlateBrushTileType::Horizontal  ) );
+		Set( "Window.Title.Inactive",             new IMAGE_BRUSH( "Old/Window/WindowTitle_Inactive", Icon32x32, FLinearColor::White, ESlateBrushTileType::Horizontal  ) );
+		Set( "Window.Title.Flash",				  new IMAGE_BRUSH( "Old/Window/WindowTitle_Flashing", Icon24x24, FLinearColor::White, ESlateBrushTileType::Horizontal  ) );
 
 #if !PLATFORM_MAC
 		const FButtonStyle MinimizeButtonStyle = FButtonStyle(Button)
-			.SetNormal ( IMAGE_BRUSH( "Old/Window/WindowButton_Minimize_Normal", FVector2D(27, 18) ) )
-			.SetHovered( IMAGE_BRUSH( "Old/Window/WindowButton_Minimize_Hovered", FVector2D(27, 18) ) )
-			.SetPressed( IMAGE_BRUSH( "Old/Window/WindowButton_Minimize_Pressed", FVector2D(27, 18) ) );
+			.SetNormal ( IMAGE_BRUSH( "Old/Window/WindowButton_Minimize_Normal", FVector2f(27.f, 18.f) ) )
+			.SetHovered( IMAGE_BRUSH( "Old/Window/WindowButton_Minimize_Hovered", FVector2f(27.f, 18.f) ) )
+			.SetPressed( IMAGE_BRUSH( "Old/Window/WindowButton_Minimize_Pressed", FVector2f(27.f, 18.f) ) );
 
-		Set( "Window.Buttons.Minimize.Normal",    new IMAGE_BRUSH( "Old/Window/WindowButton_Minimize_Normal", FVector2D(27, 18) ) );
-		Set( "Window.Buttons.Minimize.Hovered",   new IMAGE_BRUSH( "Old/Window/WindowButton_Minimize_Hovered", FVector2D(27, 18) ) );
-		Set( "Window.Buttons.Minimize.Pressed",   new IMAGE_BRUSH( "Old/Window/WindowButton_Minimize_Pressed", FVector2D(27, 18) ) );
-		Set( "Window.Buttons.Minimize.Disabled",   new IMAGE_BRUSH( "Old/Window/WindowButton_Minimize_Disabled", FVector2D(27, 18) ) );
+		Set( "Window.Buttons.Minimize.Normal",    new IMAGE_BRUSH( "Old/Window/WindowButton_Minimize_Normal", FVector2f(27.f, 18.f) ) );
+		Set( "Window.Buttons.Minimize.Hovered",   new IMAGE_BRUSH( "Old/Window/WindowButton_Minimize_Hovered", FVector2f(27.f, 18.f) ) );
+		Set( "Window.Buttons.Minimize.Pressed",   new IMAGE_BRUSH( "Old/Window/WindowButton_Minimize_Pressed", FVector2f(27.f, 18.f) ) );
+		Set( "Window.Buttons.Minimize.Disabled",   new IMAGE_BRUSH( "Old/Window/WindowButton_Minimize_Disabled", FVector2f(27.f, 18.f) ) );
 
 		const FButtonStyle MaximizeButtonStyle = FButtonStyle(Button)
-			.SetNormal ( IMAGE_BRUSH( "Old/Window/WindowButton_Maximize_Normal", FVector2D(23, 18) ) )
-			.SetHovered( IMAGE_BRUSH( "Old/Window/WindowButton_Maximize_Hovered", FVector2D(23, 18) ) )
-			.SetPressed( IMAGE_BRUSH( "Old/Window/WindowButton_Maximize_Pressed", FVector2D(23, 18) ) );
+			.SetNormal ( IMAGE_BRUSH( "Old/Window/WindowButton_Maximize_Normal", FVector2f(23.f, 18.f) ) )
+			.SetHovered( IMAGE_BRUSH( "Old/Window/WindowButton_Maximize_Hovered", FVector2f(23.f, 18.f) ) )
+			.SetPressed( IMAGE_BRUSH( "Old/Window/WindowButton_Maximize_Pressed", FVector2f(23.f, 18.f) ) );
 
-		Set( "Window.Buttons.Maximize.Normal",    new IMAGE_BRUSH( "Old/Window/WindowButton_Maximize_Normal", FVector2D(23, 18) ) );
-		Set( "Window.Buttons.Maximize.Hovered",   new IMAGE_BRUSH( "Old/Window/WindowButton_Maximize_Hovered", FVector2D(23, 18) ) );
-		Set( "Window.Buttons.Maximize.Pressed",   new IMAGE_BRUSH( "Old/Window/WindowButton_Maximize_Pressed", FVector2D(23, 18) ) );
-		Set( "Window.Buttons.Maximize.Disabled",   new IMAGE_BRUSH( "Old/Window/WindowButton_Maximize_Disabled", FVector2D(23, 18) ) );
+		Set( "Window.Buttons.Maximize.Normal",    new IMAGE_BRUSH( "Old/Window/WindowButton_Maximize_Normal", FVector2f(23.f, 18.f) ) );
+		Set( "Window.Buttons.Maximize.Hovered",   new IMAGE_BRUSH( "Old/Window/WindowButton_Maximize_Hovered", FVector2f(23.f, 18.f) ) );
+		Set( "Window.Buttons.Maximize.Pressed",   new IMAGE_BRUSH( "Old/Window/WindowButton_Maximize_Pressed", FVector2f(23.f, 18.f) ) );
+		Set( "Window.Buttons.Maximize.Disabled",   new IMAGE_BRUSH( "Old/Window/WindowButton_Maximize_Disabled", FVector2f(23.f, 18.f) ) );
 
 		const FButtonStyle RestoreButtonStyle = FButtonStyle(Button)
-			.SetNormal ( IMAGE_BRUSH( "Old/Window/WindowButton_Restore_Normal", FVector2D(23, 18) ) )
-			.SetHovered( IMAGE_BRUSH( "Old/Window/WindowButton_Restore_Hovered", FVector2D(23, 18) ) )
-			.SetPressed( IMAGE_BRUSH( "Old/Window/WindowButton_Restore_Pressed", FVector2D(23, 18) ) );
+			.SetNormal ( IMAGE_BRUSH( "Old/Window/WindowButton_Restore_Normal", FVector2f(23.f, 18.f) ) )
+			.SetHovered( IMAGE_BRUSH( "Old/Window/WindowButton_Restore_Hovered", FVector2f(23.f, 18.f) ) )
+			.SetPressed( IMAGE_BRUSH( "Old/Window/WindowButton_Restore_Pressed", FVector2f(23.f, 18.f) ) );
 
-		Set( "Window.Buttons.Restore.Normal",     new IMAGE_BRUSH( "Old/Window/WindowButton_Restore_Normal", FVector2D(23, 18) ) );
-		Set( "Window.Buttons.Restore.Hovered",    new IMAGE_BRUSH( "Old/Window/WindowButton_Restore_Hovered", FVector2D(23, 18) ) );
-		Set( "Window.Buttons.Restore.Pressed",    new IMAGE_BRUSH( "Old/Window/WindowButton_Restore_Pressed", FVector2D(23, 18) ) );	
+		Set( "Window.Buttons.Restore.Normal",     new IMAGE_BRUSH( "Old/Window/WindowButton_Restore_Normal", FVector2f(23.f, 18.f) ) );
+		Set( "Window.Buttons.Restore.Hovered",    new IMAGE_BRUSH( "Old/Window/WindowButton_Restore_Hovered", FVector2f(23.f, 18.f) ) );
+		Set( "Window.Buttons.Restore.Pressed",    new IMAGE_BRUSH( "Old/Window/WindowButton_Restore_Pressed", FVector2f(23.f, 18.f) ) );	
 
 		Set( "Window.Buttons.Minimize", MinimizeButtonStyle );
 		Set( "Window.Buttons.Maximize", MaximizeButtonStyle );
@@ -2983,28 +2983,28 @@ void FSlateEditorStyle::FStyle::SetupWindowStyles()
         
 #endif
         const FButtonStyle CloseButtonStyle = FButtonStyle(Button)
-        .SetNormal ( IMAGE_BRUSH( "Old/Window/WindowButton_Close_Normal", FVector2D(44, 18) ) )
-        .SetHovered( IMAGE_BRUSH( "Old/Window/WindowButton_Close_Hovered", FVector2D(44, 18) ) )
-        .SetPressed( IMAGE_BRUSH( "Old/Window/WindowButton_Close_Pressed", FVector2D(44, 18) ) );
+        .SetNormal ( IMAGE_BRUSH( "Old/Window/WindowButton_Close_Normal", FVector2f(44.f, 18.f) ) )
+        .SetHovered( IMAGE_BRUSH( "Old/Window/WindowButton_Close_Hovered", FVector2f(44.f, 18.f) ) )
+        .SetPressed( IMAGE_BRUSH( "Old/Window/WindowButton_Close_Pressed", FVector2f(44.f, 18.f) ) );
         
-        Set( "Window.Buttons.Close.Normal",       new IMAGE_BRUSH( "Old/Window/WindowButton_Close_Normal", FVector2D(44, 18) ) );
-        Set( "Window.Buttons.Close.Hovered",      new IMAGE_BRUSH( "Old/Window/WindowButton_Close_Hovered", FVector2D(44, 18) ) );
-        Set( "Window.Buttons.Close.Pressed",      new IMAGE_BRUSH( "Old/Window/WindowButton_Close_Pressed", FVector2D(44, 18) ) );		Set( "Window.Buttons.Close", CloseButtonStyle );
+        Set( "Window.Buttons.Close.Normal",       new IMAGE_BRUSH( "Old/Window/WindowButton_Close_Normal", FVector2f(44.f, 18.f) ) );
+        Set( "Window.Buttons.Close.Hovered",      new IMAGE_BRUSH( "Old/Window/WindowButton_Close_Hovered", FVector2f(44.f, 18.f) ) );
+        Set( "Window.Buttons.Close.Pressed",      new IMAGE_BRUSH( "Old/Window/WindowButton_Close_Pressed", FVector2f(44.f, 18.f) ) );		Set( "Window.Buttons.Close", CloseButtonStyle );
 
 		// Title Text
 		const FTextBlockStyle TitleTextStyle = FTextBlockStyle(NormalText)
 			.SetFont( DEFAULT_FONT( "Regular", 9 ) )
 			.SetColorAndOpacity( FLinearColor::White )
-			.SetShadowOffset( FVector2D( 1,1 ) )
+			.SetShadowOffset( FVector2f::UnitVector )
 			.SetShadowColorAndOpacity( FLinearColor::Black );
 		Set( "Window.TitleText", TitleTextStyle );
 
-		Set( "ChildWindow.Background",            new IMAGE_BRUSH( "Common/NoiseBackground", FVector2D(64, 64), FLinearColor::White, ESlateBrushTileType::Both) );
+		Set( "ChildWindow.Background",            new IMAGE_BRUSH( "Common/NoiseBackground", Icon64x64, FLinearColor::White, ESlateBrushTileType::Both) );
 
 		// Update the window style in the *core* style, as SWindow is hardcoded to pull from that 
 		FSlateColor WindowBackgroundColor(FLinearColor::White);
-		FSlateBrush WindowBackgroundMain(IMAGE_BRUSH("Old/Window/WindowBackground", FVector2D(74, 74), FLinearColor::White, ESlateBrushTileType::Both));
-		FSlateBrush WindowBackgroundChild(IMAGE_BRUSH("Common/NoiseBackground", FVector2D(64, 64), FLinearColor::White, ESlateBrushTileType::Both));
+		FSlateBrush WindowBackgroundMain(IMAGE_BRUSH("Old/Window/WindowBackground", FVector2f(74.f, 74.f), FLinearColor::White, ESlateBrushTileType::Both));
+		FSlateBrush WindowBackgroundChild(IMAGE_BRUSH("Common/NoiseBackground", Icon64x64, FLinearColor::White, ESlateBrushTileType::Both));
 
 		if (Settings.IsValid())
 		{
@@ -3036,7 +3036,7 @@ void FSlateEditorStyle::FStyle::SetupProjectBadgeStyle()
 		.SetColorAndOpacity(FLinearColor(0.0f, 0.0f, 0.0f, 1.f))
 	);
 
-	Set("SProjectBadge.BadgeShape", new BOX_BRUSH("ProjectBadge/Badge", Icon16x16, FMargin(6.0f/16)));
+	Set("SProjectBadge.BadgeShape", new BOX_BRUSH("ProjectBadge/Badge", Icon16x16, FMargin(6.0f/16.f)));
 	Set("SProjectBadge.BadgePadding", FMargin(10.0f, 2.5f));
 }
 
@@ -3050,7 +3050,7 @@ void FSlateEditorStyle::FStyle::SetupDockingStyles()
 			Set( "Docking.TabFont", FTextBlockStyle(NormalText)
 				.SetFont( DEFAULT_FONT( "Regular", 9 ) )
 				.SetColorAndOpacity( FLinearColor(0.72f, 0.72f, 0.72f, 1.f) )
-				.SetShadowOffset( FVector2D( 1,1 ) )
+				.SetShadowOffset( FVector2f::UnitVector )
 				.SetShadowColorAndOpacity( FLinearColor::Black )
 				);
 		}
@@ -3066,40 +3066,40 @@ void FSlateEditorStyle::FStyle::SetupDockingStyles()
 
 			// Panel Tab
 			// Legacy styles used by other editor specific widgets; see "Docking.Tab" in FCoreStyle for the current tab style
-			Set( "Docking.Tab.Normal", new BOX_BRUSH( "/Docking/Tab_Inactive", 4/16.0f ) );
-			Set( "Docking.Tab.Active", new BOX_BRUSH( "/Docking/Tab_Active", 4/16.0f ) );
-			Set( "Docking.Tab.Foreground", new BOX_BRUSH( "/Docking/Tab_Foreground", 4/16.0f ) );
-			Set( "Docking.Tab.Hovered", new BOX_BRUSH( "/Docking/Tab_Hovered", 4/16.0f ) );
-			Set( "Docking.Tab.ColorOverlay", new BOX_BRUSH( "/Docking/Tab_ColorOverlay", 4/16.0f ) );
-			Set( "Docking.Tab.Padding", FMargin(5, 2, 5, 2) );
+			Set( "Docking.Tab.Normal", new BOX_BRUSH( "/Docking/Tab_Inactive", 4.f/16.0f ) );
+			Set( "Docking.Tab.Active", new BOX_BRUSH( "/Docking/Tab_Active", 4.f/16.0f ) );
+			Set( "Docking.Tab.Foreground", new BOX_BRUSH( "/Docking/Tab_Foreground", 4.f/16.0f ) );
+			Set( "Docking.Tab.Hovered", new BOX_BRUSH( "/Docking/Tab_Hovered", 4.f/16.0f ) );
+			Set( "Docking.Tab.ColorOverlay", new BOX_BRUSH( "/Docking/Tab_ColorOverlay", 4.f/16.0f ) );
+			Set( "Docking.Tab.Padding", FMargin(5.f, 2.f, 5.f, 2.f) );
 			Set( "Docking.Tab.OverlapWidth", -1.0f );
-			Set( "Docking.Tab.ContentAreaBrush", new BOX_BRUSH( "/Docking/TabContentArea", FMargin(4/16.0f) ) );
-			Set( "Docking.Tab.TabWellBrush", new IMAGE_BRUSH( "/Docking/TabWellSeparator", FVector2D(16,4) ) );
-			Set( "Docking.Tab.TabWellPadding", FMargin(0, 0, 4, 0) );
+			Set( "Docking.Tab.ContentAreaBrush", new BOX_BRUSH( "/Docking/TabContentArea", FMargin(4.f/16.0f) ) );
+			Set( "Docking.Tab.TabWellBrush", new IMAGE_BRUSH( "/Docking/TabWellSeparator", FVector2f(16.f,4.f) ) );
+			Set( "Docking.Tab.TabWellPadding", FMargin(0.f, 0.f, 4.f, 0.f) );
 			Set( "Docking.Tab.FlashColor", TabFlashColor );
 			Set( "Docking.Tab.CloseButton", CloseButton );
 
 			// App Tab
 			// Legacy styles used by other editor specific widgets; see "Docking.MajorTab" in FCoreStyle for the current tab style
-			Set( "Docking.MajorTab.Normal", new BOX_BRUSH( "/Docking/AppTab_Inactive", FMargin(24.0f/64.0f, 4/32.0f) ) );
-			Set( "Docking.MajorTab.Active", new BOX_BRUSH( "/Docking/AppTab_Active", FMargin(24.0f/64.0f, 4/32.0f) ) );
-			Set( "Docking.MajorTab.ColorOverlay", new BOX_BRUSH( "/Docking/AppTab_ColorOverlayIcon", FMargin(24.0f/64.0f, 4/32.0f) ) );
-			Set( "Docking.MajorTab.Foreground", new BOX_BRUSH( "/Docking/AppTab_Foreground", FMargin(24.0f/64.0f, 4/32.0f) ) );
-			Set( "Docking.MajorTab.Hovered", new BOX_BRUSH( "/Docking/AppTab_Hovered", FMargin(24.0f/64.0f, 4/32.0f) ) );
-			Set( "Docking.MajorTab.Padding", FMargin(17, 4, 15, 4) );
+			Set( "Docking.MajorTab.Normal", new BOX_BRUSH( "/Docking/AppTab_Inactive", FMargin(24.0f/64.0f, 4.f/32.0f) ) );
+			Set( "Docking.MajorTab.Active", new BOX_BRUSH( "/Docking/AppTab_Active", FMargin(24.0f/64.0f, 4.f/32.0f) ) );
+			Set( "Docking.MajorTab.ColorOverlay", new BOX_BRUSH( "/Docking/AppTab_ColorOverlayIcon", FMargin(24.0f/64.0f, 4.f/32.0f) ) );
+			Set( "Docking.MajorTab.Foreground", new BOX_BRUSH( "/Docking/AppTab_Foreground", FMargin(24.0f/64.0f, 4.f/32.0f) ) );
+			Set( "Docking.MajorTab.Hovered", new BOX_BRUSH( "/Docking/AppTab_Hovered", FMargin(24.0f/64.0f, 4.f/32.0f) ) );
+			Set( "Docking.MajorTab.Padding", FMargin(17.f, 4.f, 15.f, 4.f) );
 			Set( "Docking.MajorTab.OverlapWidth", 21.0f );
-			Set( "Docking.MajorTab.ContentAreaBrush", new BOX_BRUSH( "/Docking/AppTabContentArea", FMargin(4/16.0f) ) );
-			Set( "Docking.MajorTab.TabWellBrush", new IMAGE_BRUSH( "/Docking/AppTabWellSeparator", FVector2D(16,2) ) );
-			Set( "Docking.MajorTab.TabWellPadding", FMargin(0, 2, 0, 0) );
+			Set( "Docking.MajorTab.ContentAreaBrush", new BOX_BRUSH( "/Docking/AppTabContentArea", FMargin(4.f/16.0f) ) );
+			Set( "Docking.MajorTab.TabWellBrush", new IMAGE_BRUSH( "/Docking/AppTabWellSeparator", FVector2f(16.f,2.f) ) );
+			Set( "Docking.MajorTab.TabWellPadding", FMargin(0.f, 2.f, 0.f, 0.f) );
 			Set( "Docking.MajorTab.FlashColor", TabFlashColor );
 			Set( "Docking.MajorTab.CloseButton", FButtonStyle( CloseButton ) );
 		}
 
 		Set( "Docking.DefaultTabIcon", new IMAGE_BRUSH( "Old/Docking/DefaultTabIcon", Icon16x16 ) );
 
-		Set( "Docking.TabConextButton.Normal", new IMAGE_BRUSH( "/Docking/TabContextButton", FVector2D(24,24) ) );
-		Set( "Docking.TabConextButton.Pressed", new IMAGE_BRUSH( "/Docking/TabContextButton", FVector2D(24,24) ) );
-		Set( "Docking.TabConextButton.Hovered", new IMAGE_BRUSH( "/Docking/TabContextButton", FVector2D(24,24) ) );
+		Set( "Docking.TabConextButton.Normal", new IMAGE_BRUSH( "/Docking/TabContextButton", Icon24x24 ) );
+		Set( "Docking.TabConextButton.Pressed", new IMAGE_BRUSH( "/Docking/TabContextButton", Icon24x24 ) );
+		Set( "Docking.TabConextButton.Hovered", new IMAGE_BRUSH( "/Docking/TabContextButton", Icon24x24 ) );
 	}
 #endif // WITH_EDITOR || (IS_PROGRAM && WITH_UNREAL_DEVELOPER_TOOLS)
 	}
@@ -3261,13 +3261,13 @@ void FSlateEditorStyle::FStyle::SetupTutorialStyles()
 		Set( "Tutorials.Content.Color", FLinearColor(1.0f,1.0f,1.0f,0.9f) );
 		Set( "Tutorials.Content.Color.Hovered", FLinearColor(1.0f,1.0f,1.0f,1.0f) );
 
-		Set( "Tutorials.Browser.CategoryArrow", new IMAGE_BRUSH( "Tutorials/BrowserCategoryArrow", FVector2D(24.0f, 24.0f), FSlateColor::UseForeground() ) );
-		Set( "Tutorials.Browser.DefaultTutorialIcon", new IMAGE_BRUSH( "Tutorials/DefaultTutorialIcon_40x", FVector2D(40.0f, 40.0f), FLinearColor::White ) );
-		Set( "Tutorials.Browser.DefaultCategoryIcon", new IMAGE_BRUSH( "Tutorials/DefaultCategoryIcon_40x", FVector2D(40.0f, 40.0f), FLinearColor::White ) );
+		Set( "Tutorials.Browser.CategoryArrow", new IMAGE_BRUSH( "Tutorials/BrowserCategoryArrow", FVector2f(24.0f, 24.0f), FSlateColor::UseForeground() ) );
+		Set( "Tutorials.Browser.DefaultTutorialIcon", new IMAGE_BRUSH( "Tutorials/DefaultTutorialIcon_40x", FVector2f(40.0f, 40.0f), FLinearColor::White ) );
+		Set( "Tutorials.Browser.DefaultCategoryIcon", new IMAGE_BRUSH( "Tutorials/DefaultCategoryIcon_40x", FVector2f(40.0f, 40.0f), FLinearColor::White ) );
 
-		Set( "Tutorials.Browser.BackButton.Image", new IMAGE_BRUSH( "Tutorials/BrowserBack", FVector2D(32.0f, 32.0f), FLinearColor(1.0f, 1.0f, 1.0f, 1.0f) ) );
-		Set( "Tutorials.Browser.PlayButton.Image", new IMAGE_BRUSH( "Tutorials/BrowserPlay", FVector2D(32.0f, 32.0f), FLinearColor(1.0f, 1.0f, 1.0f, 1.0f) ) );
-		Set( "Tutorials.Browser.RestartButton", new IMAGE_BRUSH( "Tutorials/BrowserRestart", FVector2D(16.0f, 16.0f), FLinearColor(1.0f, 1.0f, 1.0f, 1.0f) ) );
+		Set( "Tutorials.Browser.BackButton.Image", new IMAGE_BRUSH( "Tutorials/BrowserBack", FVector2f(32.0f, 32.0f), FLinearColor(1.0f, 1.0f, 1.0f, 1.0f) ) );
+		Set( "Tutorials.Browser.PlayButton.Image", new IMAGE_BRUSH( "Tutorials/BrowserPlay", FVector2f(32.0f, 32.0f), FLinearColor(1.0f, 1.0f, 1.0f, 1.0f) ) );
+		Set( "Tutorials.Browser.RestartButton", new IMAGE_BRUSH( "Tutorials/BrowserRestart", FVector2f(16.0f, 16.0f), FLinearColor(1.0f, 1.0f, 1.0f, 1.0f) ) );
 
 		Set( "Tutorials.Browser.Completed", new IMAGE_BRUSH( "Tutorials/TutorialCompleted", Icon32x32 ) );
 
@@ -3276,11 +3276,11 @@ void FSlateEditorStyle::FStyle::SetupTutorialStyles()
 			.SetFontSize(9));
 
 		Set( "Tutorials.Navigation.Button", FButtonStyle()
-			.SetNormal( BOX_BRUSH( "Common/ButtonHoverHint", FMargin(4/16.0f), FLinearColor(0,0,0,0) ) )
-			.SetHovered( BOX_BRUSH( "Common/ButtonHoverHint", FMargin(4/16.0f), FLinearColor(0,0,0,0) ) )
-			.SetPressed( BOX_BRUSH( "Common/ButtonHoverHint", FMargin(4/16.0f), FLinearColor(0,0,0,0) ) )
-			.SetNormalPadding( FMargin(0,0,0,1))
-			.SetPressedPadding( FMargin(0,1,0,0)));
+			.SetNormal( BOX_BRUSH( "Common/ButtonHoverHint", FMargin(4.f/16.0f), FLinearColor::Transparent ) )
+			.SetHovered( BOX_BRUSH( "Common/ButtonHoverHint", FMargin(4.f/16.0f), FLinearColor::Transparent ) )
+			.SetPressed( BOX_BRUSH( "Common/ButtonHoverHint", FMargin(4.f/16.0f), FLinearColor::Transparent ) )
+			.SetNormalPadding( FMargin(0.f,0.f,0.f,1.f))
+			.SetPressedPadding( FMargin(0.f,1.f,0.f,0.f)));
 
 		Set( "Tutorials.Navigation.NextButton", new IMAGE_BRUSH( "Tutorials/NavigationNext", Icon32x32 ) );
 		Set( "Tutorials.Navigation.HomeButton", new IMAGE_BRUSH( "Tutorials/NavigationHome", Icon32x32 ) );
@@ -3292,13 +3292,13 @@ void FSlateEditorStyle::FStyle::SetupTutorialStyles()
 		Set( "Tutorials.ButtonColor", TutorialButtonColor );
 		Set( "Tutorials.ButtonHighlightColor", TutorialSelectionColor );
 		Set( "Tutorials.ButtonDisabledColor", SelectionColor_Inactive );
-		Set( "Tutorials.ContentAreaBackground", new BOX_BRUSH( "Tutorials/TutorialContentBackground", FMargin(4/16.0f) ) );
-		Set( "Tutorials.HomeContentAreaBackground", new BOX_BRUSH( "Tutorials/TutorialHomeContentBackground", FMargin(4/16.0f) ) );
+		Set( "Tutorials.ContentAreaBackground", new BOX_BRUSH( "Tutorials/TutorialContentBackground", FMargin(4.f/16.0f) ) );
+		Set( "Tutorials.HomeContentAreaBackground", new BOX_BRUSH( "Tutorials/TutorialHomeContentBackground", FMargin(4.f/16.0f) ) );
 		Set( "Tutorials.ContentAreaFrame", new BOX_BRUSH( "Tutorials/ContentAreaFrame", FMargin(26.0f/64.0f) ) );
-		Set( "Tutorials.CurrentExcerpt", new IMAGE_BRUSH( "Tutorials/CurrentExcerpt", FVector2D(24.0f, 24.0f), TutorialSelectionColor ) );
-		Set( "Tutorials.Home", new IMAGE_BRUSH( "Tutorials/HomeButton", FVector2D(32.0f, 32.0f) ) );
-		Set( "Tutorials.Back", new IMAGE_BRUSH("Tutorials/BackButton", FVector2D(24.0f, 24.0f) ) );
-		Set( "Tutorials.Next", new IMAGE_BRUSH("Tutorials/NextButton", FVector2D(24.0f, 24.0f) ) );
+		Set( "Tutorials.CurrentExcerpt", new IMAGE_BRUSH( "Tutorials/CurrentExcerpt", FVector2f(24.0f, 24.0f), TutorialSelectionColor ) );
+		Set( "Tutorials.Home", new IMAGE_BRUSH( "Tutorials/HomeButton", FVector2f(32.0f, 32.0f) ) );
+		Set( "Tutorials.Back", new IMAGE_BRUSH("Tutorials/BackButton", FVector2f(24.0f, 24.0f) ) );
+		Set( "Tutorials.Next", new IMAGE_BRUSH("Tutorials/NextButton", FVector2f(24.0f, 24.0f) ) );
 
 		Set("Tutorials.PageHeader", FTextBlockStyle(TutorialHeaderText)
 			.SetFontSize(22));
@@ -3324,15 +3324,15 @@ void FSlateEditorStyle::FStyle::SetupTutorialStyles()
 			.SetFontSize(24));
 
 		Set( "Tutorials.Hyperlink.Button", FButtonStyle(DocumentationHyperlinkButton)
-			.SetNormal(BORDER_BRUSH( "Old/HyperlinkDotted", FMargin(0,0,0,3/16.0f), FLinearColor::Black))
-			.SetHovered(BORDER_BRUSH( "Old/HyperlinkUnderline", FMargin(0,0,0,3/16.0f), FLinearColor::Black)));
+			.SetNormal(BORDER_BRUSH( "Old/HyperlinkDotted", FMargin(0.f,0.f,0.f,3.f/16.0f), FLinearColor::Black))
+			.SetHovered(BORDER_BRUSH( "Old/HyperlinkUnderline", FMargin(0.f,0.f,0.f,3.f/16.0f), FLinearColor::Black)));
 
-		Set( "Tutorials.Separator", new BOX_BRUSH( "Common/Separator", 1/4.0f, FLinearColor::Black));
+		Set( "Tutorials.Separator", new BOX_BRUSH( "Common/Separator", 1.f/4.0f, FLinearColor::Black));
 
 		Set( "Tutorials.ProgressBar", FProgressBarStyle()
 			.SetBackgroundImage( BOX_BRUSH( "Common/ProgressBar_Background", FMargin(5.f/12.f) ) )
 			.SetFillImage( BOX_BRUSH( "Common/ProgressBar_NeutralFill", FMargin(5.f/12.f) ) )
-			.SetMarqueeImage( IMAGE_BRUSH( "Common/ProgressBar_Marquee", FVector2D(20,12), FLinearColor::White, ESlateBrushTileType::Horizontal ) )
+			.SetMarqueeImage( IMAGE_BRUSH( "Common/ProgressBar_Marquee", FVector2f(20.f,12.f), FLinearColor::White, ESlateBrushTileType::Horizontal ) )
 			);
 
 		// Default text styles
@@ -3340,7 +3340,7 @@ void FSlateEditorStyle::FStyle::SetupTutorialStyles()
 			const FTextBlockStyle RichTextNormal = FTextBlockStyle()
 				.SetFont(DEFAULT_FONT("Regular", 11))
 				.SetColorAndOpacity(FSlateColor::UseForeground())
-				.SetShadowOffset(FVector2D::ZeroVector)
+				.SetShadowOffset(FVector2f::ZeroVector)
 				.SetShadowColorAndOpacity(FLinearColor::Black)
 				.SetHighlightColor(FLinearColor(0.02f, 0.3f, 0.0f))
 				.SetHighlightShape(BOX_BRUSH("Common/TextBlockHighlightShape", FMargin(3.f /8.f)));
@@ -3427,40 +3427,40 @@ void FSlateEditorStyle::FStyle::SetupTutorialStyles()
 
 			Set("TutorialEditableText.Toolbar.ToggleButtonCheckbox", FCheckBoxStyle()
 				.SetCheckBoxType(ESlateCheckBoxType::ToggleButton)
-				.SetUncheckedImage(BOX_BRUSH("Tutorials/FlatColorSquare", FVector2D(1.0f, 1.0f), FMargin(1), NormalColor))
-				.SetUncheckedHoveredImage(BOX_BRUSH("Tutorials/FlatColorSquare", FVector2D(1.0f, 1.0f), FMargin(1), HoverColor))
-				.SetUncheckedPressedImage(BOX_BRUSH("Tutorials/FlatColorSquare", FVector2D(1.0f, 1.0f), FMargin(1), HoverColor))
-				.SetCheckedImage(BOX_BRUSH("Tutorials/FlatColorSquare", FVector2D(1.0f, 1.0f), FMargin(1), SelectedColor))
-				.SetCheckedHoveredImage(BOX_BRUSH("Tutorials/FlatColorSquare", FVector2D(1.0f, 1.0f), FMargin(1), HoverColor))
-				.SetCheckedPressedImage(BOX_BRUSH("Tutorials/FlatColorSquare", FVector2D(1.0f, 1.0f), FMargin(1), HoverColor))
+				.SetUncheckedImage(BOX_BRUSH("Tutorials/FlatColorSquare", FVector2f::UnitVector, FMargin(1.f), NormalColor))
+				.SetUncheckedHoveredImage(BOX_BRUSH("Tutorials/FlatColorSquare", FVector2f::UnitVector, FMargin(1.f), HoverColor))
+				.SetUncheckedPressedImage(BOX_BRUSH("Tutorials/FlatColorSquare", FVector2f::UnitVector, FMargin(1.f), HoverColor))
+				.SetCheckedImage(BOX_BRUSH("Tutorials/FlatColorSquare", FVector2f::UnitVector, FMargin(1.f), SelectedColor))
+				.SetCheckedHoveredImage(BOX_BRUSH("Tutorials/FlatColorSquare", FVector2f::UnitVector, FMargin(1.f), HoverColor))
+				.SetCheckedPressedImage(BOX_BRUSH("Tutorials/FlatColorSquare", FVector2f::UnitVector, FMargin(1.f), HoverColor))
 				);
 
 			const FButtonStyle TutorialButton = FButtonStyle()
-				.SetNormal(BOX_BRUSH("Tutorials/FlatColorSquare", FVector2D(1.0f, 1.0f), FMargin(1), NormalColor))
-				.SetHovered(BOX_BRUSH("Tutorials/FlatColorSquare", FVector2D(1.0f, 1.0f), FMargin(1), HoverColor))
-				.SetPressed(BOX_BRUSH("Tutorials/FlatColorSquare", FVector2D(1.0f, 1.0f), FMargin(1), SelectedColor))
-				.SetNormalPadding(FMargin(2, 2, 2, 2))
-				.SetPressedPadding(FMargin(2, 3, 2, 1));
+				.SetNormal(BOX_BRUSH("Tutorials/FlatColorSquare", FVector2f::UnitVector, FMargin(1.f), NormalColor))
+				.SetHovered(BOX_BRUSH("Tutorials/FlatColorSquare", FVector2f::UnitVector, FMargin(1.f), HoverColor))
+				.SetPressed(BOX_BRUSH("Tutorials/FlatColorSquare", FVector2f::UnitVector, FMargin(1.f), SelectedColor))
+				.SetNormalPadding(FMargin(2.f))
+				.SetPressedPadding(FMargin(2.f, 3.f, 2.f, 1.f));
 			Set("TutorialEditableText.Toolbar.Button", TutorialButton);
 
 			const FComboButtonStyle ComboButton = FComboButtonStyle()
 				.SetButtonStyle(Button)
 				.SetDownArrowImage(IMAGE_BRUSH("Common/ComboArrow", Icon8x8))
-				.SetMenuBorderBrush(BOX_BRUSH("Tutorials/FlatColorSquare", FVector2D(1.0f, 1.0f), FMargin(1), NormalColor))
+				.SetMenuBorderBrush(BOX_BRUSH("Tutorials/FlatColorSquare", FVector2f::UnitVector, FMargin(1.f), NormalColor))
 				.SetMenuBorderPadding(FMargin(0.0f));
 			Set("TutorialEditableText.Toolbar.ComboButton", ComboButton);
 
 			{
 				const FButtonStyle ComboBoxButton = FButtonStyle()
-					.SetNormal(BOX_BRUSH("Tutorials/FlatColorSquare", FVector2D(1.0f, 1.0f), FMargin(1), FLinearColor::White))
-					.SetHovered(BOX_BRUSH("Tutorials/FlatColorSquare", FVector2D(1.0f, 1.0f), FMargin(1), FLinearColor::White))
-					.SetPressed(BOX_BRUSH("Tutorials/FlatColorSquare", FVector2D(1.0f, 1.0f), FMargin(1), FLinearColor::White))
-					.SetNormalPadding(FMargin(2, 2, 2, 2))
-					.SetPressedPadding(FMargin(2, 3, 2, 1));
+					.SetNormal(BOX_BRUSH("Tutorials/FlatColorSquare", FVector2f::UnitVector, FMargin(1.f), FLinearColor::White))
+					.SetHovered(BOX_BRUSH("Tutorials/FlatColorSquare", FVector2f::UnitVector, FMargin(1.f), FLinearColor::White))
+					.SetPressed(BOX_BRUSH("Tutorials/FlatColorSquare", FVector2f::UnitVector, FMargin(1.f), FLinearColor::White))
+					.SetNormalPadding(FMargin(2.f))
+					.SetPressedPadding(FMargin(2.f, 3.f, 2.f, 1.f));
 
 				const FComboButtonStyle ComboBoxComboButton = FComboButtonStyle(ComboButton)
 					.SetButtonStyle(ComboBoxButton)
-					.SetMenuBorderPadding(FMargin(1.0));
+					.SetMenuBorderPadding(FMargin(1.0f));
 
 				Set("TutorialEditableText.Toolbar.ComboBox", FComboBoxStyle()
 					.SetComboButtonStyle(ComboBoxComboButton)
@@ -3471,8 +3471,8 @@ void FSlateEditorStyle::FStyle::SetupTutorialStyles()
 		// In-editor tutorial launch button
 		{
 			Set("TutorialLaunch.Button", FButtonStyle()
-				.SetNormalPadding(0)
-				.SetPressedPadding(0)
+				.SetNormalPadding(0.f)
+				.SetPressedPadding(0.f)
 				.SetNormal(IMAGE_BRUSH("Tutorials/TutorialButton_Default_16x", Icon16x16))
 				.SetHovered(IMAGE_BRUSH("Tutorials/TutorialButton_Hovered_16x", Icon16x16))
 				.SetPressed(IMAGE_BRUSH("Tutorials/TutorialButton_Pressed_16x", Icon16x16))
@@ -3494,7 +3494,7 @@ void FSlateEditorStyle::FStyle::SetupPropertyEditorStyles()
 		Set( "PropertyEditor.RemoveColumn", new IMAGE_BRUSH( "Common/PushPin_Down", Icon16x16, FColor( 96, 194, 253, 255 ).ReinterpretAsLinear() ) );
 		Set( "PropertyEditor.AddColumn", new IMAGE_BRUSH( "Common/PushPin_Up", Icon16x16, FColor( 96, 194, 253, 255 ).ReinterpretAsLinear() ) );
 
-		Set( "PropertyEditor.AddColumnOverlay",	new IMAGE_BRUSH( "Common/TinyChalkArrow", FVector2D( 71, 20 ), FColor( 96, 194, 253, 255 ).ReinterpretAsLinear() ) );
+		Set( "PropertyEditor.AddColumnOverlay",	new IMAGE_BRUSH( "Common/TinyChalkArrow", FVector2f( 71.f, 20.f), FColor( 96, 194, 253, 255 ).ReinterpretAsLinear() ) );
 		Set( "PropertyEditor.AddColumnMessage", FTextBlockStyle(NormalText)
 			.SetFont( DEFAULT_FONT( "BoldCondensedItalic", 10 ) )
 			.SetColorAndOpacity(FColor( 96, 194, 253, 255 ).ReinterpretAsLinear())
@@ -3503,7 +3503,7 @@ void FSlateEditorStyle::FStyle::SetupPropertyEditorStyles()
 		Set( "PropertyEditor.AssetName", FTextBlockStyle(NormalText)
 			.SetFont( DEFAULT_FONT( "Regular", 12 ) )
 			.SetColorAndOpacity( FLinearColor::White )
-			.SetShadowOffset( FVector2D(1,1) )
+			.SetShadowOffset( FVector2f::UnitVector )
 			.SetShadowColorAndOpacity( FLinearColor::Black )
 		);
 
@@ -3514,20 +3514,20 @@ void FSlateEditorStyle::FStyle::SetupPropertyEditorStyles()
 		Set( "PropertyEditor.AssetClass", FTextBlockStyle(NormalText)
 			.SetFont( DEFAULT_FONT( "Regular", 10 ) )
 			.SetColorAndOpacity( FLinearColor::White )
-			.SetShadowOffset( FVector2D(1,1) )
+			.SetShadowOffset( FVector2f::UnitVector )
 			.SetShadowColorAndOpacity( FLinearColor::Black )
 		);
 
 		const FButtonStyle AssetComboStyle = FButtonStyle()
-			.SetNormal( BOX_BRUSH( "Common/ButtonHoverHint", FMargin(4/16.0f), FLinearColor(1,1,1,0.15f) ) )
-			.SetHovered( BOX_BRUSH( "Common/ButtonHoverHint", FMargin(4/16.0f), FLinearColor(1,1,1,0.25f) ) )
-			.SetPressed( BOX_BRUSH( "Common/ButtonHoverHint", FMargin(4/16.0f), FLinearColor(1,1,1,0.30f) ) )
-			.SetNormalPadding( FMargin(0,0,0,1) )
-			.SetPressedPadding( FMargin(0,1,0,0) );
+			.SetNormal( BOX_BRUSH( "Common/ButtonHoverHint", FMargin(4.f/16.0f), FLinearColor(1.f,1.f,1.f,0.15f) ) )
+			.SetHovered( BOX_BRUSH( "Common/ButtonHoverHint", FMargin(4.f/16.0f), FLinearColor(1.f,1.f,1.f,0.25f) ) )
+			.SetPressed( BOX_BRUSH( "Common/ButtonHoverHint", FMargin(4.f/16.0f), FLinearColor(1.f,1.f,1.f,0.30f) ) )
+			.SetNormalPadding( FMargin(0.f,0.f,0.f,1.f) )
+			.SetPressedPadding( FMargin(0.f,1.f,0.f,0.f) );
 		Set( "PropertyEditor.AssetComboStyle", AssetComboStyle );
 
-		Set( "PropertyEditor.HorizontalDottedLine",		new IMAGE_BRUSH( "Common/HorizontalDottedLine_16x1px", FVector2D(16.0f, 1.0f), FLinearColor::White, ESlateBrushTileType::Horizontal ) );
-		Set( "PropertyEditor.VerticalDottedLine",		new IMAGE_BRUSH( "Common/VerticalDottedLine_1x16px", FVector2D(1.0f, 16.0f), FLinearColor::White, ESlateBrushTileType::Vertical ) );
+		Set( "PropertyEditor.HorizontalDottedLine",		new IMAGE_BRUSH( "Common/HorizontalDottedLine_16x1px", FVector2f(16.0f, 1.0f), FLinearColor::White, ESlateBrushTileType::Horizontal ) );
+		Set( "PropertyEditor.VerticalDottedLine",		new IMAGE_BRUSH( "Common/VerticalDottedLine_1x16px", FVector2f(1.0f, 16.0f), FLinearColor::White, ESlateBrushTileType::Vertical ) );
 		Set( "PropertyEditor.SlateBrushPreview",		new BOX_BRUSH( "PropertyView/SlateBrushPreview_32px", Icon32x32, FMargin(3.f/32.f, 3.f/32.f, 15.f/32.f, 13.f/32.f) ) );
 
 		Set( "PropertyTable.TableRow", FTableRowStyle()
@@ -3603,7 +3603,7 @@ void FSlateEditorStyle::FStyle::SetupPropertyEditorStyles()
 		Set( "PropertyWindow.Favorites_Disabled", new IMAGE_BRUSH( "Icons/EmptyStar_16x", Icon16x16 ) );
 		Set( "PropertyWindow.Locked", new IMAGE_BRUSH( "Icons/padlock_locked_16x", Icon16x16 ) );
 		Set( "PropertyWindow.Unlocked", new IMAGE_BRUSH( "Icons/padlock_unlocked_16x", Icon16x16 ) );
-		Set( "PropertyWindow.DiffersFromDefault", new IMAGE_BRUSH( "/PropertyView/DiffersFromDefault_8x8", FVector2D(8,8) ) ) ;
+		Set( "PropertyWindow.DiffersFromDefault", new IMAGE_BRUSH( "/PropertyView/DiffersFromDefault_8x8", Icon8x8 ) ) ;
 		
 		Set( "PropertyWindow.NormalFont", DEFAULT_FONT("Regular", 8) );
 		Set( "PropertyWindow.BoldFont", DEFAULT_FONT( "Bold", 8 ) );
@@ -3638,9 +3638,9 @@ void FSlateEditorStyle::FStyle::SetupPropertyEditorStyles()
 		FTextBlockStyle BPWarningMessageHyperlinkTextStyle = FTextBlockStyle(BPWarningMessageTextStyle).SetColorAndOpacity(FLinearColor(0.25f, 0.5f, 1.0f));
 
 		FButtonStyle EditBPHyperlinkButton = FButtonStyle()
-			.SetNormal(BORDER_BRUSH("Old/HyperlinkDotted", FMargin(0, 0, 0, 3 / 16.0f), FLinearColor(0.25f, 0.5f, 1.0f)))
+			.SetNormal(BORDER_BRUSH("Old/HyperlinkDotted", FMargin(0.f, 0.f, 0.f, 3.f / 16.0f), FLinearColor(0.25f, 0.5f, 1.0f)))
 			.SetPressed(FSlateNoResource())
-			.SetHovered(BORDER_BRUSH("Old/HyperlinkUnderline", FMargin(0, 0, 0, 3 / 16.0f), FLinearColor(0.25f, 0.5f, 1.0f)));
+			.SetHovered(BORDER_BRUSH("Old/HyperlinkUnderline", FMargin(0.f, 0.f, 0.f, 3.f / 16.0f), FLinearColor(0.25f, 0.5f, 1.0f)));
 
 		FHyperlinkStyle BPWarningMessageHyperlinkStyle = FHyperlinkStyle()
 			.SetUnderlineStyle(EditBPHyperlinkButton)
@@ -3653,21 +3653,21 @@ void FSlateEditorStyle::FStyle::SetupPropertyEditorStyles()
 
 		Set( "DetailsView.GroupSection", new BOX_BRUSH( "Common/RoundedSelection_16x", FMargin(4.0f/16.0f) ) );
 
-		Set( "DetailsView.PulldownArrow.Down", new IMAGE_BRUSH( "PropertyView/AdvancedButton_Down", FVector2D(10,8) ) );
-		Set( "DetailsView.PulldownArrow.Down.Hovered", new IMAGE_BRUSH( "PropertyView/AdvancedButton_Down_Hovered", FVector2D(10,8) ) );
-		Set( "DetailsView.PulldownArrow.Up", new IMAGE_BRUSH( "PropertyView/AdvancedButton_Up", FVector2D(10,8) ) );
-		Set( "DetailsView.PulldownArrow.Up.Hovered", new IMAGE_BRUSH( "PropertyView/AdvancedButton_Up_Hovered", FVector2D(10,8) ) );
+		Set( "DetailsView.PulldownArrow.Down", new IMAGE_BRUSH( "PropertyView/AdvancedButton_Down", FVector2f(10.f,8.f) ) );
+		Set( "DetailsView.PulldownArrow.Down.Hovered", new IMAGE_BRUSH( "PropertyView/AdvancedButton_Down_Hovered", FVector2f(10.f,8.f) ) );
+		Set( "DetailsView.PulldownArrow.Up", new IMAGE_BRUSH( "PropertyView/AdvancedButton_Up", FVector2f(10.f,8.f) ) );
+		Set( "DetailsView.PulldownArrow.Up.Hovered", new IMAGE_BRUSH( "PropertyView/AdvancedButton_Up_Hovered", FVector2f(10.f,8.f) ) );
 
 		Set( "DetailsView.EditRawProperties", new IMAGE_BRUSH( "Icons/icon_PropertyMatrix_16px", Icon16x16, FLinearColor::Black ) );
 		Set( "DetailsView.EditConfigProperties", new IMAGE_BRUSH( "Icons/icon_PropertyMatrix_16px", Icon16x16, FLinearColor::White ) );
 
 		Set( "DetailsView.CollapsedCategory", new BOX_BRUSH( "Common/GroupBorder", FMargin(4.0f/16.0f) ) );
-		Set( "DetailsView.CategoryTop", new BOX_BRUSH( "PropertyView/DetailCategoryTop", FMargin( 4/16.0f, 8.0f/16.0f, 4/16.0f, 4/16.0f ) ) );
+		Set( "DetailsView.CategoryTop", new BOX_BRUSH( "PropertyView/DetailCategoryTop", FMargin( 4.f/16.0f, 8.0f/16.0f, 4.f/16.0f, 4.f/16.0f ) ) );
 		Set( "DetailsView.CollapsedCategory_Hovered", new BOX_BRUSH( "Common/GroupBorder", FMargin(4.0f/16.0f), FLinearColor(0.5f,0.5f,0.5f,1.0f)  ) );
-		Set( "DetailsView.CategoryTop_Hovered", new BOX_BRUSH( "PropertyView/DetailCategoryTop", FMargin( 4/16.0f, 8.0f/16.0f, 4/16.0f, 4/16.0f ), FLinearColor(0.5f,0.5f,0.5f,1.0f) ) );
+		Set( "DetailsView.CategoryTop_Hovered", new BOX_BRUSH( "PropertyView/DetailCategoryTop", FMargin( 4.f/16.0f, 8.0f/16.0f, 4.f/16.0f, 4.f/16.0f ), FLinearColor(0.5f,0.5f,0.5f,1.0f) ) );
 		Set( "DetailsView.CategoryBottom", new BOX_BRUSH( "PropertyView/DetailCategoryBottom", FMargin(4.0f/16.0f) ) );
-		Set( "DetailsView.CategoryMiddle", new IMAGE_BRUSH( "PropertyView/DetailCategoryMiddle", FVector2D( 16, 16 ) ) );
-		Set( "DetailsView.CategoryMiddle_Hovered", new IMAGE_BRUSH( "PropertyView/DetailCategoryMiddle_Hovered", FVector2D( 16, 16 ) ) );
+		Set( "DetailsView.CategoryMiddle", new IMAGE_BRUSH( "PropertyView/DetailCategoryMiddle", Icon16x16 ) );
+		Set( "DetailsView.CategoryMiddle_Hovered", new IMAGE_BRUSH( "PropertyView/DetailCategoryMiddle_Hovered", Icon16x16 ) );
 		Set( "DetailsView.CategoryMiddle_Highlighted", new BOX_BRUSH( "Common/TextBox_Special_Active", FMargin(8.0f/32.0f) ) );
 		Set( "DetailsView.CategoryMiddle_Active", new BOX_BRUSH( "Common/TextBox_Special_Active", FMargin(8.0f/32.0f), SelectionColor_Pressed ) );
 		
@@ -3681,13 +3681,13 @@ void FSlateEditorStyle::FStyle::SetupPropertyEditorStyles()
 			);
 
 		Set( "DetailsView.AdvancedDropdownBorder", new BOX_BRUSH( "PropertyView/DetailCategoryAdvanced", FMargin(4.0f/16.0f) ) );
-		Set( "DetailsView.AdvancedDropdownBorder.Open", new IMAGE_BRUSH( "Common/ScrollBoxShadowTop", FVector2D(64,8) ) );
+		Set( "DetailsView.AdvancedDropdownBorder.Open", new IMAGE_BRUSH( "Common/ScrollBoxShadowTop", FVector2f(64.f,8.f) ) );
 		Set( "DetailsView.CategoryFontStyle", DEFAULT_FONT( "Bold", 10 ) );
 
 		Set( "DetailsView.CategoryTextStyle", 
 			FTextBlockStyle(NormalText)
 			.SetFont(GetFontStyle("DetailsView.CategoryFontStyle"))
-			.SetShadowOffset(FVector2D(1.0f, 1.0f))
+			.SetShadowOffset(FVector2f::UnitVector)
 		);
 
 		Set( "DetailsView.TreeView.TableRow", FTableRowStyle()
@@ -3704,7 +3704,7 @@ void FSlateEditorStyle::FStyle::SetupPropertyEditorStyles()
 			.SetSelectedTextColor( InvertedForeground )
 			);
 
-		Set("DetailsView.DropZone.Below", new BOX_BRUSH("Common/VerticalBoxDropZoneIndicator_Below", FMargin(10.0f / 16.0f, 0, 0, 10.0f / 16.0f), SelectionColor_Subdued));
+		Set("DetailsView.DropZone.Below", new BOX_BRUSH("Common/VerticalBoxDropZoneIndicator_Below", FMargin(10.0f / 16.0f, 0.f, 0.f, 10.0f / 16.0f), SelectionColor_Subdued));
 
 	}
 	}
@@ -3769,28 +3769,28 @@ void FSlateEditorStyle::FStyle::SetupProfilerStyle()
 		Set( "Profiler.CaptionBold", FTextBlockStyle(NormalText)
 			.SetFont( DEFAULT_FONT( "Bold", 10 ) )
 			.SetColorAndOpacity( FLinearColor::White )
-			.SetShadowOffset( FVector2D(1.0f, 1.0f) )
+			.SetShadowOffset( FVector2f::UnitVector )
 			.SetShadowColorAndOpacity( FLinearColor(0.f,0.f,0.f,0.8f) )
 		);
 
 		Set( "Profiler.Caption", FTextBlockStyle(NormalText)
 			.SetFont( DEFAULT_FONT( "Regular", 10 ) )
 			.SetColorAndOpacity( FLinearColor::White )
-			.SetShadowOffset( FVector2D(1.0f, 1.0f) )
+			.SetShadowOffset( FVector2f::UnitVector )
 			.SetShadowColorAndOpacity( FLinearColor(0.f,0.f,0.f,0.8f) )
 		);
 
 		Set( "Profiler.TooltipBold", FTextBlockStyle(NormalText)
 			.SetFont( DEFAULT_FONT( "Bold", 8 ) )
 			.SetColorAndOpacity( FLinearColor(0.5f,0.5f,0.5f,1.0f) )
-			.SetShadowOffset( FVector2D(1.0f, 1.0f) )
+			.SetShadowOffset( FVector2f::UnitVector )
 			.SetShadowColorAndOpacity( FLinearColor(0.f,0.f,0.f,0.8f) )
 		);
 
 		Set( "Profiler.Tooltip", FTextBlockStyle(NormalText)
 			.SetFont( DEFAULT_FONT( "Regular", 8 ) )
 			.SetColorAndOpacity( FLinearColor::White )
-			.SetShadowOffset( FVector2D(1.0f, 1.0f) )
+			.SetShadowOffset( FVector2f::UnitVector )
 			.SetShadowColorAndOpacity( FLinearColor(0.f,0.f,0.f,0.8f) )
 		);
 
@@ -3838,7 +3838,7 @@ void FSlateEditorStyle::FStyle::SetupProfilerStyle()
 		Set( "Profiler.EventGraph.DarkText", FTextBlockStyle(NormalText)
 			.SetFont( DEFAULT_FONT( "Regular", 8 ) )
 			.SetColorAndOpacity( FLinearColor::Black )
-			.SetShadowOffset( FVector2D(0.0f, 0.0f) )
+			.SetShadowOffset( FVector2f::ZeroVector )
 			);
 
 		// Thread-view
@@ -3891,7 +3891,7 @@ void FSlateEditorStyle::FStyle::SetupGraphEditorStyles()
 	{
 		Set( "Graph.ForegroundColor", FLinearColor(218.0f/255.0f, 218.0f/255.0f, 218.0f/255.0f, 1.0f) );
 
-		Set( "Graph.TitleBackground", new BOX_BRUSH( "Old/Graph/GraphTitleBackground", FMargin(0) ) );
+		Set( "Graph.TitleBackground", new BOX_BRUSH( "Old/Graph/GraphTitleBackground", FMargin(0.f) ) );
 		Set( "Graph.Shadow", new BOX_BRUSH( "Old/Window/WindowBorder", 0.48f ) );
 		Set( "Graph.Arrow", new IMAGE_BRUSH( "Old/Graph/Arrow", Icon16x16 ) );
 		Set( "Graph.ExecutionBubble", new IMAGE_BRUSH( "Old/Graph/ExecutionBubble", Icon16x16 ) );
@@ -3899,78 +3899,78 @@ void FSlateEditorStyle::FStyle::SetupGraphEditorStyles()
 		Set( "Graph.PlayInEditor", new BOX_BRUSH( "/Graph/RegularNode_shadow_selected", FMargin(18.0f/64.0f) ) );
 		Set( "Graph.ReadOnlyBorder", new BOX_BRUSH( "/Graph/Graph_readonly_border", FMargin(18.0f / 64.0f) ) );
 
-		Set( "Graph.Panel.SolidBackground", new IMAGE_BRUSH( "/Graph/GraphPanel_SolidBackground", FVector2D(16, 16), FLinearColor::White, ESlateBrushTileType::Both) );
-		Set( "Graph.Panel.GridLineColor",   FLinearColor(0.024f, 0.024, 0.024f) );
+		Set( "Graph.Panel.SolidBackground", new IMAGE_BRUSH( "/Graph/GraphPanel_SolidBackground", Icon16x16, FLinearColor::White, ESlateBrushTileType::Both) );
+		Set( "Graph.Panel.GridLineColor",   FLinearColor(0.024f, 0.024f, 0.024f) );
 		Set( "Graph.Panel.GridRuleColor",   FLinearColor(0.010f, 0.010f, 0.010f) );
 		Set( "Graph.Panel.GridCenterColor", FLinearColor(0.005f, 0.005f, 0.005f) );
 		
 		Set( "Graph.Panel.GridRulePeriod", 8.0f ); // should be a strictly positive integral value
 
-		Set( "Graph.Node.Separator", new IMAGE_BRUSH( "Old/Graph/NodeVerticalSeparator", FVector2D(8,8) ) );
-		Set( "Graph.Node.TitleBackground", new BOX_BRUSH( "Old/Graph/NodeTitleBackground", FMargin(12.0f/64) ) );
-		Set( "Graph.Node.NodeBackground", new BOX_BRUSH( "Old/Graph/NodeBackground", FMargin(12.0f/64) ) );
+		Set( "Graph.Node.Separator", new IMAGE_BRUSH( "Old/Graph/NodeVerticalSeparator", Icon8x8 ) );
+		Set( "Graph.Node.TitleBackground", new BOX_BRUSH( "Old/Graph/NodeTitleBackground", FMargin(12.0f/64.f) ) );
+		Set( "Graph.Node.NodeBackground", new BOX_BRUSH( "Old/Graph/NodeBackground", FMargin(12.0f/64.f) ) );
 
 		Set( "Graph.Node.Body", new BOX_BRUSH( "/Graph/RegularNode_body", FMargin(16.f/64.f, 25.f/64.f, 16.f/64.f, 16.f/64.f) ) );
 		Set( "Graph.Node.TintedBody", new BOX_BRUSH( "/Graph/TintedNode_body", FMargin(16.f/64.f, 25.f/64.f, 16.f/64.f, 16.f/64.f) ) );
-		Set( "Graph.Node.DisabledBanner", new IMAGE_BRUSH( "/Graph/GraphPanel_StripesBackground", FVector2D(64, 64), FLinearColor(0.5f, 0.5f, 0.5f, 0.3f), ESlateBrushTileType::Both ) );
-		Set( "Graph.Node.DevelopmentBanner", new IMAGE_BRUSH( "/Graph/GraphPanel_StripesBackground", FVector2D(64, 64), FLinearColor::Yellow * FLinearColor(1.f, 1.f, 1.f, 0.3f), ESlateBrushTileType::Both ) );
+		Set( "Graph.Node.DisabledBanner", new IMAGE_BRUSH( "/Graph/GraphPanel_StripesBackground", Icon64x64, FLinearColor(0.5f, 0.5f, 0.5f, 0.3f), ESlateBrushTileType::Both ) );
+		Set( "Graph.Node.DevelopmentBanner", new IMAGE_BRUSH( "/Graph/GraphPanel_StripesBackground", Icon64x64, FLinearColor::Yellow * FLinearColor(1.f, 1.f, 1.f, 0.3f), ESlateBrushTileType::Both ) );
 		Set( "Graph.Node.TitleGloss", new BOX_BRUSH( "/Graph/RegularNode_title_gloss", FMargin(12.0f/64.0f) ) );
-		Set( "Graph.Node.ColorSpill", new BOX_BRUSH( "/Graph/RegularNode_color_spill", FMargin(8.0f/64.0f, 3.0f/32.0f, 0, 0) ) );
+		Set( "Graph.Node.ColorSpill", new BOX_BRUSH( "/Graph/RegularNode_color_spill", FMargin(8.0f/64.0f, 3.0f/32.0f, 0.f, 0.f) ) );
 		Set( "Graph.Node.TitleHighlight", new BOX_BRUSH( "/Graph/RegularNode_title_highlight", FMargin(16.0f/64.0f, 1.0f, 16.0f/64.0f, 0.0f) ) );
-		Set( "Graph.Node.IndicatorOverlay", new IMAGE_BRUSH( "/Graph/IndicatorOverlay_color_spill", FVector2D(128.f, 32.f) ) );
+		Set( "Graph.Node.IndicatorOverlay", new IMAGE_BRUSH( "/Graph/IndicatorOverlay_color_spill", FVector2f(128.f, 32.f) ) );
 
-		Set( "Graph.Node.ShadowSize", FVector2D(12,12) );
+		Set( "Graph.Node.ShadowSize", Icon12x12 );
 		Set( "Graph.Node.ShadowSelected", new BOX_BRUSH( "/Graph/RegularNode_shadow_selected", FMargin(18.0f/64.0f) ) );
 		Set( "Graph.Node.Shadow", new BOX_BRUSH( "/Graph/RegularNode_shadow", FMargin(18.0f/64.0f) ) );
 
-		Set( "Graph.Node.RerouteShadow", new IMAGE_BRUSH( "/Graph/RerouteNode_shadow", FVector2D(64.0f, 64.0f) ) );
-		Set( "Graph.Node.RerouteShadowSelected", new IMAGE_BRUSH( "/Graph/RerouteNode_shadow_selected", FVector2D(64.0f, 64.0f) ) );
+		Set( "Graph.Node.RerouteShadow", new IMAGE_BRUSH( "/Graph/RerouteNode_shadow", Icon64x64 ) );
+		Set( "Graph.Node.RerouteShadowSelected", new IMAGE_BRUSH( "/Graph/RerouteNode_shadow_selected", Icon64x64 ) );
 
 		Set( "Graph.CompactNode.ShadowSelected", new BOX_BRUSH( "/Graph/MathNode_shadow_selected", FMargin(18.0f/64.0f) ) );
 
-		Set( "Graph.Node.CommentBubble", new BOX_BRUSH( "Old/Graph/CommentBubble", FMargin(8/32.0f) ) );
-		Set( "Graph.Node.CommentArrow", new IMAGE_BRUSH( "Old/Graph/CommentBubbleArrow", FVector2D(8,8) ) );
+		Set( "Graph.Node.CommentBubble", new BOX_BRUSH( "Old/Graph/CommentBubble", FMargin(8.f/32.0f) ) );
+		Set( "Graph.Node.CommentArrow", new IMAGE_BRUSH( "Old/Graph/CommentBubbleArrow", Icon8x8 ) );
 		Set( "Graph.Node.CommentFont", DEFAULT_FONT( "Regular", 10 ) );
-		Set( "Graph.Node.Comment.BubbleOffset", FMargin(8,0,0,0) );
-		Set( "Graph.Node.Comment.PinIconPadding", FMargin(0,2,0,0) );
-		Set("Graph.Node.Comment.BubblePadding", FVector2D(3, 3));
-		Set("Graph.Node.Comment.BubbleWidgetMargin", FMargin(4, 4));
+		Set( "Graph.Node.Comment.BubbleOffset", FMargin(8.f,0.f,0.f,0.f) );
+		Set( "Graph.Node.Comment.PinIconPadding", FMargin(0.f,2.f,0.f,0.f) );
+		Set("Graph.Node.Comment.BubblePadding", FVector2f(3.f, 3.f));
+		Set("Graph.Node.Comment.BubbleWidgetMargin", FMargin(4.f, 4.f));
 
 		const FCheckBoxStyle CommentTitleButton = FCheckBoxStyle()
 			.SetCheckBoxType( ESlateCheckBoxType::CheckBox )
-			.SetUncheckedImage( IMAGE_BRUSH( "Icons/icon_Blueprint_CommentBubbleOn_16x", FVector2D(16,16), FLinearColor(1.f, 1.f, 1.f, 0.8f)))
-			.SetUncheckedHoveredImage( IMAGE_BRUSH( "Icons/icon_Blueprint_CommentBubbleOn_16x", FVector2D(16,16), FLinearColor(1.f, 1.f, 1.f, 0.9f)))
-			.SetUncheckedPressedImage( IMAGE_BRUSH( "Icons/icon_Blueprint_CommentBubbleOn_16x", FVector2D(16,16), FLinearColor(1.f, 1.f, 1.f, 1.f)))
-			.SetCheckedImage( IMAGE_BRUSH( "Icons/icon_Blueprint_CommentBubbleOn_16x", FVector2D(16,16), FLinearColor(1.f, 1.f, 1.f, 0.8f)))
-			.SetCheckedHoveredImage( IMAGE_BRUSH( "Icons/icon_Blueprint_CommentBubbleOn_16x", FVector2D(16,16), FLinearColor(1.f, 1.f, 1.f, 1.f)))
-			.SetCheckedPressedImage( IMAGE_BRUSH( "Icons/icon_Blueprint_CommentBubbleOff_16x", FVector2D(16,16), FLinearColor(1.f, 1.f, 1.f, 0.6f)));
+			.SetUncheckedImage( IMAGE_BRUSH( "Icons/icon_Blueprint_CommentBubbleOn_16x", Icon16x16, FLinearColor(1.f, 1.f, 1.f, 0.8f)))
+			.SetUncheckedHoveredImage( IMAGE_BRUSH( "Icons/icon_Blueprint_CommentBubbleOn_16x", Icon16x16, FLinearColor(1.f, 1.f, 1.f, 0.9f)))
+			.SetUncheckedPressedImage( IMAGE_BRUSH( "Icons/icon_Blueprint_CommentBubbleOn_16x", Icon16x16, FLinearColor(1.f, 1.f, 1.f, 1.f)))
+			.SetCheckedImage( IMAGE_BRUSH( "Icons/icon_Blueprint_CommentBubbleOn_16x", Icon16x16, FLinearColor(1.f, 1.f, 1.f, 0.8f)))
+			.SetCheckedHoveredImage( IMAGE_BRUSH( "Icons/icon_Blueprint_CommentBubbleOn_16x", Icon16x16, FLinearColor(1.f, 1.f, 1.f, 1.f)))
+			.SetCheckedPressedImage( IMAGE_BRUSH( "Icons/icon_Blueprint_CommentBubbleOff_16x", Icon16x16, FLinearColor(1.f, 1.f, 1.f, 0.6f)));
 		Set( "CommentTitleButton", CommentTitleButton );
 
 		const FCheckBoxStyle CommentBubbleButton = FCheckBoxStyle()
 			.SetCheckBoxType( ESlateCheckBoxType::CheckBox )
-			.SetUncheckedImage( IMAGE_BRUSH( "Icons/icon_Blueprint_CommentBubbleOn_16x", FVector2D(10,10), FLinearColor(1.f, 1.f, 1.f, 0.5f)))
-			.SetUncheckedHoveredImage( IMAGE_BRUSH( "Icons/icon_Blueprint_CommentBubbleOn_16x", FVector2D(10,10), FLinearColor(1.f, 1.f, 1.f, 0.9f)))
-			.SetUncheckedPressedImage( IMAGE_BRUSH( "Icons/icon_Blueprint_CommentBubbleOn_16x", FVector2D(10,10), FLinearColor(1.f, 1.f, 1.f, 1.f)))
-			.SetCheckedImage( IMAGE_BRUSH( "Icons/icon_Blueprint_CommentBubbleOn_16x", FVector2D(10,10), FLinearColor(1.f, 1.f, 1.f, 0.8f)))
-			.SetCheckedHoveredImage( IMAGE_BRUSH( "Icons/icon_Blueprint_CommentBubbleOn_16x", FVector2D(10,10), FLinearColor(1.f, 1.f, 1.f, 1.f)))
-			.SetCheckedPressedImage( IMAGE_BRUSH( "Icons/icon_Blueprint_CommentBubbleOn_16x", FVector2D(10,10), FLinearColor(1.f, 1.f, 1.f, 0.6f)));
+			.SetUncheckedImage( IMAGE_BRUSH( "Icons/icon_Blueprint_CommentBubbleOn_16x", Icon10x10, FLinearColor(1.f, 1.f, 1.f, 0.5f)))
+			.SetUncheckedHoveredImage( IMAGE_BRUSH( "Icons/icon_Blueprint_CommentBubbleOn_16x", Icon10x10, FLinearColor(1.f, 1.f, 1.f, 0.9f)))
+			.SetUncheckedPressedImage( IMAGE_BRUSH( "Icons/icon_Blueprint_CommentBubbleOn_16x", Icon10x10, FLinearColor(1.f, 1.f, 1.f, 1.f)))
+			.SetCheckedImage( IMAGE_BRUSH( "Icons/icon_Blueprint_CommentBubbleOn_16x", Icon10x10, FLinearColor(1.f, 1.f, 1.f, 0.8f)))
+			.SetCheckedHoveredImage( IMAGE_BRUSH( "Icons/icon_Blueprint_CommentBubbleOn_16x", Icon10x10, FLinearColor(1.f, 1.f, 1.f, 1.f)))
+			.SetCheckedPressedImage( IMAGE_BRUSH( "Icons/icon_Blueprint_CommentBubbleOn_16x", Icon10x10, FLinearColor(1.f, 1.f, 1.f, 0.6f)));
 		Set( "CommentBubbleButton", CommentBubbleButton );
 
 		const FCheckBoxStyle CommentBubblePin = FCheckBoxStyle()
 			.SetCheckBoxType( ESlateCheckBoxType::CheckBox )
-			.SetUncheckedImage( IMAGE_BRUSH( "Icons/icon_Blueprint_CommentBubbleUnPin_16x", FVector2D(10,10), FLinearColor(1.f, 1.f, 1.f, 0.5f)))
-			.SetUncheckedHoveredImage( IMAGE_BRUSH( "Icons/icon_Blueprint_CommentBubbleUnPin_16x", FVector2D(10,10), FLinearColor(1.f, 1.f, 1.f, 0.9f)))
-			.SetUncheckedPressedImage( IMAGE_BRUSH( "Icons/icon_Blueprint_CommentBubblePin_16x", FVector2D(10,10), FLinearColor(1.f, 1.f, 1.f, 1.f)))
-			.SetCheckedImage( IMAGE_BRUSH( "Icons/icon_Blueprint_CommentBubblePin_16x", FVector2D(10,10), FLinearColor(1.f, 1.f, 1.f, 0.8f)))
-			.SetCheckedHoveredImage( IMAGE_BRUSH( "Icons/icon_Blueprint_CommentBubblePin_16x", FVector2D(10,10), FLinearColor(1.f, 1.f, 1.f, 1.f)))
-			.SetCheckedPressedImage( IMAGE_BRUSH( "Icons/icon_Blueprint_CommentBubbleUnPin_16x", FVector2D(10,10), FLinearColor(1.f, 1.f, 1.f, 0.6f)));
+			.SetUncheckedImage( IMAGE_BRUSH( "Icons/icon_Blueprint_CommentBubbleUnPin_16x", Icon10x10, FLinearColor(1.f, 1.f, 1.f, 0.5f)))
+			.SetUncheckedHoveredImage( IMAGE_BRUSH( "Icons/icon_Blueprint_CommentBubbleUnPin_16x", Icon10x10, FLinearColor(1.f, 1.f, 1.f, 0.9f)))
+			.SetUncheckedPressedImage( IMAGE_BRUSH( "Icons/icon_Blueprint_CommentBubblePin_16x", Icon10x10, FLinearColor(1.f, 1.f, 1.f, 1.f)))
+			.SetCheckedImage( IMAGE_BRUSH( "Icons/icon_Blueprint_CommentBubblePin_16x", Icon10x10, FLinearColor(1.f, 1.f, 1.f, 0.8f)))
+			.SetCheckedHoveredImage( IMAGE_BRUSH( "Icons/icon_Blueprint_CommentBubblePin_16x", Icon10x10, FLinearColor(1.f, 1.f, 1.f, 1.f)))
+			.SetCheckedPressedImage( IMAGE_BRUSH( "Icons/icon_Blueprint_CommentBubbleUnPin_16x", Icon10x10, FLinearColor(1.f, 1.f, 1.f, 0.6f)));
 		Set( "CommentBubblePin", CommentBubblePin );
 
 
 		Set( "Graph.VarNode.Body", new BOX_BRUSH( "/Graph/VarNode_body", FMargin(16.f/64.f, 12.f/28.f) ) );
-		Set( "Graph.VarNode.ColorSpill", new IMAGE_BRUSH( "/Graph/VarNode_color_spill", FVector2D(132,28) ) );
+		Set( "Graph.VarNode.ColorSpill", new IMAGE_BRUSH( "/Graph/VarNode_color_spill", FVector2f(132.f,28.f) ) );
 		Set( "Graph.VarNode.Gloss", new BOX_BRUSH( "/Graph/VarNode_gloss", FMargin(16.f/64.f, 16.f/28.f, 16.f/64.f, 4.f/28.f) ) );
-		Set( "Graph.VarNode.IndicatorOverlay", new IMAGE_BRUSH("/Graph/IndicatorOverlay_color_spill", FVector2D(64.f, 28.f)));
+		Set( "Graph.VarNode.IndicatorOverlay", new IMAGE_BRUSH("/Graph/IndicatorOverlay_color_spill", FVector2f(64.f, 28.f)));
 		
 		Set( "Graph.VarNode.ShadowSelected", new BOX_BRUSH( "/Graph/VarNode_shadow_selected", FMargin(26.0f/64.0f) ) );
 		Set( "Graph.VarNode.Shadow", new BOX_BRUSH( "/Graph/VarNode_shadow", FMargin(26.0f/64.0f) ) );
@@ -3995,7 +3995,7 @@ void FSlateEditorStyle::FStyle::SetupGraphEditorStyles()
 				FTextBlockStyle GraphStateNodeTitle = FTextBlockStyle(NormalText)
 					.SetFont( DEFAULT_FONT( "Bold", 14 ) )
 					.SetColorAndOpacity( FLinearColor(230.0f/255.0f,230.0f/255.0f,230.0f/255.0f) )
-					.SetShadowOffset( FVector2D( 2,2 ) )
+					.SetShadowOffset( FVector2f( 2.f,2.f) )
 					.SetShadowColorAndOpacity( FLinearColor(0.f,0.f,0.f, 0.7f) );
 				Set( "Graph.StateNode.NodeTitle", GraphStateNodeTitle );
 
@@ -4019,8 +4019,8 @@ void FSlateEditorStyle::FStyle::SetupGraphEditorStyles()
 			{
 				FMargin TestMargin(16.f/64.f, 16.f/28.f, 16.f/64.f, 4.f/28.f);
 				Set( "Graph.TransitionNode.ColorSpill", new BOX_BRUSH( "/Persona/StateMachineEditor/Trans_Node_ColorSpill", TestMargin ) );
-				Set( "Graph.TransitionNode.Icon", new IMAGE_BRUSH( "/Persona/StateMachineEditor/Trans_Node_Icon", FVector2D(25,25) ) );
-				Set( "Graph.TransitionNode.Icon_Inertialization", new IMAGE_BRUSH( "/Persona/StateMachineEditor/Trans_Node_Icon_Inertialization", FVector2D(25,25) ) );
+				Set( "Graph.TransitionNode.Icon", new IMAGE_BRUSH( "/Persona/StateMachineEditor/Trans_Node_Icon", FVector2f(25.f,25.f) ) );
+				Set( "Graph.TransitionNode.Icon_Inertialization", new IMAGE_BRUSH( "/Persona/StateMachineEditor/Trans_Node_Icon_Inertialization", FVector2f(25.f,25.f) ) );
 			}
 
 			// Transition rule tooltip name 
@@ -4028,7 +4028,7 @@ void FSlateEditorStyle::FStyle::SetupGraphEditorStyles()
 				Set( "Graph.TransitionNode.TooltipName", FTextBlockStyle(NormalText)
 					.SetFont( DEFAULT_FONT( "Bold", 12 ) )
 					.SetColorAndOpacity( FLinearColor(218.0f/255.0f,218.0f/255.0f,218.0f/255.0f) )
-					.SetShadowOffset( FVector2D(1.0f, 1.0f) )
+					.SetShadowOffset( FVector2f::UnitVector )
 					.SetShadowColorAndOpacity( FLinearColor(0.f,0.f,0.f, 0.7f) )
 				);
 			}
@@ -4038,7 +4038,7 @@ void FSlateEditorStyle::FStyle::SetupGraphEditorStyles()
 				Set( "Graph.TransitionNode.TooltipRule", FTextBlockStyle(NormalText)
 					.SetFont( DEFAULT_FONT( "Bold", 8 ) )
 					.SetColorAndOpacity( FLinearColor(180.0f/255.0f,180.0f/255.0f,180.0f/255.0f) )
-					.SetShadowOffset( FVector2D(1.0f, 1.0f) )
+					.SetShadowOffset( FVector2f::UnitVector )
 					.SetShadowColorAndOpacity( FLinearColor(0.f,0.f,0.f, 0.7f) )
 				);
 			}
@@ -4051,8 +4051,8 @@ void FSlateEditorStyle::FStyle::SetupGraphEditorStyles()
 			Set("Persona.RetargetManager.ImportantText", FTextBlockStyle(NormalText)
 				.SetFont(DEFAULT_FONT("Bold", 11))
 				.SetColorAndOpacity(FLinearColor(1.0f, 1.0f, 1.0f, 1.0f))
-				.SetShadowOffset(FVector2D(1, 1))
-				.SetShadowColorAndOpacity(FLinearColor(0, 0, 0, 0.9f)));
+				.SetShadowOffset(FVector2f::UnitVector)
+				.SetShadowColorAndOpacity(FLinearColor(0.f, 0.f, 0.f, 0.9f)));
 		}
 
 		// Behavior Tree Editor
@@ -4131,10 +4131,10 @@ void FSlateEditorStyle::FStyle::SetupGraphEditorStyles()
 
 			Set( "BTEditor.DebuggerOverlay.Breakpoint.Disabled", new IMAGE_BRUSH( "Old/Kismet2/Breakpoint_Disabled", Icon32x32 ) );
 			Set( "BTEditor.DebuggerOverlay.Breakpoint.Enabled", new IMAGE_BRUSH( "Old/Kismet2/Breakpoint_Valid", Icon32x32 ) );
-			Set( "BTEditor.DebuggerOverlay.ActiveNodePointer", new IMAGE_BRUSH( "Old/Kismet2/IP_Normal", FVector2D(128,96)) );
-			Set( "BTEditor.DebuggerOverlay.SearchTriggerPointer", new IMAGE_BRUSH( "/BehaviorTree/SearchTriggerPointer", FVector2D(48,64)) );
-			Set( "BTEditor.DebuggerOverlay.FailedTriggerPointer", new IMAGE_BRUSH( "/BehaviorTree/FailedTriggerPointer", FVector2D(48,64)) );
-			Set( "BTEditor.DebuggerOverlay.BreakOnBreakpointPointer", new IMAGE_BRUSH( "Old/Kismet2/IP_Breakpoint", FVector2D(128,96)) );
+			Set( "BTEditor.DebuggerOverlay.ActiveNodePointer", new IMAGE_BRUSH( "Old/Kismet2/IP_Normal", FVector2f(128.f,96.f)) );
+			Set( "BTEditor.DebuggerOverlay.SearchTriggerPointer", new IMAGE_BRUSH( "/BehaviorTree/SearchTriggerPointer", FVector2f(48.f,64.f)) );
+			Set( "BTEditor.DebuggerOverlay.FailedTriggerPointer", new IMAGE_BRUSH( "/BehaviorTree/FailedTriggerPointer", FVector2f(48.f,64.f)) );
+			Set( "BTEditor.DebuggerOverlay.BreakOnBreakpointPointer", new IMAGE_BRUSH( "Old/Kismet2/IP_Breakpoint", FVector2f(128.f,96.f)) );
 
 			Set( "BTEditor.Blackboard.NewEntry", new IMAGE_BRUSH( "BehaviorTree/Blackboard_AddKey_40x", Icon40x40 ) );
 			Set( "BTEditor.Blackboard.NewEntry.Small", new IMAGE_BRUSH( "BehaviorTree/Blackboard_AddKey_40x", Icon20x20 ) );
@@ -4176,38 +4176,38 @@ void FSlateEditorStyle::FStyle::SetupGraphEditorStyles()
 		// Visible on hover button for transition node
 		{
 			Set( "TransitionNodeButton.Normal", new FSlateNoResource() );
-			Set( "TransitionNodeButton.Hovered", new IMAGE_BRUSH( "/Persona/StateMachineEditor/Trans_Button_Hovered", FVector2D(12,25) ) );
-			Set( "TransitionNodeButton.Pressed", new IMAGE_BRUSH( "/Persona/StateMachineEditor/Trans_Button_Pressed", FVector2D(12,25) ) );
+			Set( "TransitionNodeButton.Hovered", new IMAGE_BRUSH( "/Persona/StateMachineEditor/Trans_Button_Hovered", FVector2f(12.f,25.f) ) );
+			Set( "TransitionNodeButton.Pressed", new IMAGE_BRUSH( "/Persona/StateMachineEditor/Trans_Button_Pressed", FVector2f(12.f,25.f) ) );
 		}
 
 		{
-			Set( "Graph.AnimationResultNode.Body", new IMAGE_BRUSH( "/Graph/Animation/AnimationNode_Result_128x", FVector2D(128, 128) ) );
+			Set( "Graph.AnimationResultNode.Body", new IMAGE_BRUSH( "/Graph/Animation/AnimationNode_Result_128x", Icon128x128 ) );
 			Set( "Graph.AnimationFastPathIndicator", new IMAGE_BRUSH( "/Graph/Animation/AnimationNode_FastPath", Icon32x32 ) );
 		}
 
 		// SoundCueEditor Graph Nodes
 		{
-			Set( "Graph.SoundResultNode.Body", new IMAGE_BRUSH( "/Graph/SoundCue_SpeakerIcon", FVector2D(144, 144) ) );
+			Set( "Graph.SoundResultNode.Body", new IMAGE_BRUSH( "/Graph/SoundCue_SpeakerIcon", FVector2f(144.f, 144.f) ) );
 		}
 
-		Set( "Graph.Node.NodeEntryTop", new IMAGE_BRUSH( "Old/Graph/NodeEntryTop", FVector2D(64,12) ) );
-		Set( "Graph.Node.NodeEntryBottom", new IMAGE_BRUSH( "Old/Graph/NodeEntryBottom", FVector2D(64,12) ) );
-		Set( "Graph.Node.NodeExitTop", new IMAGE_BRUSH( "Old/Graph/NodeExitTop", FVector2D(64,12) ) );
-		Set( "Graph.Node.NodeExitBottom", new IMAGE_BRUSH( "Old/Graph/NodeExitBottom", FVector2D(64,12) ) );
+		Set( "Graph.Node.NodeEntryTop", new IMAGE_BRUSH( "Old/Graph/NodeEntryTop", FVector2f(64.f,12.f) ) );
+		Set( "Graph.Node.NodeEntryBottom", new IMAGE_BRUSH( "Old/Graph/NodeEntryBottom", FVector2f(64.f,12.f) ) );
+		Set( "Graph.Node.NodeExitTop", new IMAGE_BRUSH( "Old/Graph/NodeExitTop", FVector2f(64.f,12.f) ) );
+		Set( "Graph.Node.NodeExitBottom", new IMAGE_BRUSH( "Old/Graph/NodeExitBottom", FVector2f(64.f,12.f) ) );
 
 		Set( "Graph.Node.NodeEntryShadow", new BOX_BRUSH( "Old/Graph/NodeEntryShadow", FMargin(5.f/80, 21.f/52) ) );
 		Set( "Graph.Node.NodeEntryShadowSelected", new BOX_BRUSH( "Old/Graph/NodeEntryShadowSelected", FMargin(5.f/80, 21.f/52) ) );
 		Set( "Graph.Node.NodeExitShadow", new BOX_BRUSH( "Old/Graph/NodeExitShadow", FMargin(5.f/80, 21.f/52) ) );
 		Set( "Graph.Node.NodeExitShadowSelected", new BOX_BRUSH( "Old/Graph/NodeExitShadowSelected", FMargin(5.f/80, 21.f/52) ) );
 
-		Set( "Graph.Node.Autoplay", new IMAGE_BRUSH( "Graph/Icons/Overlay_Autoplay", FVector2D(22,22) ) );
-		Set( "Graph.Node.Loop", new IMAGE_BRUSH( "Graph/Icons/Overlay_Loop", FVector2D(22,22) ) );
+		Set( "Graph.Node.Autoplay", new IMAGE_BRUSH( "Graph/Icons/Overlay_Autoplay", FVector2f(22.f,22.f) ) );
+		Set( "Graph.Node.Loop", new IMAGE_BRUSH( "Graph/Icons/Overlay_Loop", FVector2f(22.f,22.f) ) );
 
 		{
 			FTextBlockStyle GraphNodeTitle = FTextBlockStyle(NormalText)
 				.SetFont( DEFAULT_FONT( "Bold", 10 ) )
 				.SetColorAndOpacity( FLinearColor(218.0f/255.0f,218.0f/255.0f,218.0f/255.0f) )
-				.SetShadowOffset( FVector2D::ZeroVector )
+				.SetShadowOffset( FVector2f::ZeroVector )
 				.SetShadowColorAndOpacity( FLinearColor(0.f,0.f,0.f, 0.7f) );
 			Set( "Graph.Node.NodeTitle", GraphNodeTitle );
 
@@ -4229,14 +4229,14 @@ void FSlateEditorStyle::FStyle::SetupGraphEditorStyles()
 			Set( "Graph.Node.NodeTitleExtraLines", FTextBlockStyle(NormalText)
 				.SetFont( DEFAULT_FONT( "Italic", 9 ) )
 				.SetColorAndOpacity( FLinearColor(218.0f/255.0f,218.0f/255.0f,96.0f/255.0f, 0.5f) )
-				.SetShadowOffset( FVector2D::ZeroVector )
+				.SetShadowOffset( FVector2f::ZeroVector )
 				.SetShadowColorAndOpacity( FLinearColor(0.f,0.f,0.f, 0.7f) )
 			);
 		
 			FTextBlockStyle GraphCommentBlockTitle = FTextBlockStyle(NormalText)
 				.SetFont( DEFAULT_FONT( "Bold", 18 ) )
 				.SetColorAndOpacity( FLinearColor(218.0f/255.0f,218.0f/255.0f,218.0f/255.0f) )
-				.SetShadowOffset( FVector2D(1.5f, 1.5f) )
+				.SetShadowOffset( FVector2f(1.5f, 1.5f) )
 				.SetShadowColorAndOpacity( FLinearColor(0.f,0.f,0.f, 0.7f) );
 			Set( "Graph.CommentBlock.Title", GraphCommentBlockTitle );
 
@@ -4258,21 +4258,21 @@ void FSlateEditorStyle::FStyle::SetupGraphEditorStyles()
 			Set( "Graph.CompactNode.Title", FTextBlockStyle(NormalText)
 				.SetFont( DEFAULT_FONT( "BoldCondensed", 20 ) )
 				.SetColorAndOpacity( FLinearColor(1.0f, 1.0f, 1.0f, 0.5f) )
-				.SetShadowOffset( FVector2D::ZeroVector )
+				.SetShadowOffset( FVector2f::ZeroVector )
 				.SetShadowColorAndOpacity( FLinearColor::White )
 			);
 
 			Set( "Graph.ArrayCompactNode.Title", FTextBlockStyle(NormalText)
 				.SetFont( DEFAULT_FONT( "BoldCondensed", 20 ) )
 				.SetColorAndOpacity( FLinearColor(1.0f, 1.0f, 1.0f, 0.5f) ) //218.0f/255.0f, 218.0f/255.0f, 218.0f/255.0f, 0.25f) )
-				.SetShadowOffset( FVector2D::ZeroVector )
+				.SetShadowOffset( FVector2f::ZeroVector )
 				.SetShadowColorAndOpacity( FLinearColor::White )
 				);
 
 			Set( "Graph.Node.PinName", FTextBlockStyle(NormalText)
 				.SetFont( DEFAULT_FONT( "Regular", 9 ) )
 				.SetColorAndOpacity( FLinearColor(218.0f/255.0f,218.0f/255.0f,218.0f/255.0f) )
-				.SetShadowOffset( FVector2D::ZeroVector )
+				.SetShadowOffset( FVector2f::ZeroVector )
 				.SetShadowColorAndOpacity( FLinearColor(0.8f,0.8f,0.8f, 0.5) )
 			);
 
@@ -4281,7 +4281,7 @@ void FSlateEditorStyle::FStyle::SetupGraphEditorStyles()
 				FTextBlockStyle InlineEditableTextBlockReadOnly = FTextBlockStyle(NormalText)
 					.SetFont(DEFAULT_FONT("Regular", 9))
 					.SetColorAndOpacity(FLinearColor(218.0f / 255.0f, 218.0f / 255.0f, 218.0f / 255.0f))
-					.SetShadowOffset(FVector2D::ZeroVector)
+					.SetShadowOffset(FVector2f::ZeroVector)
 					.SetShadowColorAndOpacity(FLinearColor(0.8f, 0.8f, 0.8f, 0.5));
 
 				FEditableTextBoxStyle InlineEditableTextBlockEditable = FEditableTextBoxStyle()
@@ -4323,50 +4323,50 @@ void FSlateEditorStyle::FStyle::SetupGraphEditorStyles()
 			Set( "Graph.Node.DefaultPinName", FTextBlockStyle(NormalText)
 				.SetFont( DEFAULT_FONT( "Italic", 9 ) )
 				.SetColorAndOpacity( FLinearColor(218.0f/255.0f,218.0f/255.0f,218.0f/255.0f) )
-				.SetShadowOffset( FVector2D::ZeroVector )
+				.SetShadowOffset( FVector2f::ZeroVector )
 				.SetShadowColorAndOpacity( FLinearColor(0.8f,0.8f,0.8f, 0.5) )
 			);
 		}
-		Set( "Graph.Pin.DefaultPinSeparator", new IMAGE_BRUSH( "/Graph/PinSeparator", FVector2D(64,8) ) );
+		Set( "Graph.Pin.DefaultPinSeparator", new IMAGE_BRUSH( "/Graph/PinSeparator", FVector2f(64.f,8.f) ) );
 
 		/** Original Pin Styles */
-		Set( "Graph.Pin.Connected", new IMAGE_BRUSH( "/Graph/Pin_connected", FVector2D(11,11) ) );
-		Set( "Graph.Pin.Disconnected", new IMAGE_BRUSH( "/Graph/Pin_disconnected", FVector2D(11,11) ) );
-		Set( "Graph.ArrayPin.Connected", new IMAGE_BRUSH( "/Graph/ArrayPin_connected", FVector2D(11,11) ) );
-		Set( "Graph.ArrayPin.Disconnected", new IMAGE_BRUSH( "/Graph/ArrayPin_disconnected", FVector2D(11,11) ) );
-		Set( "Graph.RefPin.Connected", new IMAGE_BRUSH( "/Graph/RefPin_connected", FVector2D(11,11) ) );
-		Set( "Graph.RefPin.Disconnected", new IMAGE_BRUSH( "/Graph/RefPin_disconnected", FVector2D(11,11) ) );
+		Set( "Graph.Pin.Connected", new IMAGE_BRUSH( "/Graph/Pin_connected", FVector2f(11.f,11.f) ) );
+		Set( "Graph.Pin.Disconnected", new IMAGE_BRUSH( "/Graph/Pin_disconnected", FVector2f(11.f,11.f) ) );
+		Set( "Graph.ArrayPin.Connected", new IMAGE_BRUSH( "/Graph/ArrayPin_connected", FVector2f(11.f,11.f) ) );
+		Set( "Graph.ArrayPin.Disconnected", new IMAGE_BRUSH( "/Graph/ArrayPin_disconnected", FVector2f(11.f,11.f) ) );
+		Set( "Graph.RefPin.Connected", new IMAGE_BRUSH( "/Graph/RefPin_connected", FVector2f(11.f,11.f) ) );
+		Set( "Graph.RefPin.Disconnected", new IMAGE_BRUSH( "/Graph/RefPin_disconnected", FVector2f(11.f,11.f) ) );
 
-		Set("Graph.Pin.CopyNodePinLeft_Connected", new IMAGE_BRUSH("/Graph/CopyNodePinLeft_connected", FVector2D(12, 24)));
-		Set("Graph.Pin.CopyNodePinLeft_Disconnected", new IMAGE_BRUSH("/Graph/CopyNodePinLeft_disconnected", FVector2D(12, 24)));
+		Set("Graph.Pin.CopyNodePinLeft_Connected", new IMAGE_BRUSH("/Graph/CopyNodePinLeft_connected", FVector2f(12.f, 24.f)));
+		Set("Graph.Pin.CopyNodePinLeft_Disconnected", new IMAGE_BRUSH("/Graph/CopyNodePinLeft_disconnected", FVector2f(12.f, 24.f)));
 
-		Set("Graph.Pin.CopyNodePinRight_Connected", new IMAGE_BRUSH("/Graph/CopyNodePinRight_connected", FVector2D(12, 24)));
-		Set("Graph.Pin.CopyNodePinRight_Disconnected", new IMAGE_BRUSH("/Graph/CopyNodePinRight_disconnected", FVector2D(12, 24)));
+		Set("Graph.Pin.CopyNodePinRight_Connected", new IMAGE_BRUSH("/Graph/CopyNodePinRight_connected", FVector2f(12.f, 24.f)));
+		Set("Graph.Pin.CopyNodePinRight_Disconnected", new IMAGE_BRUSH("/Graph/CopyNodePinRight_disconnected", FVector2f(12.f, 24.f)));
 
 		/** Variant A Pin Styles */
-		Set( "Graph.Pin.Connected_VarA", new IMAGE_BRUSH( "/Graph/Pin_connected_VarA", FVector2D(15,11)) );
-		Set( "Graph.Pin.Disconnected_VarA", new IMAGE_BRUSH( "/Graph/Pin_disconnected_VarA", FVector2D(15,11)) );
+		Set( "Graph.Pin.Connected_VarA", new IMAGE_BRUSH( "/Graph/Pin_connected_VarA", FVector2f(15.f,11.f)) );
+		Set( "Graph.Pin.Disconnected_VarA", new IMAGE_BRUSH( "/Graph/Pin_disconnected_VarA", FVector2f(15.f,11.f)) );
 
-		Set( "Graph.DelegatePin.Connected", new IMAGE_BRUSH( "/Graph/DelegatePin_Connected", FVector2D(11,11) ) );
-		Set( "Graph.DelegatePin.Disconnected", new IMAGE_BRUSH( "/Graph/DelegatePin_Disconnected", FVector2D(11,11) ) );
+		Set( "Graph.DelegatePin.Connected", new IMAGE_BRUSH( "/Graph/DelegatePin_Connected", FVector2f(11.f,11.f) ) );
+		Set( "Graph.DelegatePin.Disconnected", new IMAGE_BRUSH( "/Graph/DelegatePin_Disconnected", FVector2f(11.f,11.f) ) );
 
-		Set( "Graph.Replication.AuthorityOnly", new IMAGE_BRUSH( "/Graph/AuthorityOnly", FVector2D(32,32) ) );
-		Set( "Graph.Replication.ClientEvent", new IMAGE_BRUSH( "/Graph/ClientEvent", FVector2D(32,32) ) );
-		Set( "Graph.Replication.Replicated", new IMAGE_BRUSH( "/Graph/Replicated", FVector2D(32,32) ) );
+		Set( "Graph.Replication.AuthorityOnly", new IMAGE_BRUSH( "/Graph/AuthorityOnly", Icon32x32 ) );
+		Set( "Graph.Replication.ClientEvent", new IMAGE_BRUSH( "/Graph/ClientEvent", Icon32x32 ) );
+		Set( "Graph.Replication.Replicated", new IMAGE_BRUSH( "/Graph/Replicated", Icon32x32 ) );
 
-		Set("Graph.Editor.EditorOnlyIcon", new IMAGE_BRUSH("/Graph/EditorOnly", FVector2D(32, 32)));
+		Set("Graph.Editor.EditorOnlyIcon", new IMAGE_BRUSH("/Graph/EditorOnly", Icon32x32));
 
-		Set( "Graph.Event.InterfaceEventIcon", new IMAGE_BRUSH("/Graph/InterfaceEventIcon", FVector2D(32,32) ) );
+		Set( "Graph.Event.InterfaceEventIcon", new IMAGE_BRUSH("/Graph/InterfaceEventIcon", Icon32x32 ) );
 
-		Set( "Graph.Latent.LatentIcon", new IMAGE_BRUSH("/Graph/LatentIcon", FVector2D(32,32) ) );
-		Set( "Graph.Message.MessageIcon", new IMAGE_BRUSH("/Graph/MessageIcon", FVector2D(32,32) ) );
+		Set( "Graph.Latent.LatentIcon", new IMAGE_BRUSH("/Graph/LatentIcon", Icon32x32 ) );
+		Set( "Graph.Message.MessageIcon", new IMAGE_BRUSH("/Graph/MessageIcon", Icon32x32 ) );
 
 		Set( "Graph.ExecPin.Connected", new IMAGE_BRUSH( "Old/Graph/ExecPin_Connected", Icon12x16 ) );
 		Set( "Graph.ExecPin.Disconnected", new IMAGE_BRUSH( "Old/Graph/ExecPin_Disconnected", Icon12x16 ) );
 		Set( "Graph.ExecPin.ConnectedHovered", new IMAGE_BRUSH( "Old/Graph/ExecPin_Connected", Icon12x16, FLinearColor(0.8f,0.8f,0.8f) ) );
 		Set( "Graph.ExecPin.DisconnectedHovered", new IMAGE_BRUSH( "Old/Graph/ExecPin_Disconnected", Icon12x16, FLinearColor(0.8f,0.8f,0.8f) ) );
 
-		const FVector2D Icon15x28(15.0f, 28.0f);
+		const FVector2f Icon15x28(15.0f, 28.0f);
 		Set("Graph.PosePin.Connected", new IMAGE_BRUSH("Graph/Animation/PosePin_Connected_15x28", Icon15x28));
 		Set("Graph.PosePin.Disconnected", new IMAGE_BRUSH("Graph/Animation/PosePin_Disconnected_15x28", Icon15x28));
 		Set("Graph.PosePin.ConnectedHovered", new IMAGE_BRUSH("Graph/Animation/PosePin_Connected_15x28", Icon15x28, FLinearColor(0.8f, 0.8f, 0.8f)));
@@ -4380,7 +4380,7 @@ void FSlateEditorStyle::FStyle::SetupGraphEditorStyles()
 
 		Set( "Graph.WatchedPinIcon_Pinned", new IMAGE_BRUSH( "Old/Graph/WatchedPinIcon_Pinned", Icon16x16 ) );
 
-		Set( "Graph.Pin.BackgroundHovered", new IMAGE_BRUSH( "/Graph/Pin_hover_cue", FVector2D(32,8)));
+		Set( "Graph.Pin.BackgroundHovered", new IMAGE_BRUSH( "/Graph/Pin_hover_cue", FVector2f(32.f,8.f)));
 		Set( "Graph.Pin.Background", new FSlateNoResource() );
 
 		Set( "Graph.Pin.ObjectSet", new IMAGE_BRUSH( "Old/Graph/Pin_ObjectSet", Icon12x12 ) );
@@ -4397,32 +4397,32 @@ void FSlateEditorStyle::FStyle::SetupGraphEditorStyles()
 		{
 			Set( "Graph.CornerText", FTextBlockStyle(NormalText)
 				.SetFont( DEFAULT_FONT( "BoldCondensed", 48 ) )
-				.SetColorAndOpacity( FLinearColor(0.8, 0.8f, 0.8f, 0.2f) )
-				.SetShadowOffset( FVector2D::ZeroVector )
+				.SetColorAndOpacity( FLinearColor(0.8f, 0.8f, 0.8f, 0.2f) )
+				.SetShadowOffset( FVector2f::ZeroVector )
 			);
 
 			Set("Graph.WarningText", FTextBlockStyle(NormalText)
 				.SetFont(DEFAULT_FONT("BoldCondensed", 20))
-				.SetColorAndOpacity(FLinearColor(0.8, 0.5f, 0.07f, 1.0f))
-				.SetShadowOffset(FVector2D(1, 1))
+				.SetColorAndOpacity(FLinearColor(0.8f, 0.5f, 0.07f, 1.0f))
+				.SetShadowOffset(FVector2f::UnitVector)
 			);
 
 			Set( "Graph.SimulatingText", FTextBlockStyle(NormalText)
 				.SetFont( DEFAULT_FONT( "BoldCondensed", 48 ) )
-				.SetColorAndOpacity( FLinearColor(0.8, 0.8f, 0.0f, 0.2f) )
-				.SetShadowOffset( FVector2D::ZeroVector )
+				.SetColorAndOpacity( FLinearColor(0.8f, 0.8f, 0.0f, 0.2f) )
+				.SetShadowOffset( FVector2f::ZeroVector )
 			);
 
 			Set( "GraphPreview.CornerText", FTextBlockStyle(NormalText)
 				.SetFont( DEFAULT_FONT( "BoldCondensed", 16 ) )
-				.SetColorAndOpacity( FLinearColor(0.8, 0.8f, 0.8f, 0.2f) )
-				.SetShadowOffset( FVector2D::ZeroVector )
+				.SetColorAndOpacity( FLinearColor(0.8f, 0.8f, 0.8f, 0.2f) )
+				.SetShadowOffset( FVector2f::ZeroVector )
 			);
 
 			Set( "Graph.InstructionText", FTextBlockStyle(NormalText)
 				.SetFont( DEFAULT_FONT( "BoldCondensed", 24 ) )
 				.SetColorAndOpacity( FLinearColor(1.f, 1.f, 1.f, 0.6f) )
-				.SetShadowOffset( FVector2D::ZeroVector )
+				.SetShadowOffset( FVector2f::ZeroVector )
 			);
 
 			Set( "Graph.InstructionBackground", new BOX_BRUSH("Common/GroupBorder", FMargin(4.0f / 16.0f), FLinearColor(0.1f, 0.1f, 0.1f, 0.7f)) );
@@ -4626,7 +4626,7 @@ void FSlateEditorStyle::FStyle::SetupGraphEditorStyles()
 			.SetFont(DEFAULT_FONT("Bold", 10))
 			.SetColorAndOpacity(FLinearColor(1.0f, 1.0f, 1.0f))
 			.SetHighlightColor(FLinearColor(1.0f, 1.0f, 1.0f))
-			.SetShadowOffset(FVector2D(1, 1))
+			.SetShadowOffset(FVector2f::UnitVector)
 			.SetShadowColorAndOpacity(FLinearColor(0, 0, 0, 0.9f)));
 
 		Set("SCSEditor.TreePanel", new BOX_BRUSH("Common/GroupBorder_FlatTop", FMargin(4.0f / 16.0f)));
@@ -4642,13 +4642,13 @@ void FSlateEditorStyle::FStyle::SetupGraphEditorStyles()
 		Set("SCSEditor.ComponentTooltip.Label",
 			FTextBlockStyle(NormalText)
 			.SetColorAndOpacity(FLinearColor(0.075f, 0.075f, 0.075f))
-			.SetShadowOffset(FVector2D(1.0f, 1.0f))
+			.SetShadowOffset(FVector2f::UnitVector)
 			.SetShadowColorAndOpacity(FLinearColor(0.8f, 0.8f, 0.8f))
 		);
 		Set("SCSEditor.ComponentTooltip.ImportantLabel",
 			FTextBlockStyle(NormalText)
 			.SetColorAndOpacity(FLinearColor(0.05f, 0.05f, 0.05f))
-			.SetShadowOffset(FVector2D(1.0f, 1.0f))
+			.SetShadowOffset(FVector2f::UnitVector)
 			.SetShadowColorAndOpacity(FLinearColor(0.8f, 0.8f, 0.8f))
 		);
 
@@ -4657,14 +4657,14 @@ void FSlateEditorStyle::FStyle::SetupGraphEditorStyles()
 			FTextBlockStyle(NormalText)
 			.SetFont(DEFAULT_FONT("Bold", 10))
 			.SetColorAndOpacity(FLinearColor(0.0f, 0.0f, 0.0f))
-			.SetShadowOffset(FVector2D(1.0f, 1.0f))
+			.SetShadowOffset(FVector2f::UnitVector)
 			.SetShadowColorAndOpacity(FLinearColor(0.8f, 0.8f, 0.8f))
 		);
 		Set("SCSEditor.ComponentTooltip.ImportantValue",
 			FTextBlockStyle(NormalText)
 			.SetFont(DEFAULT_FONT("Bold", 10))
 			.SetColorAndOpacity(FLinearColor(0.3f, 0.0f, 0.0f))
-			.SetShadowOffset(FVector2D(1.0f, 1.0f))
+			.SetShadowOffset(FVector2f::UnitVector)
 			.SetShadowColorAndOpacity(FLinearColor(0.8f, 0.8f, 0.8f))
 		);
 
@@ -4672,7 +4672,7 @@ void FSlateEditorStyle::FStyle::SetupGraphEditorStyles()
 			FTextBlockStyle(NormalText)
 			.SetFont(DEFAULT_FONT("Italic", 10))
 			.SetColorAndOpacity(FLinearColor(0.1f, 0.1f, 0.1f))
-			.SetShadowOffset(FVector2D(1.0f, 1.0f))
+			.SetShadowOffset(FVector2f::UnitVector)
 			.SetShadowColorAndOpacity(FLinearColor(0.8f, 0.8f, 0.8f))
 		);
 	}
@@ -4689,12 +4689,12 @@ void FSlateEditorStyle::FStyle::SetupGraphEditorStyles()
 		Set( "ModeSelector.ToggleButton.Hovered", new BOX_BRUSH( "Common/RoundedSelection_16x", 4.0f/16.0f, SelectionColor ) );
 
 
-		Set( "BlueprintEditor.PipelineSeparator", new BOX_BRUSH( "Old/Kismet2/BlueprintModeSeparator", FMargin(15.0f/16.0f, 20.0f/20.0f, 1.0f/16.0f, 0.0f/20.0f), FLinearColor(1,1,1,0.5f) ) );
+		Set( "BlueprintEditor.PipelineSeparator", new BOX_BRUSH( "Old/Kismet2/BlueprintModeSeparator", FMargin(15.0f/16.0f, 20.0f/20.0f, 1.0f/16.0f, 0.0f/20.0f), FLinearColor(1.f,1.f,1.f,0.5f) ) );
 	}
 
 	// Persona modes
 	{
-		Set( "Persona.PipelineSeparator", new BOX_BRUSH( "Persona/Modes/PipelineSeparator", FMargin(15.0f/16.0f, 22.0f/24.0f, 1.0f/16.0f, 1.0f/24.0f), FLinearColor(1,1,1,0.5f) ) );
+		Set( "Persona.PipelineSeparator", new BOX_BRUSH( "Persona/Modes/PipelineSeparator", FMargin(15.0f/16.0f, 22.0f/24.0f, 1.0f/16.0f, 1.0f/24.0f), FLinearColor(1.f,1.f,1.f,0.5f) ) );
 	}
 
 	// montage editor
@@ -4702,7 +4702,7 @@ void FSlateEditorStyle::FStyle::SetupGraphEditorStyles()
 		Set("Persona.MontageEditor.ChildMontageInstruction", FTextBlockStyle(NormalText)
 			.SetFont(DEFAULT_FONT("BoldCondensed", 14))
 			.SetColorAndOpacity(FLinearColor(1.f, 1.f, 1.f, 1.0f))
-			.SetShadowOffset(FVector2D::ZeroVector)
+			.SetShadowOffset(FVector2f::ZeroVector)
 		);
 	}
 #endif // WITH_EDITOR || (IS_PROGRAM && WITH_UNREAL_DEVELOPER_TOOLS)
@@ -4749,7 +4749,7 @@ void FSlateEditorStyle::FStyle::SetupLevelEditorStyle()
 		Set( "EditorViewport.RelativeCoordinateSystem_Local.Small", new IMAGE_BRUSH( "Icons/icon_axis_local_16px", Icon16x16, IconColor ) );
 		Set( "EditorViewport.RelativeCoordinateSystem_World", new IMAGE_BRUSH( "Icons/icon_axis_world_16px", Icon16x16, IconColor ) );
 		Set( "EditorViewport.RelativeCoordinateSystem_World.Small", new IMAGE_BRUSH( "Icons/icon_axis_world_16px", Icon16x16, IconColor ) );
-		Set( "EditorViewport.CamSpeedSetting", new IMAGE_BRUSH( "Icons/icon_CameraSpeed_24x16px", FVector2D( 24, 16 ), IconColor ) );
+		Set( "EditorViewport.CamSpeedSetting", new IMAGE_BRUSH( "Icons/icon_CameraSpeed_24x16px", FVector2f( 24.f, 16.f), IconColor ) );
 		
 		Set( "EditorViewport.LitMode", new IMAGE_BRUSH( "Icons/icon_ViewMode_Lit_16px", Icon16x16 ) );
 		Set( "EditorViewport.UnlitMode", new IMAGE_BRUSH( "Icons/icon_ViewMode_Unlit_16px", Icon16x16 ) );
@@ -4860,8 +4860,8 @@ void FSlateEditorStyle::FStyle::SetupLevelEditorStyle()
 
 		Set( "LevelEditor.Create", new IMAGE_BRUSH( "Icons/icon_Mode_Placement_40px", Icon40x40 ) );
 		Set( "LevelEditor.Create.Small", new IMAGE_BRUSH( "Icons/icon_Mode_Placement_40px", Icon20x20 ) );
-		Set( "LevelEditor.Create.OutlineHoriz", new IMAGE_BRUSH( "Common/WorkingFrame_Marquee", FVector2D(34.0f, 3.0f), FLinearColor::White, ESlateBrushTileType::Horizontal) );
-		Set( "LevelEditor.Create.OutlineVert", new IMAGE_BRUSH( "Common/WorkingFrame_Marquee_Vert", FVector2D(3.0f, 34.0f), FLinearColor::White, ESlateBrushTileType::Vertical) );
+		Set( "LevelEditor.Create.OutlineHoriz", new IMAGE_BRUSH( "Common/WorkingFrame_Marquee", FVector2f(34.0f, 3.0f), FLinearColor::White, ESlateBrushTileType::Horizontal) );
+		Set( "LevelEditor.Create.OutlineVert", new IMAGE_BRUSH( "Common/WorkingFrame_Marquee_Vert", FVector2f(3.0f, 34.0f), FLinearColor::White, ESlateBrushTileType::Vertical) );
 
 		Set( "LevelEditor.EditorModes", new IMAGE_BRUSH( "Icons/icon_Editor_Modes_40x", Icon40x40 ) );
 		Set( "LevelEditor.EditorModes.Small", new IMAGE_BRUSH( "Icons/icon_Editor_Modes_40x", Icon20x20 ) );
@@ -4971,23 +4971,23 @@ void FSlateEditorStyle::FStyle::SetupLevelEditorStyle()
 		Set( "PlacementBrowser.Tab.Text", FTextBlockStyle( NormalText )
 			.SetFont( DEFAULT_FONT( "Bold", 10 ) )
 			.SetColorAndOpacity( FLinearColor( 1.0f, 1.0f, 1.0f, 0.9f ) )
-			.SetShadowOffset( FVector2D( 1, 1 ) )
-			.SetShadowColorAndOpacity( FLinearColor( 0, 0, 0, 0.9f ) ) );
+			.SetShadowOffset( FVector2f::UnitVector )
+			.SetShadowColorAndOpacity( FLinearColor( 0.f, 0.f, 0.f, 0.9f ) ) );
 
 		Set( "PlacementBrowser.Asset.Name", FTextBlockStyle( NormalText )
 			.SetFont( DEFAULT_FONT( "Regular", 10 ) )
 			.SetColorAndOpacity( FLinearColor( 1.0f, 1.0f, 1.0f, 0.9f ) )
-			.SetShadowOffset( FVector2D( 1, 1 ) )
-			.SetShadowColorAndOpacity( FLinearColor( 0, 0, 0, 0.9f ) ) );
+			.SetShadowOffset( FVector2f::UnitVector )
+			.SetShadowColorAndOpacity( FLinearColor( 0.f, 0.f, 0.f, 0.9f ) ) );
 
 		Set( "PlacementBrowser.Asset.Type", FTextBlockStyle( NormalText )
 			.SetFont( DEFAULT_FONT( "Regular", 8 ) )
 			.SetColorAndOpacity( FLinearColor( 0.8f, 0.8f, 0.8f, 0.9f ) )
-			.SetShadowOffset( FVector2D( 1, 1 ) )
-			.SetShadowColorAndOpacity( FLinearColor( 0, 0, 0, 0.9f ) ) );
+			.SetShadowOffset( FVector2f::UnitVector )
+			.SetShadowColorAndOpacity( FLinearColor( 0.f, 0.f, 0.f, 0.9f ) ) );
 
 		Set( "PlacementBrowser.ActiveTabNub", new IMAGE_BRUSH( "Icons/TabTriangle_24x", Icon24x24, FLinearColor( FColor( 42, 42, 42 ) ) ) );
-		Set( "PlacementBrowser.ActiveTabBar", new IMAGE_BRUSH( "Common/Selection", FVector2D(2.0f, 2.0f), SelectionColor ) );
+		Set( "PlacementBrowser.ActiveTabBar", new IMAGE_BRUSH( "Common/Selection", FVector2f(2.0f, 2.0f), SelectionColor ) );
 
 		Set( "PlacementBrowser.ShowAllContent", new IMAGE_BRUSH( "Icons/icon_Placement_AllContent_20px", Icon20x20 ) );
 		Set( "PlacementBrowser.ShowAllContent.Small", new IMAGE_BRUSH( "Icons/icon_Placement_AllContent_20px", Icon20x20 ) );
@@ -5010,7 +5010,7 @@ void FSlateEditorStyle::FStyle::SetupLevelEditorStyle()
 
 	{
 
-		Set( "AssetDeleteDialog.Background", new IMAGE_BRUSH( "Common/Selection", Icon8x8, FLinearColor( 0.016, 0.016, 0.016 ) ) );
+		Set( "AssetDeleteDialog.Background", new IMAGE_BRUSH( "Common/Selection", Icon8x8, FLinearColor( 0.016f, 0.016f, 0.016f ) ) );
 	}
 
 	// Level editor tool box icons
@@ -5020,8 +5020,8 @@ void FSlateEditorStyle::FStyle::SetupLevelEditorStyle()
 
 	// Editor viewport layout command icons
 	{
-		const FVector2D IconLayoutSize(47.0f, 37.0f);
-		const FVector2D IconLayoutSizeSmall(47.0f, 37.0f);		// small version set to same size as these are in their own menu and don't clutter the UI
+		const FVector2f IconLayoutSize(47.0f, 37.0f);
+		const FVector2f IconLayoutSizeSmall(47.0f, 37.0f);		// small version set to same size as these are in their own menu and don't clutter the UI
 
 		Set("EditorViewport.ViewportConfig_OnePane", new IMAGE_BRUSH_SVG("Starship/ViewportLayout/Layouts1Pane", IconLayoutSize));
 		Set("EditorViewport.ViewportConfig_OnePane.Small", new IMAGE_BRUSH_SVG("Starship/ViewportLayout/Layouts1Pane", IconLayoutSizeSmall));
@@ -5051,8 +5051,8 @@ void FSlateEditorStyle::FStyle::SetupLevelEditorStyle()
 
 	// Level viewport layout command icons
 	{
-		const FVector2D IconLayoutSize(47.0f, 37.0f);
-		const FVector2D IconLayoutSizeSmall(47.0f, 37.0f);		// small version set to same size as these are in their own menu and don't clutter the UI
+		const FVector2f IconLayoutSize(47.0f, 37.0f);
+		const FVector2f IconLayoutSizeSmall(47.0f, 37.0f);		// small version set to same size as these are in their own menu and don't clutter the UI
 
 		Set("LevelViewport.ViewportConfig_OnePane", new IMAGE_BRUSH_SVG("Starship/ViewportLayout/Layouts1Pane", IconLayoutSize));
 		Set("LevelViewport.ViewportConfig_OnePane.Small", new IMAGE_BRUSH_SVG("Starship/ViewportLayout/Layouts1Pane", IconLayoutSizeSmall));
@@ -5082,21 +5082,21 @@ void FSlateEditorStyle::FStyle::SetupLevelEditorStyle()
 
 	// Level editor status bar
 	{
-		Set( "TransformSettings.RelativeCoordinateSettings", new IMAGE_BRUSH( "Icons/icon_axis_16px", FVector2D( 16, 16 ) ) );
+		Set( "TransformSettings.RelativeCoordinateSettings", new IMAGE_BRUSH( "Icons/icon_axis_16px", Icon16x16 ) );
 	}
 
 	// Mesh Proxy Window
 	{
-		Set("MeshProxy.SimplygonLogo", new IMAGE_BRUSH( "Icons/SimplygonBanner_Sml", FVector2D(174, 36) ) );
+		Set("MeshProxy.SimplygonLogo", new IMAGE_BRUSH( "Icons/SimplygonBanner_Sml", FVector2f(174.f, 36.f) ) );
 	}
 #endif // WITH_EDITOR
 
 	// Level viewport 
 #if WITH_EDITOR || (IS_PROGRAM && WITH_UNREAL_DEVELOPER_TOOLS)
 	{
-		Set( "LevelViewport.ActiveViewportBorder", new BORDER_BRUSH( "Old/White", FMargin(1), SelectionColor ) ); 
+		Set( "LevelViewport.ActiveViewportBorder", new BORDER_BRUSH( "Old/White", FMargin(1.f), SelectionColor ) );
 		Set( "LevelViewport.NoViewportBorder", new FSlateNoResource() );
-		Set( "LevelViewport.DebugBorder", new BOX_BRUSH( "Old/Window/ViewportDebugBorder", 0.8f, FLinearColor(.7,0,0,.5) ) );
+		Set( "LevelViewport.DebugBorder", new BOX_BRUSH( "Old/Window/ViewportDebugBorder", 0.8f, FLinearColor(.7f,0.f,0.f,.5f) ) );
 		Set( "LevelViewport.BlackBackground", new FSlateColorBrush( FLinearColor::Black ) ); 
 		Set( "LevelViewport.StartingPlayInEditorBorder", new BOX_BRUSH( "Old/Window/ViewportDebugBorder", 0.8f, FLinearColor(0.1f,1.0f,0.1f,1.0f) ) );
 		Set( "LevelViewport.StartingSimulateBorder", new BOX_BRUSH( "Old/Window/ViewportDebugBorder", 0.8f, FLinearColor(1.0f,1.0f,0.1f,1.0f) ) );
@@ -5115,10 +5115,10 @@ void FSlateEditorStyle::FStyle::SetupLevelEditorStyle()
 
 	// Level editor ui command icons
 	{
-		Set( "LevelEditor.ShowAll", new IMAGE_BRUSH( "Old/SelectionDetails/ShowAll", FVector2D(32,32) ) );
-		Set( "LevelEditor.ShowSelectedOnly", new IMAGE_BRUSH( "Old/SelectionDetails/ShowSelected", FVector2D(32,32) ) );
-		Set( "LevelEditor.ShowSelected", new IMAGE_BRUSH( "Old/SelectionDetails/ShowSelected", FVector2D(32,32) ) );
-		Set( "LevelEditor.HideSelected", new IMAGE_BRUSH( "Old/SelectionDetails/HideSelected", FVector2D(32,32) ) );
+		Set( "LevelEditor.ShowAll", new IMAGE_BRUSH( "Old/SelectionDetails/ShowAll", Icon32x32 ) );
+		Set( "LevelEditor.ShowSelectedOnly", new IMAGE_BRUSH( "Old/SelectionDetails/ShowSelected", Icon32x32 ) );
+		Set( "LevelEditor.ShowSelected", new IMAGE_BRUSH( "Old/SelectionDetails/ShowSelected", Icon32x32 ) );
+		Set( "LevelEditor.HideSelected", new IMAGE_BRUSH( "Old/SelectionDetails/HideSelected", Icon32x32 ) );
 	}
 
 	// Level viewport toolbar
@@ -5126,14 +5126,14 @@ void FSlateEditorStyle::FStyle::SetupLevelEditorStyle()
 		Set( "EditorViewportToolBar.Font", DEFAULT_FONT( "Bold", 9 ) );
 
 		Set( "EditorViewportToolBar.MenuButton", FButtonStyle(Button)
-			.SetNormal(BOX_BRUSH( "Common/SmallRoundedButton", FMargin(7.f/16.f), FLinearColor(1,1,1,0.75f)))
-			.SetHovered(BOX_BRUSH( "Common/SmallRoundedButton", FMargin(7.f/16.f), FLinearColor(1,1,1, 1.0f)))
+			.SetNormal(BOX_BRUSH( "Common/SmallRoundedButton", FMargin(7.f/16.f), FLinearColor(1.f,1.f,1.f,0.75f)))
+			.SetHovered(BOX_BRUSH( "Common/SmallRoundedButton", FMargin(7.f/16.f), FLinearColor(1.f,1.f,1.f, 1.0f)))
 			.SetPressed(BOX_BRUSH( "Common/SmallRoundedButton", FMargin(7.f/16.f) ))
 		);
 
 		Set("EditorViewportToolBar.MenuButtonWarning", FButtonStyle(Button)
 			.SetNormal(BOX_BRUSH("Common/SmallRoundedButton", FMargin(7.f / 16.f), LogColor_Warning_LinearRef))
-			.SetHovered(BOX_BRUSH("Common/SmallRoundedButton", FMargin(7.f / 16.f), *LogColor_Warning_LinearRef*FLinearColor(1.1, 1.1, 1.1, 1)))
+			.SetHovered(BOX_BRUSH("Common/SmallRoundedButton", FMargin(7.f / 16.f), *LogColor_Warning_LinearRef*FLinearColor(1.1f, 1.1f, 1.1f, 1.f)))
 			.SetPressed(BOX_BRUSH("Common/SmallRoundedButton", FMargin(7.f / 16.f), LogColor_Warning_LinearRef))
 		);
 
@@ -5189,7 +5189,7 @@ void FSlateEditorStyle::FStyle::SetupLevelEditorStyle()
 		const FString SmallRoundedButtonMiddle(TEXT("Common/SmallRoundedToggleCenter"));
 		const FString SmallRoundedButtonEnd(TEXT("Common/SmallRoundedToggleRight"));
 
-		const FLinearColor NormalColor(0.15, 0.15, 0.15, 1);
+		const FLinearColor NormalColor(0.15f, 0.15f, 0.15f, 1.f);
 
 		Set("Property.ToggleButton", FCheckBoxStyle()
 			.SetCheckBoxType(ESlateCheckBoxType::ToggleButton)
@@ -5299,14 +5299,14 @@ void FSlateEditorStyle::FStyle::SetupLevelEditorStyle()
 		Set( "EditorModesToolbar.Icon", new IMAGE_BRUSH( "Icons/icon_tab_toolbar_16px", Icon16x16 ) );
 		Set( "EditorModesToolbar.Expand", new IMAGE_BRUSH( "Icons/toolbar_expand_16x", Icon16x16) );
 		Set( "EditorModesToolbar.SubMenuIndicator", new IMAGE_BRUSH( "Common/SubmenuArrow", Icon8x8 ) );
-		Set( "EditorModesToolbar.SToolBarComboButtonBlock.Padding", FMargin( 0 ) );
+		Set( "EditorModesToolbar.SToolBarComboButtonBlock.Padding", FMargin( 0.f) );
 		Set( "EditorModesToolbar.SToolBarComboButtonBlock.ComboButton.Color", DefaultForeground );
-		Set( "EditorModesToolbar.SToolBarButtonBlock.Padding", FMargin( 1.0f, 0.0f, 0.0f, 0 ) );
+		Set( "EditorModesToolbar.SToolBarButtonBlock.Padding", FMargin( 1.0f, 0.0f, 0.0f, 0.f) );
 		Set( "EditorModesToolbar.SToolBarButtonBlock.CheckBox.Padding", FMargin( 6.0f, 4.0f, 6.0f, 6.0f ) );
-		Set( "EditorModesToolbar.SToolBarCheckComboButtonBlock.Padding", FMargin( 0 ) );
+		Set( "EditorModesToolbar.SToolBarCheckComboButtonBlock.Padding", FMargin( 0.f) );
 
-		Set( "EditorModesToolbar.Block.IndentedPadding", FMargin( 0 ) );
-		Set( "EditorModesToolbar.Block.Padding", FMargin( 0 ) );
+		Set( "EditorModesToolbar.Block.IndentedPadding", FMargin( 0.f) );
+		Set( "EditorModesToolbar.Block.Padding", FMargin( 0.f) );
 
 		Set( "EditorModesToolbar.Separator", new BOX_BRUSH( "Old/Button", 4.0f / 32.0f ) );
 		Set( "EditorModesToolbar.Separator.Padding", FMargin( 8.f, 0.f, 8.f, 0.f) );
@@ -5316,7 +5316,7 @@ void FSlateEditorStyle::FStyle::SetupLevelEditorStyle()
 		Set( "EditorModesToolbar.Keybinding", FTextBlockStyle(NormalText) .SetFont( DEFAULT_FONT( "Regular", 8 ) ) );
 
 		Set( "EditorModesToolbar.Heading.Font", DEFAULT_FONT( "Regular", 8 ) );
-		Set( "EditorModesToolbar.Heading.ColorAndOpacity", FLinearColor( 0.4f, 0.4, 0.4f, 1.0f ) );
+		Set( "EditorModesToolbar.Heading.ColorAndOpacity", FLinearColor( 0.4f, 0.4f, 0.4f, 1.0f ) );
 
 		/* Create style for "EditorModesToolbar.CheckBox" ... */
 		const FCheckBoxStyle EditorModesToolbarCheckBoxStyle = FCheckBoxStyle()
@@ -5354,12 +5354,12 @@ void FSlateEditorStyle::FStyle::SetupLevelEditorStyle()
 		/* Create style for "EditorModesToolbar.ToggleButton" ... */
 		const FCheckBoxStyle EditorModesToolbarToggleButtonStyle = FCheckBoxStyle()
 			.SetCheckBoxType( ESlateCheckBoxType::ToggleButton )
-			.SetUncheckedImage( BOX_BRUSH( "/EditorModes/Tab_Inactive", 4 / 16.0f ) )
-			.SetUncheckedPressedImage( BOX_BRUSH( "/EditorModes/Tab_Active", 4 / 16.0f ) )
-			.SetUncheckedHoveredImage( BOX_BRUSH( "/EditorModes/Tab_Active", 4 / 16.0f ) )
-			.SetCheckedHoveredImage( BOX_BRUSH( "/EditorModes/Tab_Active", 4 / 16.0f ) )
-			.SetCheckedPressedImage( BOX_BRUSH( "/EditorModes/Tab_Active", 4 / 16.0f ) )
-			.SetCheckedImage( BOX_BRUSH( "/EditorModes/Tab_Active", 4 / 16.0f ) );
+			.SetUncheckedImage( BOX_BRUSH( "/EditorModes/Tab_Inactive", 4.f / 16.0f ) )
+			.SetUncheckedPressedImage( BOX_BRUSH( "/EditorModes/Tab_Active", 4.f / 16.0f ) )
+			.SetUncheckedHoveredImage( BOX_BRUSH( "/EditorModes/Tab_Active", 4.f / 16.0f ) )
+			.SetCheckedHoveredImage( BOX_BRUSH( "/EditorModes/Tab_Active", 4.f / 16.0f ) )
+			.SetCheckedPressedImage( BOX_BRUSH( "/EditorModes/Tab_Active", 4.f / 16.0f ) )
+			.SetCheckedImage( BOX_BRUSH( "/EditorModes/Tab_Active", 4.f / 16.0f ) );
 
 		/* ... and add new style */
 		Set( "EditorModesToolbar.ToggleButton", EditorModesToolbarToggleButtonStyle );
@@ -5396,7 +5396,7 @@ void FSlateEditorStyle::FStyle::SetupLevelEditorStyle()
 
 	// Scalability (Performance Warning)
 	{
-		Set( "Scalability.ScalabilitySettings", new IMAGE_BRUSH("Scalability/ScalabilitySettings", FVector2D(473.0f, 266.0f) ) );
+		Set( "Scalability.ScalabilitySettings", new IMAGE_BRUSH("Scalability/ScalabilitySettings", FVector2f(473.0f, 266.0f) ) );
 	}
 
 	Set("WorkspaceMenu.AdditionalUI", new IMAGE_BRUSH("Icons/icon_ViewMode_LODColoration_16px", Icon16x16));
@@ -5414,8 +5414,8 @@ void FSlateEditorStyle::FStyle::SetupPersonaStyle()
 		Set( "AnimViewportMenu.RotateMode.Small", new IMAGE_BRUSH( "Icons/icon_rotate_40x", Icon16x16 ) );
 		Set( "AnimViewportMenu.CameraFollow", new IMAGE_BRUSH( "Persona/Viewport/Camera_FollowBounds_40px", Icon32x32) );
 		Set( "AnimViewportMenu.CameraFollow.Small", new IMAGE_BRUSH( "Persona/Viewport/Camera_FollowBounds_40px", Icon16x16 ) );
-		Set( "AnimViewport.LocalSpaceEditing", new IMAGE_BRUSH( "Icons/icon_axis_local_16px", FVector2D( 16, 16 ) ) );
-		Set( "AnimViewport.WorldSpaceEditing", new IMAGE_BRUSH( "Icons/icon_axis_world_16px", FVector2D( 16, 16 ) ) );
+		Set( "AnimViewport.LocalSpaceEditing", new IMAGE_BRUSH( "Icons/icon_axis_local_16px", Icon16x16 ) );
+		Set( "AnimViewport.WorldSpaceEditing", new IMAGE_BRUSH( "Icons/icon_axis_world_16px", Icon16x16 ) );
 		Set( "AnimViewportMenu.SetShowNormals", new IMAGE_BRUSH( TEXT("Icons/icon_StaticMeshEd_Normals_40x"), Icon40x40 ) );
 		Set( "AnimViewportMenu.SetShowNormals.Small", new IMAGE_BRUSH( TEXT("Icons/icon_StaticMeshEd_Normals_40x"), Icon20x20 ) );
 		Set( "AnimViewportMenu.SetShowTangents", new IMAGE_BRUSH( TEXT("Icons/icon_StaticMeshEd_Tangents_40x"), Icon40x40 ) );
@@ -5433,17 +5433,17 @@ void FSlateEditorStyle::FStyle::SetupPersonaStyle()
 
 		Set("AnimViewport.MessageText", FTextBlockStyle(NormalText)
 			.SetFont(DEFAULT_FONT("Bold", FCoreStyle::RegularTextSize))
-			.SetShadowOffset(FVector2D(1.0f, 1.0f))
+			.SetShadowOffset(FVector2f::UnitVector)
 			.SetShadowColorAndOpacity(FLinearColor(0.0f, 0.0f, 0.0f)));
 		Set("AnimViewport.WarningText", FTextBlockStyle(NormalText)
 			.SetFont(DEFAULT_FONT("Bold", FCoreStyle::RegularTextSize))
 			.SetColorAndOpacity(FLinearColor::Yellow)
-			.SetShadowOffset(FVector2D(1.0f, 1.0f))
+			.SetShadowOffset(FVector2f::UnitVector)
 			.SetShadowColorAndOpacity(FLinearColor(0.0f, 0.0f, 0.0f)));
 		Set("AnimViewport.ErrorText", FTextBlockStyle(NormalText)
 			.SetFont(DEFAULT_FONT("Bold", FCoreStyle::RegularTextSize))
 			.SetColorAndOpacity(FLinearColor::Red)
-			.SetShadowOffset(FVector2D(1.0f, 1.0f))
+			.SetShadowOffset(FVector2f::UnitVector)
 			.SetShadowColorAndOpacity(FLinearColor(0.0f, 0.0f, 0.0f)));
 		
 		// Viewport notifications 
@@ -5556,7 +5556,7 @@ void FSlateEditorStyle::FStyle::SetupPersonaStyle()
 			.SetFont(DEFAULT_FONT("Bold", 10)));
 
 		Set("SkeletonTree.HyperlinkSpinBox", FSpinBoxStyle(GetWidgetStyle<FSpinBoxStyle>("SpinBox"))
-			.SetTextPadding(FMargin(0))
+			.SetTextPadding(FMargin(0.f))
 			.SetBackgroundBrush(BORDER_BRUSH("Old/HyperlinkDotted", FMargin(0, 0, 0, 3 / 16.0f), FSlateColor::UseSubduedForeground()))
 			.SetHoveredBackgroundBrush(FSlateNoResource())
 			.SetInactiveFillBrush(FSlateNoResource())
@@ -5566,8 +5566,8 @@ void FSlateEditorStyle::FStyle::SetupPersonaStyle()
 			);
 
 		Set("SkeletonTree.BlendProfile", new IMAGE_BRUSH(TEXT("Icons/icon_Persona_NewBlendSpace_16x"), Icon16x16));
-		Set("SkeletonTree.InlineEditorShadowTop", new IMAGE_BRUSH(TEXT("Common/ScrollBoxShadowTop"), FVector2D(64, 8)));
-		Set("SkeletonTree.InlineEditorShadowBottom", new IMAGE_BRUSH(TEXT("Common/ScrollBoxShadowBottom"), FVector2D(64, 8)));
+		Set("SkeletonTree.InlineEditorShadowTop", new IMAGE_BRUSH(TEXT("Common/ScrollBoxShadowTop"), FVector2f(64.f, 8.f)));
+		Set("SkeletonTree.InlineEditorShadowBottom", new IMAGE_BRUSH(TEXT("Common/ScrollBoxShadowBottom"), FVector2f(64.f, 8.f)));
 	}
 
 	// Animation editor
@@ -5594,12 +5594,12 @@ void FSlateEditorStyle::FStyle::SetupPersonaStyle()
 
 		Set("AnimTimeline.Outliner.Label", FTextBlockStyle(NormalText)
 			.SetColorAndOpacity(FLinearColor(1.0f, 1.0f, 1.0f))
-			.SetShadowOffset(FVector2D(1.0f, 1.0f)));
+			.SetShadowOffset(FVector2f::UnitVector));
 
 		Set("AnimTimeline.Outliner.ItemColor", FLinearColor(0.04f, 0.04f, 0.04f, 0.8f));
 		Set("AnimTimeline.Outliner.HeaderColor", FLinearColor(0.03f, 0.03f, 0.03f, 1.0f));
 
-		Set("AnimTimeline.SectionMarker", new IMAGE_BRUSH(TEXT("Sequencer/SectionMarker"), FVector2D(11, 12)));
+		Set("AnimTimeline.SectionMarker", new IMAGE_BRUSH(TEXT("Sequencer/SectionMarker"), FVector2f(11.f, 12.f)));
 	}
 
 	// Skeletal mesh editor
@@ -5659,8 +5659,8 @@ void FSlateEditorStyle::FStyle::SetupPersonaStyle()
 			Set( "Kismet.Tabs.BlueprintDefaults", new IMAGE_BRUSH( "Icons/icon_BlueprintEditor_Defaults_40x", Icon16x16 ) );
 		}
 
-		Set("Kismet.Palette.Favorites", new IMAGE_BRUSH("Icons/Star_16x", Icon16x16, FLinearColor(0.4f, 0.4, 0.4f, 1.f)));
-		Set("Kismet.Palette.Library",   new IMAGE_BRUSH("Icons/icon_MeshPaint_Find_16x", Icon16x16, FLinearColor(0.4f, 0.4, 0.4f, 1.f)));
+		Set("Kismet.Palette.Favorites", new IMAGE_BRUSH("Icons/Star_16x", Icon16x16, FLinearColor(0.4f, 0.4f, 0.4f, 1.f)));
+		Set("Kismet.Palette.Library",   new IMAGE_BRUSH("Icons/icon_MeshPaint_Find_16x", Icon16x16, FLinearColor(0.4f, 0.4f, 0.4f, 1.f)));
 		const FCheckBoxStyle KismetFavoriteToggleStyle = FCheckBoxStyle()
 			.SetCheckBoxType(ESlateCheckBoxType::CheckBox)
 			.SetUncheckedImage( IMAGE_BRUSH("Icons/EmptyStar_16x", Icon10x10, FLinearColor(0.8f, 0.8f, 0.8f, 1.f)) )
@@ -5742,8 +5742,8 @@ void FSlateEditorStyle::FStyle::SetupPersonaStyle()
 		Set( "Kismet.CreateBlueprint", new IMAGE_BRUSH( "/Icons/CreateBlueprint", Icon16x16) );
 		Set( "Kismet.HarvestBlueprintFromActors", new IMAGE_BRUSH( "/Icons/HarvestBlueprintFromActors", Icon16x16) );
 
-		Set( "Kismet.Comment.Handle", new IMAGE_BRUSH( "Old/Kismet2/Comment_Handle", FVector2D(14.0f, 12.0f)) );
-		Set( "Kismet.Comment.Background", new IMAGE_BRUSH( "Old/Kismet2/Comment_Background", FVector2D(100.0f, 68.0f)) );
+		Set( "Kismet.Comment.Handle", new IMAGE_BRUSH( "Old/Kismet2/Comment_Handle", FVector2f(14.0f, 12.0f)) );
+		Set( "Kismet.Comment.Background", new IMAGE_BRUSH( "Old/Kismet2/Comment_Background", FVector2f(100.0f, 68.0f)) );
 
 		Set( "Kismet.AllClasses.VariableIcon", new IMAGE_BRUSH( "/Icons/pill_16x", Icon16x16 ) );
 		Set( "Kismet.AllClasses.ArrayVariableIcon", new IMAGE_BRUSH( "/Icons/pillarray_16x", Icon16x16 ) );
@@ -5897,8 +5897,8 @@ void FSlateEditorStyle::FStyle::SetupPersonaStyle()
 
 
 		Set("PlayToolbar.Label", FTextBlockStyle(NormalText).SetFont(DEFAULT_FONT("Regular", 10)));
-		Set("PlayToolbar.Label.Padding", FMargin(5, 0, 0, 0));
-		Set("PlayToolbar.Label.DropShadowSize", FVector2D::ZeroVector);
+		Set("PlayToolbar.Label.Padding", FMargin(5.f, 0.f, 0.f, 0.f));
+		Set("PlayToolbar.Label.DropShadowSize", FVector2f::ZeroVector);
 		Set("PlayToolbar.SToolBarButtonBlock.Padding", FMargin(5.0f, 0.f));
 		Set("PlayToolbar.IconSize", Icon40x40);
 
@@ -5972,8 +5972,8 @@ void FSlateEditorStyle::FStyle::SetupPersonaStyle()
 		Set( "Kismet.DebuggerOverlay.Breakpoint.EnabledAndInvalid", new IMAGE_BRUSH( "Old/Kismet2/Breakpoint_Invalid", Icon32x32 ) );
 		Set( "Kismet.DebuggerOverlay.Breakpoint.EnabledAndValid", new IMAGE_BRUSH( "Old/Kismet2/Breakpoint_Valid", Icon32x32 ) );
 
-		Set( "Kismet.DebuggerOverlay.InstructionPointer", new IMAGE_BRUSH( "Old/Kismet2/IP_Normal", FVector2D(128,96)) );
-		Set( "Kismet.DebuggerOverlay.InstructionPointerBreakpoint", new IMAGE_BRUSH( "Old/Kismet2/IP_Breakpoint", FVector2D(128,96)) );
+		Set( "Kismet.DebuggerOverlay.InstructionPointer", new IMAGE_BRUSH( "Old/Kismet2/IP_Normal", FVector2f(128.f,96.f)) );
+		Set( "Kismet.DebuggerOverlay.InstructionPointerBreakpoint", new IMAGE_BRUSH( "Old/Kismet2/IP_Breakpoint", FVector2f(128.f,96.f)) );
 
 		Set ("Kismet.CallStackViewer.CurrentStackFrame", new IMAGE_BRUSH( "Old/Kismet2/DebuggerArrow", Icon12x12 ));
 		Set( "Kismet.CallStackViewer.CurrentStackFrameColor", FLinearColor(0.728f, 0.364f, 0.003f) );
@@ -6026,7 +6026,7 @@ void FSlateEditorStyle::FStyle::SetupPersonaStyle()
 
 		Set( "PinnedCommandList.Heading", FTextBlockStyle(NormalText)
 			.SetFont(DEFAULT_FONT( "Regular", 8))
-			.SetColorAndOpacity(FLinearColor(0.4f, 0.4, 0.4f, 1.0f)));
+			.SetColorAndOpacity(FLinearColor(0.4f, 0.4f, 0.4f, 1.0f)));
 
 		/* Set images for various SCheckBox states associated with menu check box items... */
 		const FCheckBoxStyle BasicMenuCheckBoxStyle = FCheckBoxStyle()
@@ -6133,7 +6133,7 @@ void FSlateEditorStyle::FStyle::SetupPersonaStyle()
 
 		Set( "ViewportPinnedCommandList.Heading", FTextBlockStyle(NormalText)
 			.SetFont(DEFAULT_FONT( "Regular", 8))
-			.SetColorAndOpacity(FLinearColor(0.4f, 0.4, 0.4f, 1.0f)));
+			.SetColorAndOpacity(FLinearColor(0.4f, 0.4f, 0.4f, 1.0f)));
 
 		/* Set images for various SCheckBox states associated with menu check box items... */
 		const FCheckBoxStyle BasicMenuCheckBoxStyle = FCheckBoxStyle()
@@ -6558,12 +6558,12 @@ void FSlateEditorStyle::FStyle::SetupContentBrowserStyle()
 		Set("ContentBrowser.SourceTreeItemFont", DEFAULT_FONT("Regular", 10));
 
 		Set( "ContentBrowser.SourceTreeRootItemFont", DEFAULT_FONT( "Regular", 12 ) );
-		Set( "ContentBrowser.AssetTreeFolderClosed", new IMAGE_BRUSH( "Icons/FolderClosed", FVector2D(18, 16) ) );
-		Set( "ContentBrowser.BreadcrumbPathPickerFolder", new IMAGE_BRUSH( "Icons/FolderClosed", FVector2D(18, 16), FLinearColor::Gray ) );
-		Set( "ContentBrowser.AssetTreeFolderOpen", new IMAGE_BRUSH( "Icons/FolderOpen", FVector2D(18, 16) ) );
-		Set( "ContentBrowser.AssetTreeFolderDeveloper", new IMAGE_BRUSH( "Icons/FolderDeveloper", FVector2D(18, 16) ) );
-		Set( "ContentBrowser.AssetTreeFolderOpenCode", new IMAGE_BRUSH( "Icons/FolderOpen_Code", FVector2D(18, 16) ) );
-		Set( "ContentBrowser.AssetTreeFolderClosedCode", new IMAGE_BRUSH( "Icons/FolderClosed_Code", FVector2D(18, 16) ) );
+		Set( "ContentBrowser.AssetTreeFolderClosed", new IMAGE_BRUSH( "Icons/FolderClosed", FVector2f(18.f, 16.f) ) );
+		Set( "ContentBrowser.BreadcrumbPathPickerFolder", new IMAGE_BRUSH( "Icons/FolderClosed", FVector2f(18.f, 16.f), FLinearColor::Gray ) );
+		Set( "ContentBrowser.AssetTreeFolderOpen", new IMAGE_BRUSH( "Icons/FolderOpen", FVector2f(18.f, 16.f) ) );
+		Set( "ContentBrowser.AssetTreeFolderDeveloper", new IMAGE_BRUSH( "Icons/FolderDeveloper", FVector2f(18.f, 16.f) ) );
+		Set( "ContentBrowser.AssetTreeFolderOpenCode", new IMAGE_BRUSH( "Icons/FolderOpen_Code", FVector2f(18.f, 16.f) ) );
+		Set( "ContentBrowser.AssetTreeFolderClosedCode", new IMAGE_BRUSH( "Icons/FolderClosed_Code", FVector2f(18.f, 16.f) ) );
 		Set( "ContentBrowser.AddCollectionButtonIcon", new IMAGE_BRUSH( "Icons/PlusSymbol_12x", Icon12x12 ) );
 
 		Set( "ContentBrowser.Splitter", FSplitterStyle()
@@ -6607,7 +6607,7 @@ void FSlateEditorStyle::FStyle::SetupContentBrowserStyle()
 			);
 
 		Set( "ContentBrowser.TileViewTooltip.ToolTipBorder", new FSlateColorBrush( FLinearColor::Black ) );
-		Set( "ContentBrowser.TileViewTooltip.NonContentBorder", new BOX_BRUSH( "/Docking/TabContentArea", FMargin(4/16.0f) ) );
+		Set( "ContentBrowser.TileViewTooltip.NonContentBorder", new BOX_BRUSH( "/Docking/TabContentArea", FMargin(4.f/16.0f) ) );
 		Set( "ContentBrowser.TileViewTooltip.ContentBorder", new BOX_BRUSH( "Common/GroupBorder", FMargin(4.0f/16.0f) ) );
 		Set( "ContentBrowser.TileViewTooltip.PillBorder", new FSlateRoundedBoxBrush(FLinearColor::Transparent, 12.0f, FLinearColor::White, 0.5f));
 		Set( "ContentBrowser.TileViewTooltip.NameFont", DEFAULT_FONT( "Regular", 12 ) );
@@ -6620,12 +6620,12 @@ void FSlateEditorStyle::FStyle::SetupContentBrowserStyle()
 		// Filter list
 		/* Set images for various SCheckBox states associated with "ContentBrowser.FilterButton" ... */
 		const FCheckBoxStyle ContentBrowserFilterButtonCheckBoxStyle = FCheckBoxStyle()
-			.SetUncheckedImage( IMAGE_BRUSH( "ContentBrowser/FilterUnchecked", FVector2D( 10.0f,20.0f ) ) )
-			.SetUncheckedHoveredImage( IMAGE_BRUSH( "ContentBrowser/FilterUnchecked", FVector2D( 10.0f, 20.0f ), FLinearColor( 0.5f, 0.5f, 0.5f, 1.0f ) ) )
-			.SetUncheckedPressedImage( IMAGE_BRUSH( "ContentBrowser/FilterUnchecked",FVector2D( 10.0f, 20.0f ), FLinearColor( 0.5f, 0.5f, 0.5f, 1.0f ) ) )
-			.SetCheckedImage( IMAGE_BRUSH( "ContentBrowser/FilterChecked",  FVector2D( 10.0f, 20.0f ) ) )
-			.SetCheckedHoveredImage( IMAGE_BRUSH( "ContentBrowser/FilterChecked",  FVector2D( 10.0f, 20.0f ), FLinearColor( 0.5f, 0.5f, 0.5f, 1.0f ) ) )
-			.SetCheckedPressedImage( IMAGE_BRUSH( "ContentBrowser/FilterChecked", FVector2D( 10.0f, 20.0f ), FLinearColor( 0.5f, 0.5f, 0.5f, 1.0f ) ) );
+			.SetUncheckedImage( IMAGE_BRUSH( "ContentBrowser/FilterUnchecked", FVector2f( 10.0f,20.0f ) ) )
+			.SetUncheckedHoveredImage( IMAGE_BRUSH( "ContentBrowser/FilterUnchecked", FVector2f( 10.0f, 20.0f ), FLinearColor( 0.5f, 0.5f, 0.5f, 1.0f ) ) )
+			.SetUncheckedPressedImage( IMAGE_BRUSH( "ContentBrowser/FilterUnchecked",FVector2f( 10.0f, 20.0f ), FLinearColor( 0.5f, 0.5f, 0.5f, 1.0f ) ) )
+			.SetCheckedImage( IMAGE_BRUSH( "ContentBrowser/FilterChecked",  FVector2f( 10.0f, 20.0f ) ) )
+			.SetCheckedHoveredImage( IMAGE_BRUSH( "ContentBrowser/FilterChecked",  FVector2f( 10.0f, 20.0f ), FLinearColor( 0.5f, 0.5f, 0.5f, 1.0f ) ) )
+			.SetCheckedPressedImage( IMAGE_BRUSH( "ContentBrowser/FilterChecked", FVector2f( 10.0f, 20.0f ), FLinearColor( 0.5f, 0.5f, 0.5f, 1.0f ) ) );
 		/* ... and add the new style */
 		Set( "ContentBrowser.FilterButton", ContentBrowserFilterButtonCheckBoxStyle );
 
@@ -6656,15 +6656,15 @@ void FSlateEditorStyle::FStyle::SetupContentBrowserStyle()
 
 		// Asset tags (standard)
 		Set("ContentBrowser.AssetTagButton", FCheckBoxStyle()
-			.SetUncheckedImage(IMAGE_BRUSH("ContentBrowser/AssetTagCheckbox_Flat", FVector2D(14.0f, 28.0f)))
-			.SetUncheckedHoveredImage(IMAGE_BRUSH("ContentBrowser/AssetTagCheckbox_Flat", FVector2D(14.0f, 28.0f), FLinearColor(0.5f, 0.5f, 0.5f, 1.0f)))
-			.SetUncheckedPressedImage(IMAGE_BRUSH("ContentBrowser/AssetTagCheckbox_Flat", FVector2D(14.0f, 28.0f), FLinearColor(0.5f, 0.5f, 0.5f, 1.0f)))
-			.SetUndeterminedImage(IMAGE_BRUSH("ContentBrowser/AssetTagCheckbox_Flat", FVector2D(14.0f, 28.0f)))
-			.SetUndeterminedHoveredImage(IMAGE_BRUSH("ContentBrowser/AssetTagCheckbox_Flat", FVector2D(14.0f, 28.0f), FLinearColor(0.5f, 0.5f, 0.5f, 1.0f)))
-			.SetUndeterminedPressedImage(IMAGE_BRUSH("ContentBrowser/AssetTagCheckbox_Flat", FVector2D(14.0f, 28.0f), FLinearColor(0.5f, 0.5f, 0.5f, 1.0f)))
-			.SetCheckedImage(IMAGE_BRUSH("ContentBrowser/AssetTagCheckbox_Flat", FVector2D(14.0f, 28.0f)))
-			.SetCheckedHoveredImage(IMAGE_BRUSH("ContentBrowser/AssetTagCheckbox_Flat", FVector2D(14.0f, 28.0f), FLinearColor(0.5f, 0.5f, 0.5f, 1.0f)))
-			.SetCheckedPressedImage(IMAGE_BRUSH("ContentBrowser/AssetTagCheckbox_Flat", FVector2D(14.0f, 28.0f), FLinearColor(0.5f, 0.5f, 0.5f, 1.0f)))
+			.SetUncheckedImage(IMAGE_BRUSH("ContentBrowser/AssetTagCheckbox_Flat", FVector2f(14.0f, 28.0f)))
+			.SetUncheckedHoveredImage(IMAGE_BRUSH("ContentBrowser/AssetTagCheckbox_Flat", FVector2f(14.0f, 28.0f), FLinearColor(0.5f, 0.5f, 0.5f, 1.0f)))
+			.SetUncheckedPressedImage(IMAGE_BRUSH("ContentBrowser/AssetTagCheckbox_Flat", FVector2f(14.0f, 28.0f), FLinearColor(0.5f, 0.5f, 0.5f, 1.0f)))
+			.SetUndeterminedImage(IMAGE_BRUSH("ContentBrowser/AssetTagCheckbox_Flat", FVector2f(14.0f, 28.0f)))
+			.SetUndeterminedHoveredImage(IMAGE_BRUSH("ContentBrowser/AssetTagCheckbox_Flat", FVector2f(14.0f, 28.0f), FLinearColor(0.5f, 0.5f, 0.5f, 1.0f)))
+			.SetUndeterminedPressedImage(IMAGE_BRUSH("ContentBrowser/AssetTagCheckbox_Flat", FVector2f(14.0f, 28.0f), FLinearColor(0.5f, 0.5f, 0.5f, 1.0f)))
+			.SetCheckedImage(IMAGE_BRUSH("ContentBrowser/AssetTagCheckbox_Flat", FVector2f(14.0f, 28.0f)))
+			.SetCheckedHoveredImage(IMAGE_BRUSH("ContentBrowser/AssetTagCheckbox_Flat", FVector2f(14.0f, 28.0f), FLinearColor(0.5f, 0.5f, 0.5f, 1.0f)))
+			.SetCheckedPressedImage(IMAGE_BRUSH("ContentBrowser/AssetTagCheckbox_Flat", FVector2f(14.0f, 28.0f), FLinearColor(0.5f, 0.5f, 0.5f, 1.0f)))
 			.SetPadding(0.0f)
 			);
 		Set("ContentBrowser.AssetTagNameFont", DEFAULT_FONT("Bold", 11));
@@ -6674,15 +6674,15 @@ void FSlateEditorStyle::FStyle::SetupContentBrowserStyle()
 
 		// Asset tags (compact)
 		Set("ContentBrowser.AssetTagButton.Compact", FCheckBoxStyle()
-			.SetUncheckedImage(IMAGE_BRUSH("ContentBrowser/AssetTagCheckbox_Flat_Compact", FVector2D(10.0f, 20.0f)))
-			.SetUncheckedHoveredImage(IMAGE_BRUSH("ContentBrowser/AssetTagCheckbox_Flat_Compact", FVector2D(10.0f, 20.0f), FLinearColor(0.5f, 0.5f, 0.5f, 1.0f)))
-			.SetUncheckedPressedImage(IMAGE_BRUSH("ContentBrowser/AssetTagCheckbox_Flat_Compact", FVector2D(10.0f, 20.0f), FLinearColor(0.5f, 0.5f, 0.5f, 1.0f)))
-			.SetUndeterminedImage(IMAGE_BRUSH("ContentBrowser/AssetTagCheckbox_Flat_Compact", FVector2D(10.0f, 20.0f)))
-			.SetUndeterminedHoveredImage(IMAGE_BRUSH("ContentBrowser/AssetTagCheckbox_Flat_Compact", FVector2D(10.0f, 20.0f), FLinearColor(0.5f, 0.5f, 0.5f, 1.0f)))
-			.SetUndeterminedPressedImage(IMAGE_BRUSH("ContentBrowser/AssetTagCheckbox_Flat_Compact", FVector2D(10.0f, 20.0f), FLinearColor(0.5f, 0.5f, 0.5f, 1.0f)))
-			.SetCheckedImage(IMAGE_BRUSH("ContentBrowser/AssetTagCheckbox_Flat_Compact", FVector2D(10.0f, 20.0f)))
-			.SetCheckedHoveredImage(IMAGE_BRUSH("ContentBrowser/AssetTagCheckbox_Flat_Compact", FVector2D(10.0f, 20.0f), FLinearColor(0.5f, 0.5f, 0.5f, 1.0f)))
-			.SetCheckedPressedImage(IMAGE_BRUSH("ContentBrowser/AssetTagCheckbox_Flat_Compact", FVector2D(10.0f, 20.0f), FLinearColor(0.5f, 0.5f, 0.5f, 1.0f)))
+			.SetUncheckedImage(IMAGE_BRUSH("ContentBrowser/AssetTagCheckbox_Flat_Compact", FVector2f(10.0f, 20.0f)))
+			.SetUncheckedHoveredImage(IMAGE_BRUSH("ContentBrowser/AssetTagCheckbox_Flat_Compact", FVector2f(10.0f, 20.0f), FLinearColor(0.5f, 0.5f, 0.5f, 1.0f)))
+			.SetUncheckedPressedImage(IMAGE_BRUSH("ContentBrowser/AssetTagCheckbox_Flat_Compact", FVector2f(10.0f, 20.0f), FLinearColor(0.5f, 0.5f, 0.5f, 1.0f)))
+			.SetUndeterminedImage(IMAGE_BRUSH("ContentBrowser/AssetTagCheckbox_Flat_Compact", FVector2f(10.0f, 20.0f)))
+			.SetUndeterminedHoveredImage(IMAGE_BRUSH("ContentBrowser/AssetTagCheckbox_Flat_Compact", FVector2f(10.0f, 20.0f), FLinearColor(0.5f, 0.5f, 0.5f, 1.0f)))
+			.SetUndeterminedPressedImage(IMAGE_BRUSH("ContentBrowser/AssetTagCheckbox_Flat_Compact", FVector2f(10.0f, 20.0f), FLinearColor(0.5f, 0.5f, 0.5f, 1.0f)))
+			.SetCheckedImage(IMAGE_BRUSH("ContentBrowser/AssetTagCheckbox_Flat_Compact", FVector2f(10.0f, 20.0f)))
+			.SetCheckedHoveredImage(IMAGE_BRUSH("ContentBrowser/AssetTagCheckbox_Flat_Compact", FVector2f(10.0f, 20.0f), FLinearColor(0.5f, 0.5f, 0.5f, 1.0f)))
+			.SetCheckedPressedImage(IMAGE_BRUSH("ContentBrowser/AssetTagCheckbox_Flat_Compact", FVector2f(10.0f, 20.0f), FLinearColor(0.5f, 0.5f, 0.5f, 1.0f)))
 			.SetPadding(0.0f)
 			);
 		Set("ContentBrowser.AssetTagNameFont.Compact", DEFAULT_FONT("Regular", 10));
@@ -6694,7 +6694,7 @@ void FSlateEditorStyle::FStyle::SetupContentBrowserStyle()
 		Set( "ContentBrowser.EditModeLabelFont", FTextBlockStyle(NormalText)
 			.SetFont( DEFAULT_FONT( "Regular", 10 ) )
 			.SetColorAndOpacity( FLinearColor::Black )
-			.SetShadowOffset( FVector2D::ZeroVector )
+			.SetShadowOffset( FVector2f::ZeroVector )
 		);
 
 		Set( "ContentBrowser.EditModeLabelBorder", new FSlateColorBrush(SelectionColor_LinearRef) );
@@ -6711,7 +6711,7 @@ void FSlateEditorStyle::FStyle::SetupContentBrowserStyle()
 			.SetFont( DEFAULT_FONT( "Bold", 11 ) )
 			.SetColorAndOpacity( FLinearColor( 1.0f, 1.0f, 1.0f ) )
 			.SetHighlightColor( FLinearColor( 1.0f, 1.0f, 1.0f ) )
-			.SetShadowOffset( FVector2D( 1, 1 ) )
+			.SetShadowOffset( FVector2f::UnitVector )
 			.SetShadowColorAndOpacity( FLinearColor( 0, 0, 0, 0.9f ) ) );
 
 		// New Asset
@@ -6760,22 +6760,22 @@ void FSlateEditorStyle::FStyle::SetupContentBrowserStyle()
 		// Misc
 		Set( "ContentBrowser.ThumbnailShadow", new BOX_BRUSH( "ContentBrowser/ThumbnailShadow" , FMargin( 4.0f / 64.0f ) ) );
 		Set( "ContentBrowser.ColumnViewAssetIcon", new IMAGE_BRUSH( "Icons/doc_16x", Icon16x16 ) );
-		Set( "ContentBrowser.ColumnViewFolderIcon", new IMAGE_BRUSH( "Icons/FolderClosed", FVector2D(18, 16) ) );
-		Set( "ContentBrowser.ColumnViewDeveloperFolderIcon", new IMAGE_BRUSH( "Icons/FolderDeveloper", FVector2D(18, 16) ) );
-		Set( "ContentBrowser.ListViewFolderIcon.Base", new IMAGE_BRUSH( "Icons/Folders/Folder_Base_256x", FVector2D(256, 256) ) );
-		Set( "ContentBrowser.ListViewFolderIcon.Mask", new IMAGE_BRUSH( "Icons/Folders/Folder_BaseHi_256x", FVector2D(256, 256) ) );
-		Set( "ContentBrowser.ListViewDeveloperFolderIcon.Base", new IMAGE_BRUSH( "Icons/Folders/FolderDev_Base_256x", FVector2D(256, 256) ) );
-		Set( "ContentBrowser.ListViewDeveloperFolderIcon.Mask", new IMAGE_BRUSH( "Icons/Folders/FolderDev_BaseHi_256x", FVector2D(256, 256) ) );
-		Set( "ContentBrowser.TileViewFolderIcon.Base", new IMAGE_BRUSH( "Icons/Folders/Folder_Base_512x", FVector2D(512, 512) ) );
-		Set( "ContentBrowser.TileViewFolderIcon.Mask", new IMAGE_BRUSH( "Icons/Folders/Folder_BaseHi_512x", FVector2D(512, 512) ) );
-		Set( "ContentBrowser.TileViewDeveloperFolderIcon.Base", new IMAGE_BRUSH( "Icons/Folders/FolderDev_Base_512x", FVector2D(512, 512) ) );
-		Set( "ContentBrowser.TileViewDeveloperFolderIcon.Mask", new IMAGE_BRUSH( "Icons/Folders/FolderDev_BaseHi_512x", FVector2D(512, 512) ) );
+		Set( "ContentBrowser.ColumnViewFolderIcon", new IMAGE_BRUSH( "Icons/FolderClosed", FVector2f(18.f, 16.f) ) );
+		Set( "ContentBrowser.ColumnViewDeveloperFolderIcon", new IMAGE_BRUSH( "Icons/FolderDeveloper", FVector2f(18.f, 16.f) ) );
+		Set( "ContentBrowser.ListViewFolderIcon.Base", new IMAGE_BRUSH( "Icons/Folders/Folder_Base_256x", FVector2f(256.f, 256.f) ) );
+		Set( "ContentBrowser.ListViewFolderIcon.Mask", new IMAGE_BRUSH( "Icons/Folders/Folder_BaseHi_256x", FVector2f(256.f, 256.f) ) );
+		Set( "ContentBrowser.ListViewDeveloperFolderIcon.Base", new IMAGE_BRUSH( "Icons/Folders/FolderDev_Base_256x", FVector2f(256.f, 256.f) ) );
+		Set( "ContentBrowser.ListViewDeveloperFolderIcon.Mask", new IMAGE_BRUSH( "Icons/Folders/FolderDev_BaseHi_256x", FVector2f(256.f, 256.f) ) );
+		Set( "ContentBrowser.TileViewFolderIcon.Base", new IMAGE_BRUSH( "Icons/Folders/Folder_Base_512x", FVector2f(512.f, 512.f) ) );
+		Set( "ContentBrowser.TileViewFolderIcon.Mask", new IMAGE_BRUSH( "Icons/Folders/Folder_BaseHi_512x", FVector2f(512.f, 512.f) ) );
+		Set( "ContentBrowser.TileViewDeveloperFolderIcon.Base", new IMAGE_BRUSH( "Icons/Folders/FolderDev_Base_512x", FVector2f(512.f, 512.f) ) );
+		Set( "ContentBrowser.TileViewDeveloperFolderIcon.Mask", new IMAGE_BRUSH( "Icons/Folders/FolderDev_BaseHi_512x", FVector2f(512.f, 512.f) ) );
 		Set( "ContentBrowser.PathText", FTextBlockStyle(NormalText)
 			.SetFont( DEFAULT_FONT( "Bold", 11 ) )
 			.SetColorAndOpacity( FLinearColor( 1.0f, 1.0f, 1.0f ) )
 			.SetHighlightColor( FLinearColor( 1.0f, 1.0f, 1.0f ) )
-			.SetShadowOffset( FVector2D( 1,1 ) )
-			.SetShadowColorAndOpacity( FLinearColor(0,0,0,0.9f) ) );
+			.SetShadowOffset( FVector2f::UnitVector )
+			.SetShadowColorAndOpacity( FLinearColor(0.f,0.f,0.f,0.9f) ) );
 
 		Set("ReferenceViewer.PathText", FEditableTextBoxStyle(NormalEditableTextBoxStyle)
 			.SetFont(DEFAULT_FONT("Bold", 11)));
@@ -6984,27 +6984,27 @@ void FSlateEditorStyle::FStyle::SetupToolkitStyles()
 	{
 		Set("ProjectBrowser.Tab.Text", FTextBlockStyle(NormalText)
 			.SetFont(DEFAULT_FONT("Regular", 24))
-			.SetShadowOffset(FVector2D(0, 1)));
+			.SetShadowOffset(FVector2f(0.f, 1.f)));
 
 		Set("ProjectBrowser.Toolbar.Text", FTextBlockStyle(NormalText)
 			.SetFont(DEFAULT_FONT("Bold", 12))
-			.SetShadowColorAndOpacity(FLinearColor(0, 0, 0, 0.5f))
-			.SetShadowOffset(FVector2D(0, 1)));
+			.SetShadowColorAndOpacity(FLinearColor(0.f, 0.f, 0.f, 0.5f))
+			.SetShadowOffset(FVector2f(0.f, 1.f)));
 
 		Set("ProjectBrowser.VersionOverlayText", FTextBlockStyle(NormalText)
 			.SetFont(DEFAULT_FONT("Bold", 14))
-			.SetShadowOffset(FVector2D(0, 1)));
+			.SetShadowOffset(FVector2f(0.f, 1.f)));
 
 
-		Set("ProjectBrowser.Background", new BOX_BRUSH("Common/FlatColorSquare", FVector2D(1.0f, 1.0f), FMargin(0), FLinearColor(FColor(0xff404040))));
-		Set("ProjectBrowser.Tab.ActiveBackground", new BOX_BRUSH("Common/FlatColorSquare", FVector2D(1.0f, 1.0f), FMargin(0), FLinearColor(FColor(0xff404040))));
-		Set("ProjectBrowser.Tab.Background", new BOX_BRUSH("Common/FlatColorSquare", FVector2D(1.0f, 1.0f), FMargin(0), FLinearColor(FColor(0xff272727))));
-		Set("ProjectBrowser.Tab.ActiveHighlight", new BOX_BRUSH("Common/FlatColorSquare", FVector2D(1.0f, 1.0f), FMargin(0), SelectionColor));
-		Set("ProjectBrowser.Tab.Highlight", new BOX_BRUSH("Common/FlatColorSquare", FVector2D(1.0f, 1.0f), FMargin(0), SelectionColor_Inactive));
-		Set("ProjectBrowser.Tab.PressedHighlight", new BOX_BRUSH("Common/FlatColorSquare", FVector2D(1.0f, 1.0f), FMargin(0), SelectionColor_Pressed));
+		Set("ProjectBrowser.Background", new BOX_BRUSH("Common/FlatColorSquare", FVector2f::UnitVector, FMargin(0.f), FLinearColor(FColor(0xff404040))));
+		Set("ProjectBrowser.Tab.ActiveBackground", new BOX_BRUSH("Common/FlatColorSquare", FVector2f::UnitVector, FMargin(0.f), FLinearColor(FColor(0xff404040))));
+		Set("ProjectBrowser.Tab.Background", new BOX_BRUSH("Common/FlatColorSquare", FVector2f::UnitVector, FMargin(0.f), FLinearColor(FColor(0xff272727))));
+		Set("ProjectBrowser.Tab.ActiveHighlight", new BOX_BRUSH("Common/FlatColorSquare", FVector2f::UnitVector, FMargin(0.f), SelectionColor));
+		Set("ProjectBrowser.Tab.Highlight", new BOX_BRUSH("Common/FlatColorSquare", FVector2f::UnitVector, FMargin(0.f), SelectionColor_Inactive));
+		Set("ProjectBrowser.Tab.PressedHighlight", new BOX_BRUSH("Common/FlatColorSquare", FVector2f::UnitVector, FMargin(0.f), SelectionColor_Pressed));
 
 		Set("ProjectBrowser.TileViewTooltip.ToolTipBorder", new FSlateColorBrush(FLinearColor::Black));
-		Set("ProjectBrowser.TileViewTooltip.NonContentBorder", new BOX_BRUSH("/Docking/TabContentArea", FMargin(4 / 16.0f)));
+		Set("ProjectBrowser.TileViewTooltip.NonContentBorder", new BOX_BRUSH("/Docking/TabContentArea", FMargin(4.f / 16.0f)));
 		Set("ProjectBrowser.TileViewTooltip.ContentBorder", new BOX_BRUSH("Common/GroupBorder", FMargin(4.0f / 16.0f)));
 		Set("ProjectBrowser.TileViewTooltip.NameFont", DEFAULT_FONT("Regular", 12));
 	}
@@ -7013,7 +7013,7 @@ void FSlateEditorStyle::FStyle::SetupToolkitStyles()
 	{
 		Set("ToolkitDisplay.UnsavedChangeIcon", new IMAGE_BRUSH("Common/UnsavedChange", Icon8x8));
 		Set("ToolkitDisplay.MenuDropdown", new IMAGE_BRUSH("Common/ComboArrow", Icon8x8));
-		Set("ToolkitDisplay.ColorOverlay", new BOX_BRUSH("/Docking/Tab_ColorOverlay", 4 / 16.0f));
+		Set("ToolkitDisplay.ColorOverlay", new BOX_BRUSH("/Docking/Tab_ColorOverlay", 4.f / 16.0f));
 
 		FComboButtonStyle ComboButton = FComboButtonStyle()
 			.SetButtonStyle(GetWidgetStyle<FButtonStyle>("Button"))
@@ -7071,10 +7071,10 @@ void FSlateEditorStyle::FStyle::SetupToolkitStyles()
 		);
 
 		Set("DataTableEditor.CellListViewRow", FTableRowStyle(NormalTableRowStyle)
-			.SetEvenRowBackgroundBrush(IMAGE_BRUSH("PropertyView/DetailCategoryMiddle", FVector2D(16, 16), FLinearColor(0.5f, 0.5f, 0.5f)))
-			.SetEvenRowBackgroundHoveredBrush(IMAGE_BRUSH("PropertyView/DetailCategoryMiddle_Hovered", FVector2D(16, 16), FLinearColor(0.5f, 0.5f, 0.5f)))
-			.SetOddRowBackgroundBrush(IMAGE_BRUSH("PropertyView/DetailCategoryMiddle", FVector2D(16, 16), FLinearColor(0.2f, 0.2f, 0.2f)))
-			.SetOddRowBackgroundHoveredBrush(IMAGE_BRUSH("PropertyView/DetailCategoryMiddle_Hovered", FVector2D(16, 16), FLinearColor(0.2f, 0.2f, 0.2f)))
+			.SetEvenRowBackgroundBrush(IMAGE_BRUSH("PropertyView/DetailCategoryMiddle", Icon16x16, FLinearColor(0.5f, 0.5f, 0.5f)))
+			.SetEvenRowBackgroundHoveredBrush(IMAGE_BRUSH("PropertyView/DetailCategoryMiddle_Hovered", Icon16x16, FLinearColor(0.5f, 0.5f, 0.5f)))
+			.SetOddRowBackgroundBrush(IMAGE_BRUSH("PropertyView/DetailCategoryMiddle", Icon16x16, FLinearColor(0.2f, 0.2f, 0.2f)))
+			.SetOddRowBackgroundHoveredBrush(IMAGE_BRUSH("PropertyView/DetailCategoryMiddle_Hovered", Icon16x16, FLinearColor(0.2f, 0.2f, 0.2f)))
 			.SetActiveBrush(IMAGE_BRUSH("Common/Selection", Icon8x8, FLinearColor(0.075f, 0.075f, 0.075f)))
 			.SetActiveHoveredBrush(IMAGE_BRUSH("Common/Selection", Icon8x8, FLinearColor(0.075f, 0.075f, 0.075f)))
 			.SetInactiveBrush(IMAGE_BRUSH("Common/Selection", Icon8x8, FLinearColor(0.075f, 0.075f, 0.075f)))
@@ -7147,10 +7147,10 @@ void FSlateEditorStyle::FStyle::SetupToolkitStyles()
 		Set( "MaterialEditor.ViewportBorder", new BORDER_BRUSH("Old/White", FMargin(1), FLinearColor::White ) );
 
 		Set("MaterialEditor.CellListViewRow", FTableRowStyle(NormalTableRowStyle)
-			.SetEvenRowBackgroundBrush(IMAGE_BRUSH("PropertyView/DetailCategoryMiddle", FVector2D(16, 16), FLinearColor(0.5f, 0.5f, 0.5f)))
-			.SetEvenRowBackgroundHoveredBrush(IMAGE_BRUSH("PropertyView/DetailCategoryMiddle_Hovered", FVector2D(16, 16), FLinearColor(0.5f, 0.5f, 0.5f)))
-			.SetOddRowBackgroundBrush(IMAGE_BRUSH("PropertyView/DetailCategoryMiddle", FVector2D(16, 16), FLinearColor(0.35f, 0.35f, 0.35f)))
-			.SetOddRowBackgroundHoveredBrush(IMAGE_BRUSH("PropertyView/DetailCategoryMiddle_Hovered", FVector2D(16, 16), FLinearColor(0.35f, 0.35f, 0.35f)))
+			.SetEvenRowBackgroundBrush(IMAGE_BRUSH("PropertyView/DetailCategoryMiddle", Icon16x16, FLinearColor(0.5f, 0.5f, 0.5f)))
+			.SetEvenRowBackgroundHoveredBrush(IMAGE_BRUSH("PropertyView/DetailCategoryMiddle_Hovered", Icon16x16, FLinearColor(0.5f, 0.5f, 0.5f)))
+			.SetOddRowBackgroundBrush(IMAGE_BRUSH("PropertyView/DetailCategoryMiddle", Icon16x16, FLinearColor(0.35f, 0.35f, 0.35f)))
+			.SetOddRowBackgroundHoveredBrush(IMAGE_BRUSH("PropertyView/DetailCategoryMiddle_Hovered", Icon16x16, FLinearColor(0.35f, 0.35f, 0.35f)))
 			.SetActiveBrush(IMAGE_BRUSH("Common/Selection", Icon8x8, FLinearColor(0.075f, 0.075f, 0.075f)))
 			.SetActiveHoveredBrush(IMAGE_BRUSH("Common/Selection", Icon8x8, FLinearColor(0.075f, 0.075f, 0.075f)))
 			.SetInactiveBrush(IMAGE_BRUSH("Common/Selection", Icon8x8, FLinearColor(0.075f, 0.075f, 0.075f)))
@@ -7405,8 +7405,8 @@ void FSlateEditorStyle::FStyle::SetupToolkitStyles()
 		Set("CurveEditor.CreateTab.Small", new IMAGE_BRUSH("Icons/icon_CurveEditor_Create_40x", Icon20x20));
 		Set("CurveEditor.DeleteTab.Small", new IMAGE_BRUSH("Icons/icon_CurveEditor_DeleteTab_40x", Icon20x20));
 
-		Set("CurveEditor.Gradient.HandleDown", new BOX_BRUSH("Sequencer/ScrubHandleDown", FMargin(6.f / 13.f, 5 / 12.f, 6 / 13.f, 8 / 12.f)));
-		Set("CurveEditor.Gradient.HandleUp", new BOX_BRUSH("Sequencer/ScrubHandleUp", FMargin(6.f / 13.f, 8 / 12.f, 6 / 13.f, 5 / 12.f)));
+		Set("CurveEditor.Gradient.HandleDown", new BOX_BRUSH("Sequencer/ScrubHandleDown", FMargin(6.f / 13.f, 5.f / 12.f, 6.f / 13.f, 8.f / 12.f)));
+		Set("CurveEditor.Gradient.HandleUp", new BOX_BRUSH("Sequencer/ScrubHandleUp", FMargin(6.f / 13.f, 8.f / 12.f, 6.f / 13.f, 5.f / 12.f)));
 	}
 
 	// New Curve Editor 
@@ -7471,12 +7471,12 @@ void FSlateEditorStyle::FStyle::SetupToolkitStyles()
 		Set("GenericCurveEditor.OpenUserImplementableFilterWindow.Small", new IMAGE_BRUSH("GenericCurveEditor/Icons/MiscFilters_48x", Icon24x24));
 
 		// Key Types
-		Set("GenericCurveEditor.Key", new IMAGE_BRUSH("GenericCurveEditor/Keys/Diamond_Filled", FVector2D(11.0f, 11.0f))); 
-		Set("GenericCurveEditor.ConstantKey", new IMAGE_BRUSH("GenericCurveEditor/Keys/Square_Filled", FVector2D(11.0f, 11.0f))); 
-		Set("GenericCurveEditor.LinearKey", new IMAGE_BRUSH("GenericCurveEditor/Keys/Triangle_Filled", FVector2D(11.0f, 11.0f)));
-		Set("GenericCurveEditor.CubicKey", new IMAGE_BRUSH("GenericCurveEditor/Keys/Diamond_Filled", FVector2D(11.0f, 11.0f)));
-		Set("GenericCurveEditor.TangentHandle", new IMAGE_BRUSH("GenericCurveEditor/Keys/TangentHandle", FVector2D(8.0f, 8.0f)));
-		Set("GenericCurveEditor.WeightedTangentCubicKey", new IMAGE_BRUSH("GenericCurveEditor/Keys/Trapezoid_Filled", FVector2D(11.0f, 11.0f)));
+		Set("GenericCurveEditor.Key", new IMAGE_BRUSH("GenericCurveEditor/Keys/Diamond_Filled", FVector2f(11.0f, 11.0f))); 
+		Set("GenericCurveEditor.ConstantKey", new IMAGE_BRUSH("GenericCurveEditor/Keys/Square_Filled", FVector2f(11.0f, 11.0f))); 
+		Set("GenericCurveEditor.LinearKey", new IMAGE_BRUSH("GenericCurveEditor/Keys/Triangle_Filled", FVector2f(11.0f, 11.0f)));
+		Set("GenericCurveEditor.CubicKey", new IMAGE_BRUSH("GenericCurveEditor/Keys/Diamond_Filled", FVector2f(11.0f, 11.0f)));
+		Set("GenericCurveEditor.TangentHandle", new IMAGE_BRUSH("GenericCurveEditor/Keys/TangentHandle", FVector2f(8.0f, 8.0f)));
+		Set("GenericCurveEditor.WeightedTangentCubicKey", new IMAGE_BRUSH("GenericCurveEditor/Keys/Trapezoid_Filled", FVector2f(11.0f, 11.0f)));
 
 		// Pre-Infinity
 		Set("GenericCurveEditor.SetPreInfinityExtrapConstant", new IMAGE_BRUSH("GenericCurveEditor/Icons/PreInfinityConstant_48x", Icon48x48));
@@ -7684,14 +7684,14 @@ void FSlateEditorStyle::FStyle::SetupToolkitStyles()
 			.SetFont(DEFAULT_FONT("Bold", 11))
 			.SetColorAndOpacity(FLinearColor(1.0f, 1.0f, 1.0f))
 			.SetHighlightColor(FLinearColor(1.0f, 1.0f, 1.0f))
-			.SetShadowOffset(FVector2D(1, 1))
+			.SetShadowOffset(FVector2f::UnitVector)
 			.SetShadowColorAndOpacity(FLinearColor(0, 0, 0, 0.9f)));
 
 		Set("PhysicsAssetEditor.Tools.Font", FTextBlockStyle(NormalText)
 			.SetFont(DEFAULT_FONT("Bold", 11))
 			.SetColorAndOpacity(FLinearColor(1.0f, 1.0f, 1.0f))
 			.SetHighlightColor(FLinearColor(1.0f, 1.0f, 1.0f))
-			.SetShadowOffset(FVector2D(1, 1))
+			.SetShadowOffset(FVector2f::UnitVector)
 			.SetShadowColorAndOpacity(FLinearColor(0, 0, 0, 0.9f)));
 
 		FLinearColor Red = FLinearColor::Red;
@@ -7935,7 +7935,7 @@ void FSlateEditorStyle::FStyle::SetupAutomationStyles()
 			FTextBlockStyle(NormalText)
 			.SetFont(DEFAULT_FONT("Regular", 12))
 			.SetColorAndOpacity(FSlateColor::UseForeground())
-			.SetShadowOffset(FVector2D(1, 1))
+			.SetShadowOffset(FVector2f::UnitVector)
 			.SetShadowColorAndOpacity(FLinearColor::Black);
 
 		Set("FBXLargeFont", FBXLargeFont);
@@ -7944,7 +7944,7 @@ void FSlateEditorStyle::FStyle::SetupAutomationStyles()
 			FTextBlockStyle(NormalText)
 			.SetFont(DEFAULT_FONT("Regular", 11))
 			.SetColorAndOpacity(FSlateColor::UseForeground())
-			.SetShadowOffset(FVector2D(1, 1))
+			.SetShadowOffset(FVector2f::UnitVector)
 			.SetShadowColorAndOpacity(FLinearColor::Black);
 
 		Set("FBXMediumFont", FBXMediumFont);
@@ -7953,7 +7953,7 @@ void FSlateEditorStyle::FStyle::SetupAutomationStyles()
 			FTextBlockStyle(NormalText)
 			.SetFont(DEFAULT_FONT("Regular", 10))
 			.SetColorAndOpacity(FSlateColor::UseForeground())
-			.SetShadowOffset(FVector2D(1, 1))
+			.SetShadowOffset(FVector2f::UnitVector)
 			.SetShadowColorAndOpacity(FLinearColor::Black);
 
 		Set("FBXSmallFont", FBXSmallFont);
@@ -7961,7 +7961,7 @@ void FSlateEditorStyle::FStyle::SetupAutomationStyles()
 
 	// Asset Dialog
 	{
-		Set("AssetDialog.ErrorLabelBorder", new FSlateColorBrush(FLinearColor(0.2, 0, 0)));
+		Set("AssetDialog.ErrorLabelBorder", new FSlateColorBrush(FLinearColor(0.2f, 0.f, 0.f)));
 		Set("AssetDialog.ErrorLabelFont", FTextBlockStyle(NormalText).SetFont(DEFAULT_FONT("Regular", 10)));
 	}
 }
@@ -7998,39 +7998,39 @@ void FSlateEditorStyle::FStyle::SetupUMGEditorStyles()
 	Set("UMGEditor.ResizeAreaHandle", new IMAGE_BRUSH("Icons/UMG/ResizeAreaHandle", Icon20x20));
 
 	Set("UMGEditor.AnchorGizmo.Center", new IMAGE_BRUSH("Icons/UMG/AnchorGizmo/center", Icon16x16));
-	Set("UMGEditor.AnchorGizmo.Center.Hovered", new IMAGE_BRUSH("Icons/UMG/AnchorGizmo/center", Icon16x16, FLinearColor(0, 1, 0)));
+	Set("UMGEditor.AnchorGizmo.Center.Hovered", new IMAGE_BRUSH("Icons/UMG/AnchorGizmo/center", Icon16x16, FLinearColor::Green));
 	
-	Set("UMGEditor.AnchorGizmo.Left", new IMAGE_BRUSH("Icons/UMG/AnchorGizmo/left", FVector2D(32, 16)));
-	Set("UMGEditor.AnchorGizmo.Left.Hovered", new IMAGE_BRUSH("Icons/UMG/AnchorGizmo/left", FVector2D(32, 16), FLinearColor(0, 1, 0)));
-	Set("UMGEditor.AnchorGizmo.Right", new IMAGE_BRUSH("Icons/UMG/AnchorGizmo/right", FVector2D(32, 16)));
-	Set("UMGEditor.AnchorGizmo.Right.Hovered", new IMAGE_BRUSH("Icons/UMG/AnchorGizmo/right", FVector2D(32, 16), FLinearColor(0, 1, 0)));
+	Set("UMGEditor.AnchorGizmo.Left", new IMAGE_BRUSH("Icons/UMG/AnchorGizmo/left", FVector2f(32.f, 16.f)));
+	Set("UMGEditor.AnchorGizmo.Left.Hovered", new IMAGE_BRUSH("Icons/UMG/AnchorGizmo/left", FVector2f(32.f, 16.f), FLinearColor::Green));
+	Set("UMGEditor.AnchorGizmo.Right", new IMAGE_BRUSH("Icons/UMG/AnchorGizmo/right", FVector2f(32.f, 16.f)));
+	Set("UMGEditor.AnchorGizmo.Right.Hovered", new IMAGE_BRUSH("Icons/UMG/AnchorGizmo/right", FVector2f(32.f, 16.f), FLinearColor::Green));
 	
-	Set("UMGEditor.AnchorGizmo.Top", new IMAGE_BRUSH("Icons/UMG/AnchorGizmo/top", FVector2D(16, 32)));
-	Set("UMGEditor.AnchorGizmo.Top.Hovered", new IMAGE_BRUSH("Icons/UMG/AnchorGizmo/top", FVector2D(16, 32), FLinearColor(0, 1, 0)));
-	Set("UMGEditor.AnchorGizmo.Bottom", new IMAGE_BRUSH("Icons/UMG/AnchorGizmo/bottom", FVector2D(16, 32)));
-	Set("UMGEditor.AnchorGizmo.Bottom.Hovered", new IMAGE_BRUSH("Icons/UMG/AnchorGizmo/bottom", FVector2D(16, 32), FLinearColor(0, 1, 0)));
+	Set("UMGEditor.AnchorGizmo.Top", new IMAGE_BRUSH("Icons/UMG/AnchorGizmo/top", FVector2f(16.f, 32.f)));
+	Set("UMGEditor.AnchorGizmo.Top.Hovered", new IMAGE_BRUSH("Icons/UMG/AnchorGizmo/top", FVector2f(16.f, 32.f), FLinearColor::Green));
+	Set("UMGEditor.AnchorGizmo.Bottom", new IMAGE_BRUSH("Icons/UMG/AnchorGizmo/bottom", FVector2f(16.f, 32.f)));
+	Set("UMGEditor.AnchorGizmo.Bottom.Hovered", new IMAGE_BRUSH("Icons/UMG/AnchorGizmo/bottom", FVector2f(16.f, 32.f), FLinearColor::Green));
 
-	Set("UMGEditor.AnchorGizmo.TopLeft", new IMAGE_BRUSH("Icons/UMG/AnchorGizmo/topleft", FVector2D(24, 24)));
-	Set("UMGEditor.AnchorGizmo.TopLeft.Hovered", new IMAGE_BRUSH("Icons/UMG/AnchorGizmo/topleft", FVector2D(24, 24), FLinearColor(0, 1, 0)));
+	Set("UMGEditor.AnchorGizmo.TopLeft", new IMAGE_BRUSH("Icons/UMG/AnchorGizmo/topleft", Icon24x24));
+	Set("UMGEditor.AnchorGizmo.TopLeft.Hovered", new IMAGE_BRUSH("Icons/UMG/AnchorGizmo/topleft", Icon24x24, FLinearColor::Green));
 
-	Set("UMGEditor.AnchorGizmo.TopRight", new IMAGE_BRUSH("Icons/UMG/AnchorGizmo/topright", FVector2D(24, 24)));
-	Set("UMGEditor.AnchorGizmo.TopRight.Hovered", new IMAGE_BRUSH("Icons/UMG/AnchorGizmo/topright", FVector2D(24, 24), FLinearColor(0, 1, 0)));
+	Set("UMGEditor.AnchorGizmo.TopRight", new IMAGE_BRUSH("Icons/UMG/AnchorGizmo/topright", Icon24x24));
+	Set("UMGEditor.AnchorGizmo.TopRight.Hovered", new IMAGE_BRUSH("Icons/UMG/AnchorGizmo/topright", Icon24x24, FLinearColor::Green));
 
-	Set("UMGEditor.AnchorGizmo.BottomLeft", new IMAGE_BRUSH("Icons/UMG/AnchorGizmo/bottomleft", FVector2D(24, 24)));
-	Set("UMGEditor.AnchorGizmo.BottomLeft.Hovered", new IMAGE_BRUSH("Icons/UMG/AnchorGizmo/bottomleft", FVector2D(24, 24), FLinearColor(0, 1, 0)));
+	Set("UMGEditor.AnchorGizmo.BottomLeft", new IMAGE_BRUSH("Icons/UMG/AnchorGizmo/bottomleft", Icon24x24));
+	Set("UMGEditor.AnchorGizmo.BottomLeft.Hovered", new IMAGE_BRUSH("Icons/UMG/AnchorGizmo/bottomleft", Icon24x24, FLinearColor::Green));
 
-	Set("UMGEditor.AnchorGizmo.BottomRight", new IMAGE_BRUSH("Icons/UMG/AnchorGizmo/bottomright", FVector2D(24, 24)));
-	Set("UMGEditor.AnchorGizmo.BottomRight.Hovered", new IMAGE_BRUSH("Icons/UMG/AnchorGizmo/bottomright", FVector2D(24, 24), FLinearColor(0, 1, 0)));
+	Set("UMGEditor.AnchorGizmo.BottomRight", new IMAGE_BRUSH("Icons/UMG/AnchorGizmo/bottomright", Icon24x24));
+	Set("UMGEditor.AnchorGizmo.BottomRight.Hovered", new IMAGE_BRUSH("Icons/UMG/AnchorGizmo/bottomright", Icon24x24, FLinearColor::Green));
 
 
-	Set("UMGEditor.AnchoredWidget", new BOX_BRUSH("Common/Button", FVector2D(32, 32), 8.0f / 32.0f));
+	Set("UMGEditor.AnchoredWidget", new BOX_BRUSH("Common/Button", Icon32x32, 8.0f / 32.0f));
 	Set("UMGEditor.AnchoredWidgetAlignment", new IMAGE_BRUSH("Icons/icon_tab_DeviceManager_16x", Icon8x8));
 	
 	Set("UMGEditor.PaletteHeader", FTableRowStyle()
-		.SetEvenRowBackgroundBrush( BOX_BRUSH( "PropertyView/DetailCategoryMiddle", FMargin( 4/16.0f, 8.0f/16.0f, 4/16.0f, 4/16.0f ), FLinearColor(0.6f,0.6f,0.6f,1.0f) ) )
-		.SetEvenRowBackgroundHoveredBrush( BOX_BRUSH( "PropertyView/DetailCategoryMiddle", FMargin( 4/16.0f, 8.0f/16.0f, 4/16.0f, 4/16.0f ), FLinearColor(0.3f,0.3f,0.3f,1.0f) ) )
-		.SetOddRowBackgroundBrush( BOX_BRUSH( "PropertyView/DetailCategoryMiddle", FMargin( 4/16.0f, 8.0f/16.0f, 4/16.0f, 4/16.0f ), FLinearColor(0.6f,0.6f,0.6f,1.0f) ) )
-		.SetOddRowBackgroundHoveredBrush( BOX_BRUSH( "PropertyView/DetailCategoryMiddle", FMargin( 4/16.0f, 8.0f/16.0f, 4/16.0f, 4/16.0f ), FLinearColor(0.3f,0.3f,0.3f,1.0f) ) )
+		.SetEvenRowBackgroundBrush( BOX_BRUSH( "PropertyView/DetailCategoryMiddle", FMargin( 4.f/16.0f, 8.0f/16.0f, 4.f/16.0f, 4.f/16.0f ), FLinearColor(0.6f,0.6f,0.6f,1.0f) ) )
+		.SetEvenRowBackgroundHoveredBrush( BOX_BRUSH( "PropertyView/DetailCategoryMiddle", FMargin( 4.f/16.0f, 8.0f/16.0f, 4.f/16.0f, 4.f/16.0f ), FLinearColor(0.3f,0.3f,0.3f,1.0f) ) )
+		.SetOddRowBackgroundBrush( BOX_BRUSH( "PropertyView/DetailCategoryMiddle", FMargin( 4.f/16.0f, 8.0f/16.0f, 4.f/16.0f, 4.f/16.0f ), FLinearColor(0.6f,0.6f,0.6f,1.0f) ) )
+		.SetOddRowBackgroundHoveredBrush( BOX_BRUSH( "PropertyView/DetailCategoryMiddle", FMargin( 4.f/16.0f, 8.0f/16.0f, 4.f/16.0f, 4.f/16.0f ), FLinearColor(0.3f,0.3f,0.3f,1.0f) ) )
 		.SetSelectorFocusedBrush( BORDER_BRUSH( "Common/Selector", FMargin(4.f/16.f), SelectorColor ) )
 		.SetActiveBrush( IMAGE_BRUSH( "Common/Selection", Icon8x8, SelectionColor ) )
 		.SetActiveHoveredBrush( IMAGE_BRUSH( "Common/Selection", Icon8x8, SelectionColor ) )
@@ -8041,8 +8041,8 @@ void FSlateEditorStyle::FStyle::SetupUMGEditorStyles()
 		);
 
 	Set("UMGEditor.PaletteItem", FTableRowStyle(NormalTableRowStyle)
-		.SetEvenRowBackgroundBrush(BOX_BRUSH("PropertyView/DetailCategoryMiddle", FMargin(4 / 16.0f, 8.0f / 16.0f, 4 / 16.0f, 4 / 16.0f)))
-		.SetOddRowBackgroundBrush(BOX_BRUSH("PropertyView/DetailCategoryMiddle", FMargin(4 / 16.0f, 8.0f / 16.0f, 4 / 16.0f, 4 / 16.0f)))
+		.SetEvenRowBackgroundBrush(BOX_BRUSH("PropertyView/DetailCategoryMiddle", FMargin(4.f / 16.0f, 8.0f / 16.0f, 4.f / 16.0f, 4.f / 16.0f)))
+		.SetOddRowBackgroundBrush(BOX_BRUSH("PropertyView/DetailCategoryMiddle", FMargin(4.f / 16.0f, 8.0f / 16.0f, 4.f / 16.0f, 4 / 16.0f)))
 		);
 
 	// Style of the favorite toggle
@@ -8070,11 +8070,11 @@ void FSlateEditorStyle::FStyle::SetupUMGEditorStyles()
 		FTextBlockStyle(NormalText)
 		.SetFont(DEFAULT_FONT("Regular", 18))
 		.SetColorAndOpacity(FSlateColor::UseForeground())
-		.SetShadowOffset(FVector2D(1, 1))
+		.SetShadowOffset(FVector2f::UnitVector)
 		.SetShadowColorAndOpacity(FLinearColor::Black);
 
 
-	Set("UMGEditor.AddAnimationIcon", new IMAGE_BRUSH("Icons/PlusSymbol_12x", Icon12x12, FLinearColor(.05,.05,.05) ) );
+	Set("UMGEditor.AddAnimationIcon", new IMAGE_BRUSH("Icons/PlusSymbol_12x", Icon12x12, FLinearColor(.05f,.05f,.05f) ) );
 	Set("UMGEditor.NoAnimationFont", NoAnimationFont);
 
 	Set("UMGEditor.SwitchToDesigner", new IMAGE_BRUSH("UMG/Designer_40x", Icon40x40));
@@ -8172,7 +8172,7 @@ void FSlateEditorStyle::FStyle::SetupLocalizationDashboardStyles()
 
 	Set( "LocalizationDashboard.CommandletLog.Text", FTextBlockStyle(NormalText)
 		.SetFont( DEFAULT_FONT( "Regular", 8 ) )
-		.SetShadowOffset( FVector2D::ZeroVector )
+		.SetShadowOffset( FVector2f::ZeroVector )
 		);
 }
 
