[Truncation Warnings] Fix and enable truncation warnings for EditorConfig and EditorStyle and EditorWidgets modules

diff --git a/Engine/Source/Editor/EditorStyle/Private/StarshipStyle.cpp b/Engine/Source/Editor/EditorStyle/Private/StarshipStyle.cpp
--- a/Engine/Source/Editor/EditorStyle/Private/StarshipStyle.cpp
+++ b/Engine/Source/Editor/EditorStyle/Private/StarshipStyle.cpp
@@ -76,36 +76,36 @@ bool FStarshipEditorStyle::FStyle::IncludeEditorSpecificStyles()
 /* FStarshipEditorStyle interface
  *****************************************************************************/
 
+ // Note, these sizes are in Slate Units.
+ // Slate Units do NOT have to map to pixels.
+const FVector2f FStarshipEditorStyle::FStyle::Icon7x16 = FVector2f(7.0f, 16.0f);
+const FVector2f FStarshipEditorStyle::FStyle::Icon8x4 = FVector2f(8.0f, 4.0f);
+const FVector2f FStarshipEditorStyle::FStyle::Icon16x4 = FVector2f(16.0f, 4.0f);
+const FVector2f FStarshipEditorStyle::FStyle::Icon8x8 = FVector2f(8.0f, 8.0f);
+const FVector2f FStarshipEditorStyle::FStyle::Icon10x10 = FVector2f(10.0f, 10.0f);
+const FVector2f FStarshipEditorStyle::FStyle::Icon12x12 = FVector2f(12.0f, 12.0f);
+const FVector2f FStarshipEditorStyle::FStyle::Icon12x16 = FVector2f(12.0f, 16.0f);
+const FVector2f FStarshipEditorStyle::FStyle::Icon14x14 = FVector2f(14.0f, 14.0f);
+const FVector2f FStarshipEditorStyle::FStyle::Icon16x16 = FVector2f(16.0f, 16.0f);
+const FVector2f FStarshipEditorStyle::FStyle::Icon16x20 = FVector2f(16.0f, 20.0f);
+const FVector2f FStarshipEditorStyle::FStyle::Icon20x20 = FVector2f(20.0f, 20.0f);
+const FVector2f FStarshipEditorStyle::FStyle::Icon22x22 = FVector2f(22.0f, 22.0f);
+const FVector2f FStarshipEditorStyle::FStyle::Icon24x24 = FVector2f(24.0f, 24.0f);
+const FVector2f FStarshipEditorStyle::FStyle::Icon25x25 = FVector2f(25.0f, 25.0f);
+const FVector2f FStarshipEditorStyle::FStyle::Icon32x32 = FVector2f(32.0f, 32.0f);
+const FVector2f FStarshipEditorStyle::FStyle::Icon40x40 = FVector2f(40.0f, 40.0f);
+const FVector2f FStarshipEditorStyle::FStyle::Icon48x48 = FVector2f(48.0f, 48.0f);
+const FVector2f FStarshipEditorStyle::FStyle::Icon64x64 = FVector2f(64.0f, 64.0f);
+const FVector2f FStarshipEditorStyle::FStyle::Icon36x24 = FVector2f(36.0f, 24.0f);
+const FVector2f FStarshipEditorStyle::FStyle::Icon128x128 = FVector2f(128.0f, 128.0f);
+
 FStarshipEditorStyle::FStyle::FStyle()
 	: FSlateStyleSet(FStarshipEditorStyle::StyleSetName)
 
-	// Note, these sizes are in Slate Units.
-	// Slate Units do NOT have to map to pixels.
-	, Icon7x16(7.0f, 16.0f)
-	, Icon8x4(8.0f, 4.0f)
-	, Icon16x4(16.0f, 4.0f)
-	, Icon8x8(8.0f, 8.0f)
-	, Icon10x10(10.0f, 10.0f)
-	, Icon12x12(12.0f, 12.0f)
-	, Icon12x16(12.0f, 16.0f)
-	, Icon14x14(14.0f, 14.0f)
-	, Icon16x16(16.0f, 16.0f)
-	, Icon16x20(16.0f, 20.0f)
-	, Icon20x20(20.0f, 20.0f)
-	, Icon22x22(22.0f, 22.0f)
-	, Icon24x24(24.0f, 24.0f)
-	, Icon25x25(25.0f, 25.0f)
-	, Icon32x32(32.0f, 32.0f)
-	, Icon40x40(40.0f, 40.0f)
-	, Icon48x48(48.0f, 48.0f)
-	, Icon64x64(64.0f, 64.0f)
-	, Icon36x24(36.0f, 24.0f)
-	, Icon128x128(128.0f, 128.0f)
-
 	// These are the colors that are updated by the user style customizations
 	, SelectionColor_Subdued_LinearRef(MakeShareable(new FLinearColor(0.807f, 0.596f, 0.388f)))
 	, HighlightColor_LinearRef( MakeShareable( new FLinearColor(0.068f, 0.068f, 0.068f) ) ) 
-	, WindowHighlightColor_LinearRef(MakeShareable(new FLinearColor(0,0,0,0)))
+	, WindowHighlightColor_LinearRef(MakeShareable(new FLinearColor(0.f,0.f,0.f,0.f)))
 
 
 
@@ -294,12 +294,12 @@ void FStarshipEditorStyle::FStyle::SetupGeneralStyles()
 		Set( "TextBlock.HighlighColor", FLinearColor( 0.02f, 0.3f, 0.0f ) );
 
 		Set("TextBlock.ShadowedText", FTextBlockStyle(NormalText)
-			.SetShadowOffset(FVector2D(1.0f, 1.0f))
+			.SetShadowOffset(FVector2f::UnitVector)
 			.SetShadowColorAndOpacity(FLinearColor(0.0f, 0.0f, 0.0f)));
 
 		Set("TextBlock.ShadowedTextWarning", FTextBlockStyle(NormalText)
 			.SetColorAndOpacity(FStyleColors::Warning)
-			.SetShadowOffset(FVector2D(1.0f, 1.0f))
+			.SetShadowOffset(FVector2f::UnitVector)
 			.SetShadowColorAndOpacity(FLinearColor(0.0f, 0.0f, 0.0f)));
 
 		Set("NormalText.Subdued", FTextBlockStyle(NormalText)
@@ -309,7 +309,7 @@ void FStarshipEditorStyle::FStyle::SetupGeneralStyles()
 			.SetFont(DEFAULT_FONT("Bold", FStarshipCoreStyle::RegularTextSize))
 			.SetColorAndOpacity(FLinearColor(1.0f, 1.0f, 1.0f))
 			.SetHighlightColor(FLinearColor(1.0f, 1.0f, 1.0f))
-			.SetShadowOffset(FVector2D(1, 1))
+			.SetShadowOffset(FVector2f::UnitVector)
 			.SetShadowColorAndOpacity(FLinearColor(0, 0, 0, 0.9f)));
 
 		Set("SmallText.Subdued", FTextBlockStyle(NormalText)
@@ -327,7 +327,7 @@ void FStarshipEditorStyle::FStyle::SetupGeneralStyles()
 			.SetFont(DEFAULT_FONT("Bold", 11))
 			.SetColorAndOpacity(FLinearColor(1.0f, 1.0f, 1.0f))
 			.SetHighlightColor(FLinearColor(1.0f, 1.0f, 1.0f))
-			.SetShadowOffset(FVector2D(1, 1))
+			.SetShadowOffset(FVector2f::UnitVector)
 			.SetShadowColorAndOpacity(FLinearColor(0, 0, 0, 0.9f)));
 	}
 
@@ -362,36 +362,36 @@ void FStarshipEditorStyle::FStyle::SetupGeneralStyles()
 
 
 
-	Set( "WideDash.Horizontal", new CORE_IMAGE_BRUSH("Starship/Common/Dash_Horizontal", FVector2D(10, 1), FLinearColor::White, ESlateBrushTileType::Horizontal));
-	Set( "WideDash.Vertical", new CORE_IMAGE_BRUSH("Starship/Common/Dash_Vertical", FVector2D(1, 10), FLinearColor::White, ESlateBrushTileType::Vertical));
+	Set( "WideDash.Horizontal", new CORE_IMAGE_BRUSH("Starship/Common/Dash_Horizontal", FVector2f(10.f, 1.f), FLinearColor::White, ESlateBrushTileType::Horizontal));
+	Set( "WideDash.Vertical", new CORE_IMAGE_BRUSH("Starship/Common/Dash_Vertical", FVector2f(1.f, 10.f), FLinearColor::White, ESlateBrushTileType::Vertical));
 
 	Set("DropTarget.Background", new CORE_BOX_BRUSH("Starship/Common/DropTargetBackground", FMargin(6.0f / 64.0f)));
 
-	Set("ThinLine.Horizontal", new IMAGE_BRUSH("Common/ThinLine_Horizontal", FVector2D(11, 2), FLinearColor::White, ESlateBrushTileType::Horizontal));
+	Set("ThinLine.Horizontal", new IMAGE_BRUSH("Common/ThinLine_Horizontal", FVector2f(11.f, 2.f), FLinearColor::White, ESlateBrushTileType::Horizontal));
 
 
 	// Buttons that only provide a hover hint.
 	HoverHintOnly = FButtonStyle()
 			.SetNormal( FSlateNoResource() )
-			.SetHovered( BOX_BRUSH( "Common/ButtonHoverHint", FMargin(4/16.0f), FLinearColor(1,1,1,0.15f) ) )
-			.SetPressed( BOX_BRUSH( "Common/ButtonHoverHint", FMargin(4/16.0f), FLinearColor(1,1,1,0.25f) ) )
-			.SetNormalPadding( FMargin(0,0,0,1) )
-			.SetPressedPadding( FMargin(0,1,0,0) );
+			.SetHovered( BOX_BRUSH( "Common/ButtonHoverHint", FMargin(4.f/16.0f), FLinearColor(1.f,1.f,1.f,0.15f) ) )
+			.SetPressed( BOX_BRUSH( "Common/ButtonHoverHint", FMargin(4.f/16.0f), FLinearColor(1.f,1.f,1.f,0.25f) ) )
+			.SetNormalPadding( FMargin(0.f,0.f,0.f,1.f) )
+			.SetPressedPadding( FMargin(0.f,1.f,0.f,0.f) );
 	Set( "HoverHintOnly", HoverHintOnly );
 
 
 	FButtonStyle SimpleSharpButton = FButtonStyle()
-		.SetNormal(BOX_BRUSH("Common/Button/simple_sharp_normal", FMargin(4 / 16.0f), FLinearColor(1, 1, 1, 1)))
-		.SetHovered(BOX_BRUSH("Common/Button/simple_sharp_hovered", FMargin(4 / 16.0f), FLinearColor(1, 1, 1, 1)))
-		.SetPressed(BOX_BRUSH("Common/Button/simple_sharp_hovered", FMargin(4 / 16.0f), FLinearColor(1, 1, 1, 1)))
-		.SetNormalPadding(FMargin(0, 0, 0, 1))
-		.SetPressedPadding(FMargin(0, 1, 0, 0));
+		.SetNormal(BOX_BRUSH("Common/Button/simple_sharp_normal", FMargin(4.f / 16.0f), FLinearColor::White))
+		.SetHovered(BOX_BRUSH("Common/Button/simple_sharp_hovered", FMargin(4.f / 16.0f), FLinearColor::White))
+		.SetPressed(BOX_BRUSH("Common/Button/simple_sharp_hovered", FMargin(4.f / 16.0f), FLinearColor::White))
+		.SetNormalPadding(FMargin(0.f, 0.f, 0.f, 1.f))
+		.SetPressedPadding(FMargin(0.f, 1.f, 0.f, 0.f));
 	Set("SimpleSharpButton", SimpleSharpButton);
 
 	FButtonStyle SimpleRoundButton = FButtonStyle()
-		.SetNormal(BOX_BRUSH("Common/Button/simple_round_normal", FMargin(4 / 16.0f), FLinearColor(1, 1, 1, 1)))
-		.SetHovered(BOX_BRUSH("Common/Button/simple_round_hovered", FMargin(4 / 16.0f), FLinearColor(1, 1, 1, 1)))
-		.SetPressed(BOX_BRUSH("Common/Button/simple_round_hovered", FMargin(4 / 16.0f), FLinearColor(1, 1, 1, 1)))
+		.SetNormal(BOX_BRUSH("Common/Button/simple_round_normal", FMargin(4.f / 16.0f), FLinearColor::White))
+		.SetHovered(BOX_BRUSH("Common/Button/simple_round_hovered", FMargin(4.f / 16.0f), FLinearColor::White))
+		.SetPressed(BOX_BRUSH("Common/Button/simple_round_hovered", FMargin(4.f / 16.0f), FLinearColor::White))
 		.SetNormalPadding(FMargin(0, 0, 0, 1))
 		.SetPressedPadding(FMargin(0, 1, 0, 0));
 	Set("SimpleRoundButton", SimpleRoundButton);
@@ -400,7 +400,7 @@ void FStarshipEditorStyle::FStyle::SetupGeneralStyles()
 	{
 		Set( "Symbols.SearchGlass", new IMAGE_BRUSH( "Common/SearchGlass", Icon16x16 ) );
 		Set( "Symbols.X", new IMAGE_BRUSH( "Common/X", Icon16x16 ) );
-		Set( "Symbols.VerticalPipe", new BOX_BRUSH( "Common/VerticalPipe", FMargin(0) ) );
+		Set( "Symbols.VerticalPipe", new BOX_BRUSH( "Common/VerticalPipe", FMargin(0.f) ) );
 		Set( "Symbols.UpArrow", new IMAGE_BRUSH( "Common/UpArrow", Icon8x8 ) );
 		Set( "Symbols.DoubleUpArrow", new IMAGE_BRUSH( "Common/UpArrow2", Icon8x8 ) );
 		Set( "Symbols.DownArrow", new IMAGE_BRUSH( "Common/DownArrow", Icon8x8 ) );
@@ -466,7 +466,7 @@ void FStarshipEditorStyle::FStyle::SetupGeneralStyles()
 		Set("Icons.Scalability", new IMAGE_BRUSH_SVG("Starship/Common/Scalability_16", Icon20x20));
 		Set("Icons.ViewportScalability", new IMAGE_BRUSH_SVG("Starship/Common/ViewportScalability_16", Icon16x16));
 		Set("Icons.EllipsisVerticalNarrow",
-			new CORE_IMAGE_BRUSH_SVG("Starship/Common/ellipsis-vertical-narrow", FVector2D(6, 24)));
+			new CORE_IMAGE_BRUSH_SVG("Starship/Common/ellipsis-vertical-narrow", FVector2f(6.f, 24.f)));
 
 		Set("Icons.Toolbar.Play", new IMAGE_BRUSH_SVG("Starship/Common/play", Icon20x20));
 		Set("Icons.Toolbar.Pause", new IMAGE_BRUSH_SVG("Starship/MainToolbar/pause", Icon20x20));
@@ -502,9 +502,9 @@ void FStarshipEditorStyle::FStyle::SetupGeneralStyles()
 		Set("ReferenceViewer.IndirectReference", new IMAGE_BRUSH_SVG("Starship/ReferenceViewer/IndirectRef", Icon16x16));
 	}
 
-	Set("UnrealDefaultThumbnail", new IMAGE_BRUSH("Starship/Common/Unreal_DefaultThumbnail", FVector2D(256, 256)));
+	Set("UnrealDefaultThumbnail", new IMAGE_BRUSH("Starship/Common/Unreal_DefaultThumbnail", FVector2f(256.f, 256.f)));
 
-	Set( "WarningStripe", new IMAGE_BRUSH( "Common/WarningStripe", FVector2D(20,6), FLinearColor::White, ESlateBrushTileType::Horizontal ) );
+	Set( "WarningStripe", new IMAGE_BRUSH( "Common/WarningStripe", FVector2f(20.f,6.f), FLinearColor::White, ESlateBrushTileType::Horizontal ) );
 	
 	Set("RoundedWarning", new FSlateRoundedBoxBrush(FStyleColors::Transparent, 4.0f, FStyleColors::Warning, 1.0f));
 	Set("RoundedError", new FSlateRoundedBoxBrush(FStyleColors::Transparent, 4.0f, FStyleColors::Error, 1.0f));
@@ -523,7 +523,7 @@ void FStarshipEditorStyle::FStyle::SetupGeneralStyles()
 		//FSlateColorBrush(FLinearColor::White)
 
 		Set("RoundButton", FButtonStyle(Button)
-			.SetNormal(BOX_BRUSH("Common/RoundedSelection_16x", 4.0f / 16.0f, FLinearColor(1, 1, 1, 0.1f)))
+			.SetNormal(BOX_BRUSH("Common/RoundedSelection_16x", 4.0f / 16.0f, FLinearColor(1.f, 1.f, 1.f, 0.1f)))
 			.SetHovered(BOX_BRUSH("Common/RoundedSelection_16x", 4.0f / 16.0f, SelectionColor))
 			.SetPressed(BOX_BRUSH("Common/RoundedSelection_16x", 4.0f / 16.0f, SelectionColor_Pressed))
 			);
@@ -553,8 +553,8 @@ void FStarshipEditorStyle::FStyle::SetupGeneralStyles()
 			.SetFont(DEFAULT_FONT("Bold", 10))
 			.SetColorAndOpacity(FLinearColor(1.0f, 1.0f, 1.0f))
 			.SetHighlightColor(FLinearColor(1.0f, 1.0f, 1.0f))
-			.SetShadowOffset(FVector2D(1, 1))
-			.SetShadowColorAndOpacity(FLinearColor(0, 0, 0, 0.9f)));
+			.SetShadowOffset(FVector2f::UnitVector)
+			.SetShadowColorAndOpacity(FLinearColor(0.f, 0.f, 0.f, 0.9f)));
 
 		struct ButtonColor
 		{
@@ -576,11 +576,11 @@ void FStarshipEditorStyle::FStyle::SetupGeneralStyles()
 		};
 
 		TArray< ButtonColor > FlatButtons;
-		FlatButtons.Add(ButtonColor("FlatButton.Primary", FLinearColor(0.02899, 0.19752, 0.48195)));
-		FlatButtons.Add(ButtonColor("FlatButton.Success", FLinearColor(0.10616, 0.48777, 0.10616)));
-		FlatButtons.Add(ButtonColor("FlatButton.Info", FLinearColor(0.10363, 0.53564, 0.7372)));
-		FlatButtons.Add(ButtonColor("FlatButton.Warning", FLinearColor(0.87514, 0.42591, 0.07383)));
-		FlatButtons.Add(ButtonColor("FlatButton.Danger", FLinearColor(0.70117, 0.08464, 0.07593)));
+		FlatButtons.Add(ButtonColor("FlatButton.Primary", FLinearColor(0.02899f, 0.19752f, 0.48195f)));
+		FlatButtons.Add(ButtonColor("FlatButton.Success", FLinearColor(0.10616f, 0.48777f, 0.10616f)));
+		FlatButtons.Add(ButtonColor("FlatButton.Info", FLinearColor(0.10363f, 0.53564f, 0.7372f)));
+		FlatButtons.Add(ButtonColor("FlatButton.Warning", FLinearColor(0.87514f, 0.42591f, 0.07383f)));
+		FlatButtons.Add(ButtonColor("FlatButton.Danger", FLinearColor(0.70117f, 0.08464f, 0.07593f)));
 
 		for ( const ButtonColor& Entry : FlatButtons )
 		{
@@ -640,8 +640,8 @@ void FStarshipEditorStyle::FStyle::SetupGeneralStyles()
 		Set("GenericFilters.TextStyle", FTextBlockStyle(NormalText)
 			.SetFont(DEFAULT_FONT("Bold", 9))
 			.SetColorAndOpacity(FLinearColor(1.0f, 1.0f, 1.0f, 0.9f))
-			.SetShadowOffset(FVector2D(1, 1))
-			.SetShadowColorAndOpacity(FLinearColor(0, 0, 0, 0.9f)));
+			.SetShadowOffset(FVector2f::UnitVector)
+			.SetShadowColorAndOpacity(FLinearColor(0.f, 0.f, 0.f, 0.9f)));
 
 	}
 
@@ -681,13 +681,13 @@ void FStarshipEditorStyle::FStyle::SetupGeneralStyles()
 	{
 		const FLinearColor ActiveFilterColor = FLinearColor(1.0f,0.55f,0.0f,1.0f);
 		Set("Searching.SearchActiveTab",    new FSlateNoResource());
-		Set("Searching.SearchActiveBorder", new FSlateRoundedBoxBrush(FLinearColor::Transparent, 0.0, FStyleColors::Primary, 1.f));
+		Set("Searching.SearchActiveBorder", new FSlateRoundedBoxBrush(FLinearColor::Transparent, 0.0f, FStyleColors::Primary, 1.f));
 	}
 
 
 
 	// Images sizes are specified in Slate Screen Units. These do not necessarily correspond to pixels!
-	// An IMAGE_BRUSH( "SomeImage", FVector2D(32,32)) will have a desired size of 16x16 Slate Screen Units.
+	// An IMAGE_BRUSH( "SomeImage", FVector2f(32,32)) will have a desired size of 16x16 Slate Screen Units.
 	// This allows the original resource to be scaled up or down as needed.
 
 	Set( "WhiteTexture", new IMAGE_BRUSH( "Old/White", Icon16x16 ) );
@@ -720,9 +720,9 @@ void FStarshipEditorStyle::FStyle::SetupGeneralStyles()
 	{
 		// Dark Hyperlink - for use on light backgrounds
 		FButtonStyle DarkHyperlinkButton = FButtonStyle()
-			.SetNormal ( BORDER_BRUSH( "Old/HyperlinkDotted", FMargin(0,0,0,3/16.0f), FLinearColor::Black ) )
+			.SetNormal ( BORDER_BRUSH( "Old/HyperlinkDotted", FMargin(0.f,0.f,0.f,3.f/16.0f), FLinearColor::Black ) )
 			.SetPressed( FSlateNoResource() )
-			.SetHovered( BORDER_BRUSH( "Old/HyperlinkUnderline", FMargin(0,0,0,3/16.0f), FLinearColor::Black ) );
+			.SetHovered( BORDER_BRUSH( "Old/HyperlinkUnderline", FMargin(0.f,0.f,0.f,3.f/16.0f), FLinearColor::Black ) );
 		FHyperlinkStyle DarkHyperlink = FHyperlinkStyle()
 			.SetUnderlineStyle(DarkHyperlinkButton)
 			.SetTextStyle(NormalText)
@@ -733,7 +733,7 @@ void FStarshipEditorStyle::FStyle::SetupGeneralStyles()
 		FButtonStyle HoverOnlyHyperlinkButton = FButtonStyle()
 			.SetNormal(FSlateNoResource() )
 			.SetPressed(FSlateNoResource() )
-			.SetHovered(BORDER_BRUSH( "Old/HyperlinkUnderline", FMargin(0,0,0,3/16.0f) ) );
+			.SetHovered(BORDER_BRUSH( "Old/HyperlinkUnderline", FMargin(0.f,0.f,0.f,3.f/16.0f) ) );
 		Set("HoverOnlyHyperlinkButton", HoverOnlyHyperlinkButton);
 
 		FHyperlinkStyle HoverOnlyHyperlink = FHyperlinkStyle()
@@ -790,13 +790,13 @@ void FStarshipEditorStyle::FStyle::SetupGeneralStyles()
 	{
 		Set( "MaterialList.DragDropBorder", new BOX_BRUSH( "Old/Window/ViewportDebugBorder", 0.8f ) );
 		Set( "MaterialList.HyperlinkStyle", FTextBlockStyle(NormalText) .SetFont( DEFAULT_FONT( "Regular", 8 ) ) );
-		Set( "MaterialList.HyperlinkStyle.ShadowOffset", FVector2D::ZeroVector );
+		Set( "MaterialList.HyperlinkStyle.ShadowOffset", FVector2f::ZeroVector );
 		Set( "Icons.NaniteBrowseContent", new IMAGE_BRUSH_SVG("Starship/Common/NaniteBrowseContent", Icon16x16));
 	}
 
 	// Dialogue Wave Details
 	{
-		Set( "DialogueWaveDetails.SpeakerToTarget", new IMAGE_BRUSH( "PropertyView/SpeakerToTarget", FVector2D(30.0f, 30.0f) ) );
+		Set( "DialogueWaveDetails.SpeakerToTarget", new IMAGE_BRUSH( "PropertyView/SpeakerToTarget", FVector2f(30.0f, 30.0f) ) );
 		Set( "DialogueWaveDetails.HeaderBorder", new BOX_BRUSH( "Common/MenuBarBorder", FMargin(4.0f/16.0f) ) );
 		Set( "DialogueWaveDetails.PropertyEditorMenu", new BOX_BRUSH( "Old/Menu_Background", FMargin(8.0f/64.0f) ) );
 	}
@@ -837,10 +837,10 @@ void FStarshipEditorStyle::FStyle::SetupGeneralStyles()
 
 		Set( "TableView.Row", FTableRowStyle( NormalTableRowStyle) );
 		Set( "TableView.DarkRow",FTableRowStyle( NormalTableRowStyle)
-			.SetEvenRowBackgroundBrush(IMAGE_BRUSH("PropertyView/DetailCategoryMiddle", FVector2D(16, 16)))
-			.SetEvenRowBackgroundHoveredBrush(IMAGE_BRUSH("PropertyView/DetailCategoryMiddle_Hovered", FVector2D(16, 16)))
-			.SetOddRowBackgroundBrush(IMAGE_BRUSH("PropertyView/DetailCategoryMiddle", FVector2D(16, 16)))
-			.SetOddRowBackgroundHoveredBrush(IMAGE_BRUSH("PropertyView/DetailCategoryMiddle_Hovered", FVector2D(16, 16)))
+			.SetEvenRowBackgroundBrush(IMAGE_BRUSH("PropertyView/DetailCategoryMiddle", Icon16x16))
+			.SetEvenRowBackgroundHoveredBrush(IMAGE_BRUSH("PropertyView/DetailCategoryMiddle_Hovered", Icon16x16))
+			.SetOddRowBackgroundBrush(IMAGE_BRUSH("PropertyView/DetailCategoryMiddle", Icon16x16))
+			.SetOddRowBackgroundHoveredBrush(IMAGE_BRUSH("PropertyView/DetailCategoryMiddle_Hovered", Icon16x16))
 			.SetSelectorFocusedBrush(BORDER_BRUSH("Common/Selector", FMargin(4.f / 16.f), SelectorColor))
 			.SetActiveBrush(IMAGE_BRUSH("Common/Selection", Icon8x8, SelectionColor))
 			.SetActiveHoveredBrush(IMAGE_BRUSH("Common/Selection", Icon8x8, SelectionColor))
@@ -854,7 +854,7 @@ void FStarshipEditorStyle::FStyle::SetupGeneralStyles()
 			.SetInactiveHoveredBrush(FSlateNoResource())
 			);
 	
-		Set("ListView.PinnedItemShadow", new IMAGE_BRUSH("Starship/ListView/PinnedItemShadow", FVector2D(16.f, 8.f)));
+		Set("ListView.PinnedItemShadow", new IMAGE_BRUSH("Starship/ListView/PinnedItemShadow", FVector2f(16.f, 8.f)));
 	}
 	
 	// Spinboxes
@@ -871,19 +871,19 @@ void FStarshipEditorStyle::FStyle::SetupGeneralStyles()
 
 	// Throbber
 	{
-		Set( "SmallThrobber.Chunk", new IMAGE_BRUSH( "Common/ThrobberPiece_Small", FVector2D(8,16) ) );
+		Set( "SmallThrobber.Chunk", new IMAGE_BRUSH( "Common/ThrobberPiece_Small", FVector2f(8.f,16.f) ) );
 	}
 
 	{
-		Set("CurveEd.TimelineArea", new IMAGE_BRUSH("Old/White", Icon16x16, FLinearColor(1, 1, 1, 0.25f)));
+		Set("CurveEd.TimelineArea", new IMAGE_BRUSH("Old/White", Icon16x16, FLinearColor(1.f, 1.f, 1.f, 0.25f)));
 		Set("CurveEd.FitHorizontal", new IMAGE_BRUSH("Icons/FitHorz_16x", Icon16x16));
 		Set("CurveEd.FitVertical", new IMAGE_BRUSH("Icons/FitVert_16x", Icon16x16));
-		Set("CurveEd.CurveKey", new IMAGE_BRUSH("Common/Key", FVector2D(11.0f, 11.0f)));
-		Set("CurveEd.CurveKeySelected", new IMAGE_BRUSH("Common/Key", FVector2D(11.0f, 11.0f), SelectionColor));
+		Set("CurveEd.CurveKey", new IMAGE_BRUSH("Common/Key", FVector2f(11.0f, 11.0f)));
+		Set("CurveEd.CurveKeySelected", new IMAGE_BRUSH("Common/Key", FVector2f(11.0f, 11.0f), SelectionColor));
 		Set("CurveEd.InfoFont", DEFAULT_FONT("Regular", 8));
 		Set("CurveEd.LabelFont", DEFAULT_FONT("Bold", 10));
-		Set("CurveEd.Tangent", new IMAGE_BRUSH("Common/Tangent", FVector2D(7.0f, 7.0f), FLinearColor(0.0f, 0.66f, 0.7f)));
-		Set("CurveEd.TangentSelected", new IMAGE_BRUSH("Common/Tangent", FVector2D(7.0f, 7.0f), FLinearColor(1.0f, 1.0f, 0.0f)));
+		Set("CurveEd.Tangent", new IMAGE_BRUSH("Common/Tangent", FVector2f(7.0f, 7.0f), FLinearColor(0.0f, 0.66f, 0.7f)));
+		Set("CurveEd.TangentSelected", new IMAGE_BRUSH("Common/Tangent", FVector2f(7.0f, 7.0f), FLinearColor(1.0f, 1.0f, 0.0f)));
 		Set("CurveEd.TangentColor", FLinearColor(0.0f, 0.66f, 0.7f));
 		Set("CurveEd.TangentColorSelected", FLinearColor(1.0f, 1.0f, 0.0f));
 	}
@@ -934,7 +934,7 @@ void FStarshipEditorStyle::FStyle::SetupGeneralStyles()
 		Set( "AnimEditor.FilterSearch", new IMAGE_BRUSH( "Old/FilterSearch", Icon16x16 ) );
 		Set( "AnimEditor.FilterCancel", new IMAGE_BRUSH( "Old/FilterCancel", Icon16x16 ) );
 
-		Set( "AnimEditor.NotifyGraphBackground", new IMAGE_BRUSH( "Old/AnimEditor/NotifyTrackBackground", FVector2D(64, 64), FLinearColor::White, ESlateBrushTileType::Both) );
+		Set( "AnimEditor.NotifyGraphBackground", new IMAGE_BRUSH( "Old/AnimEditor/NotifyTrackBackground", Icon64x64, FLinearColor::White, ESlateBrushTileType::Both) );
 
 		Set( "BlendSpace.SamplePoint", new IMAGE_BRUSH( "Old/AnimEditor/BlendSpace_Sample", Icon16x16 ) );
 		Set( "BlendSpace.SamplePoint_Highlight", new IMAGE_BRUSH( "Old/AnimEditor/BlendSpace_Sample_Highlight", Icon16x16 ) );
@@ -977,7 +977,7 @@ void FStarshipEditorStyle::FStyle::SetupGeneralStyles()
 			.SetFont(DEFAULT_FONT( "Bold", 12))
 			.SetColorAndOpacity(FLinearColor(1.0f, 1.0f, 1.0f))
 			.SetHighlightColor(FLinearColor(1.0f, 1.0f, 1.0f))
-			.SetShadowOffset(FVector2D(1.0f, 1.0f))
+			.SetShadowOffset(FVector2f::UnitVector)
 			.SetShadowColorAndOpacity(FLinearColor(0.0f, 0.0f, 0.0f, 0.9f)));
 	}
 
@@ -1004,23 +1004,23 @@ void FStarshipEditorStyle::FStyle::SetupGeneralStyles()
 		Set( "InputBindingEditor.HeaderButton.Disabled", new FSlateNoResource() );
 
 
-		Set( "InputBindingEditor.Tab",  new IMAGE_BRUSH( "Icons/icon_tab_KeyBindings_16px", FVector2D( 16, 16 ) ) );
-		Set( "InputBindingEditor.AssetEditor",  new IMAGE_BRUSH( "Icons/icon_keyb_AssetEditor_16px", FVector2D( 16, 16 ) ) );
-		Set( "InputBindingEditor.AssetEditor",  new IMAGE_BRUSH( "Icons/icon_keyb_AssetEditor_16px", FVector2D( 16, 16 ) ) );
-		Set( "InputBindingEditor.GenericCommands",  new IMAGE_BRUSH( "Icons/icon_keyb_CommonCommands_16px", FVector2D( 16, 16 ) ) );
-		Set( "InputBindingEditor.FoliageEditMode",  new IMAGE_BRUSH( "Icons/icon_keyb_FoliageEditMode_16px", FVector2D( 16, 16 ) ) );
-		Set( "InputBindingEditor.LandscapeEditor",  new IMAGE_BRUSH( "Icons/icon_keyb_LandscapeEditor_16px", FVector2D( 16, 16 ) ) );
-		Set( "InputBindingEditor.LayersView",  new IMAGE_BRUSH( "Icons/icon_keyb_Layers_16px", FVector2D( 16, 16 ) ) );
-		Set( "InputBindingEditor.LevelEditor",  new IMAGE_BRUSH( "Icons/icon_keyb_LevelEditor_16px", FVector2D( 16, 16 ) ) );
-		Set( "InputBindingEditor.LevelViewport",  new IMAGE_BRUSH( "Icons/icon_keyb_LevelViewports_16px", FVector2D( 16, 16 ) ) );
-		Set( "InputBindingEditor.MainFrame",  new IMAGE_BRUSH( "Icons/icon_keyb_MainFrame_16px", FVector2D( 16, 16 ) ) );
-		Set( "InputBindingEditor.OutputLog",  new IMAGE_BRUSH( "Icons/icon_keyb_OutputLog_16px", FVector2D( 16, 16 ) ) );
-		Set( "InputBindingEditor.PlayWorld",  new IMAGE_BRUSH( "Icons/icon_keyb_PlayWorld_16px", FVector2D( 16, 16 ) ) );		
+		Set( "InputBindingEditor.Tab",  new IMAGE_BRUSH( "Icons/icon_tab_KeyBindings_16px", Icon16x16 ) );
+		Set( "InputBindingEditor.AssetEditor",  new IMAGE_BRUSH( "Icons/icon_keyb_AssetEditor_16px", Icon16x16) );
+		Set( "InputBindingEditor.AssetEditor",  new IMAGE_BRUSH( "Icons/icon_keyb_AssetEditor_16px", Icon16x16) );
+		Set( "InputBindingEditor.GenericCommands",  new IMAGE_BRUSH( "Icons/icon_keyb_CommonCommands_16px", Icon16x16) );
+		Set( "InputBindingEditor.FoliageEditMode",  new IMAGE_BRUSH( "Icons/icon_keyb_FoliageEditMode_16px", Icon16x16) );
+		Set( "InputBindingEditor.LandscapeEditor",  new IMAGE_BRUSH( "Icons/icon_keyb_LandscapeEditor_16px", Icon16x16) );
+		Set( "InputBindingEditor.LayersView",  new IMAGE_BRUSH( "Icons/icon_keyb_Layers_16px", Icon16x16) );
+		Set( "InputBindingEditor.LevelEditor",  new IMAGE_BRUSH( "Icons/icon_keyb_LevelEditor_16px", Icon16x16) );
+		Set( "InputBindingEditor.LevelViewport",  new IMAGE_BRUSH( "Icons/icon_keyb_LevelViewports_16px", Icon16x16) );
+		Set( "InputBindingEditor.MainFrame",  new IMAGE_BRUSH( "Icons/icon_keyb_MainFrame_16px", Icon16x16) );
+		Set( "InputBindingEditor.OutputLog",  new IMAGE_BRUSH( "Icons/icon_keyb_OutputLog_16px", Icon16x16) );
+		Set( "InputBindingEditor.PlayWorld",  new IMAGE_BRUSH( "Icons/icon_keyb_PlayWorld_16px", Icon16x16) );
 	}
 
 	// Package restore
 	{
-		Set( "PackageRestore.FolderOpen", new IMAGE_BRUSH( "Icons/FolderOpen", FVector2D(18, 16) ) );
+		Set( "PackageRestore.FolderOpen", new IMAGE_BRUSH( "Icons/FolderOpen", FVector2f(18.f, 16.f) ) );
 	}
 #endif // WITH_EDITOR || (IS_PROGRAM && WITH_UNREAL_DEVELOPER_TOOLS)
 
@@ -1102,11 +1102,11 @@ void FStarshipEditorStyle::FStyle::SetupGeneralStyles()
 		Set("DerivedData.RemoteCache.Uploading", new IMAGE_BRUSH_SVG("Starship/DerivedData/DD_Upload", Icon16x16, EStyleColor::AccentGreen));
 		Set("DerivedData.RemoteCache.Downloading", new IMAGE_BRUSH_SVG("Starship/DerivedData/DD_Download", Icon16x16, EStyleColor::AccentBlue));
 
-		Set("DerivedData.RemoteCache.Uploading", new IMAGE_BRUSH_SVG("Starship/DerivedData/DD_Upload", Icon16x16, FLinearColor(0.0, 1.0, 0.0, 1.0)));
-		Set("DerivedData.RemoteCache.Downloading", new IMAGE_BRUSH_SVG("Starship/DerivedData/DD_Download", Icon16x16, FLinearColor(0.0, 1.0, 1.0, 1.0)));
+		Set("DerivedData.RemoteCache.Uploading", new IMAGE_BRUSH_SVG("Starship/DerivedData/DD_Upload", Icon16x16, FLinearColor::Green));
+		Set("DerivedData.RemoteCache.Downloading", new IMAGE_BRUSH_SVG("Starship/DerivedData/DD_Download", Icon16x16, FLinearColor(0.0f, 1.0f, 1.0f, 1.0f)));
 
 		Set("DerivedData.RemoteCache.BusyBG", new IMAGE_BRUSH_SVG("Starship/DerivedData/DD_RemoteCache_UpDownBG", Icon16x16));
-		Set("DerivedData.RemoteCache.Busy", new IMAGE_BRUSH_SVG("Starship/DerivedData/DD_RemoteCache_UpDown", Icon16x16, FLinearColor(0.3, 0.3, 0.3, 1.0)));
+		Set("DerivedData.RemoteCache.Busy", new IMAGE_BRUSH_SVG("Starship/DerivedData/DD_RemoteCache_UpDown", Icon16x16, FLinearColor(0.3f, 0.3f, 0.3f, 1.0f)));
 
 		Set("DerivedData.RemoteCache.IdleBG", new IMAGE_BRUSH_SVG("Starship/DerivedData/DD_RemoteCache_IdleBG", Icon16x16));
 		Set("DerivedData.RemoteCache.Idle", new IMAGE_BRUSH_SVG("Starship/DerivedData/DD_RemoteCache_Idle", Icon16x16, EStyleColor::Success));
@@ -1149,7 +1149,7 @@ void FStarshipEditorStyle::FStyle::SetupGeneralStyles()
 		Set( "SceneOutliner.FolderOpen",    new CORE_IMAGE_BRUSH_SVG("Starship/Common/folder-open", Icon16x16, FStyleColors::AccentFolder ) );
 		Set( "SceneOutliner.CleanupActorFoldersIcon", new CORE_IMAGE_BRUSH_SVG("Starship/Common/folder-cleanup", Icon16x16));
 		Set( "SceneOutliner.World", 		new CORE_IMAGE_BRUSH_SVG("Starship/Common/world", Icon16x16 ) );
-		Set( "SceneOutliner.ChangedItemHighlight", new FSlateRoundedBoxBrush( FStyleColors::White, 1.0) );
+		Set( "SceneOutliner.ChangedItemHighlight", new FSlateRoundedBoxBrush( FStyleColors::White, 1.0f) );
 
 		const FTableRowStyle AlternatingTableRowStyle = GetWidgetStyle<FTableRowStyle>("TableView.AlternatingRow");
 		
@@ -1170,14 +1170,14 @@ void FStarshipEditorStyle::FStyle::SetupGeneralStyles()
 			.SetNormal        ( FSlateNoResource() )
 			.SetPressed       ( BOX_BRUSH( "Common/Button_Pressed", 8.0f/32.0f, SelectionColor_Pressed ) )
 			.SetHovered       ( BOX_BRUSH( "Common/Button_Hovered", 8.0f/32.0f, SelectionColor ) )
-			.SetNormalPadding ( FMargin( 2,2,4,4 ) )
-			.SetPressedPadding( FMargin( 3,3,3,3 ) )
+			.SetNormalPadding ( FMargin( 2.f,2.f,4.f,4.f) )
+			.SetPressedPadding( FMargin( 3.f,3.f,3.f,3.f) )
 		);
 
 		Set( "GraphBreadcrumbButtonText", FTextBlockStyle(NormalText)
 			.SetFont( DEFAULT_FONT( "Regular", 14 ) )
-			.SetColorAndOpacity( FLinearColor(1,1,1,0.5) )
-			.SetShadowOffset( FVector2D::ZeroVector )
+			.SetColorAndOpacity( FLinearColor(1.f,1.f,1.f,0.5f) )
+			.SetShadowOffset( FVector2f::ZeroVector )
 		);
 
 		Set("GraphBreadcrumb.BrowseBack", new IMAGE_BRUSH_SVG( "Starship/Common/PreviousArrow", Icon20x20));
@@ -1235,22 +1235,22 @@ void FStarshipEditorStyle::FStyle::SetupGeneralStyles()
 	// Asset Thumbnail
 	{
 		Set( "AssetThumbnail.AssetBackground", new FSlateColorBrush(FStyleColors::Recessed));
-		Set( "AssetThumbnail.ClassBackground", new IMAGE_BRUSH( "Common/ClassBackground_64x", FVector2D(64.f, 64.f), FLinearColor(0.75f, 0.75f, 0.75f, 1.0f) ) );
+		Set( "AssetThumbnail.ClassBackground", new IMAGE_BRUSH( "Common/ClassBackground_64x", Icon64x64, FLinearColor(0.75f, 0.75f, 0.75f, 1.0f) ) );
 		Set( "AssetThumbnail.Font", DEFAULT_FONT( "Regular", 10 ) );
 		Set( "AssetThumbnail.FontSmall", DEFAULT_FONT( "Regular", 7 ) );
-		Set( "AssetThumbnail.ColorAndOpacity", FLinearColor(1.75f, 1.75f, 1.75f, 1) );
+		Set( "AssetThumbnail.ColorAndOpacity", FLinearColor(1.75f, 1.75f, 1.75f, 1.f) );
 		Set( "AssetThumbnail.HintFont", DEFAULT_FONT( "Regular", 8 ) );
 		Set( "AssetThumbnail.HintFontSmall", DEFAULT_FONT( "Regular", 6 ) );
-		Set( "AssetThumbnail.HintColorAndOpacity", FLinearColor(0.75f, 0.75f, 0.75f, 1) );
-		Set( "AssetThumbnail.HintShadowOffset", FVector2D(1,1) );
-		Set( "AssetThumbnail.HintShadowColorAndOpacity", FLinearColor(0, 0, 0, 0.5) );
+		Set( "AssetThumbnail.HintColorAndOpacity", FLinearColor(0.75f, 0.75f, 0.75f, 1.f) );
+		Set( "AssetThumbnail.HintShadowOffset", FVector2f::UnitVector );
+		Set( "AssetThumbnail.HintShadowColorAndOpacity", FLinearColor(0.f, 0.f, 0.f, 0.5f) );
 		Set( "AssetThumbnail.HintBackground", new BOX_BRUSH( "Common/TableViewHeader", FMargin(8.0f/32.0f) ) );
-		Set( "AssetThumbnail.AssetTileViewWhiteChipBorder", new IMAGE_BRUSH_SVG("Starship/ContentBrowser/CornerChipOuter2_20", FVector2D(20, 20)));
-		Set( "AssetThumbnail.AssetTileViewWhiteChip", new IMAGE_BRUSH_SVG("Starship/ContentBrowser/CornerChipInner2_20", FVector2D(20, 20)));
-		Set( "AssetThumbnail.AssetTileViewWhiteChipSharp", new IMAGE_BRUSH_SVG("Starship/ContentBrowser/CornerChipSharp_20", FVector2D(20, 20)));
+		Set( "AssetThumbnail.AssetTileViewWhiteChipBorder", new IMAGE_BRUSH_SVG("Starship/ContentBrowser/CornerChipOuter2_20", Icon20x20));
+		Set( "AssetThumbnail.AssetTileViewWhiteChip", new IMAGE_BRUSH_SVG("Starship/ContentBrowser/CornerChipInner2_20", Icon20x20));
+		Set( "AssetThumbnail.AssetTileViewWhiteChipSharp", new IMAGE_BRUSH_SVG("Starship/ContentBrowser/CornerChipSharp_20", Icon20x20));
 		// TODO: Update with correct ratio icons for the small one once ready
-		Set( "AssetThumbnail.AssetTileViewWhiteChipBorderSmall", new IMAGE_BRUSH_SVG("Starship/ContentBrowser/CornerChipOuter2_20", FVector2D(20, 20)));
-		Set( "AssetThumbnail.AssetTileViewWhiteChipSmall", new IMAGE_BRUSH_SVG("Starship/ContentBrowser/CornerChipInner2_20", FVector2D(20, 20)));
+		Set("AssetThumbnail.AssetTileViewWhiteChipBorderSmall", new IMAGE_BRUSH_SVG("Starship/ContentBrowser/CornerChipOuter2_20", Icon20x20));
+		Set( "AssetThumbnail.AssetTileViewWhiteChipSmall", new IMAGE_BRUSH_SVG("Starship/ContentBrowser/CornerChipInner2_20", Icon20x20));
 		Set( "AssetThumbnail.AssetBorder", new FSlateRoundedBoxBrush(FStyleColors::Transparent, FVector4(4.0f, 4.0f, 4.0f, 4.0f), FStyleColors::Secondary, 2.f));
 		Set( "AssetThumbnail.AssetBorderSmall", new FSlateRoundedBoxBrush(FStyleColors::Transparent, FVector4(4.0f, 4.0f, 4.0f, 4.0f), FStyleColors::Secondary, 1.f));
 		FLinearColor OverlayColorAssetStatusOverlay = FStyleColors::Panel.GetSpecifiedColor();
@@ -1276,9 +1276,9 @@ void FStarshipEditorStyle::FStyle::SetupGeneralStyles()
 		Set("AssetThumbnail.EditMode.Primitive", EditModePrimitives);
 
 		const FButtonStyle ActionButton = FButtonStyle()
-			.SetNormal(IMAGE_BRUSH_SVG("Starship/AssetActions/PlayButtonBackground", FVector2D(32, 32), FStyleColors::Secondary))
-			.SetHovered(IMAGE_BRUSH_SVG("Starship/AssetActions/PlayButtonBackground", FVector2D(32, 32), FStyleColors::Hover))
-			.SetPressed(IMAGE_BRUSH_SVG("Starship/AssetActions/PlayButtonBackground", FVector2D(32, 32), FStyleColors::Recessed));
+			.SetNormal(IMAGE_BRUSH_SVG("Starship/AssetActions/PlayButtonBackground", Icon32x32, FStyleColors::Secondary))
+			.SetHovered(IMAGE_BRUSH_SVG("Starship/AssetActions/PlayButtonBackground", Icon32x32, FStyleColors::Hover))
+			.SetPressed(IMAGE_BRUSH_SVG("Starship/AssetActions/PlayButtonBackground", Icon32x32, FStyleColors::Recessed));
 		Set("AssetThumbnail.Action.Button", ActionButton);
 
 	}
@@ -1288,14 +1288,14 @@ void FStarshipEditorStyle::FStyle::SetupGeneralStyles()
 		Set( "SystemWideCommands.SummonOpenAssetDialog", new IMAGE_BRUSH_SVG( "Starship/Common/OpenAsset", Icon16x16 ) );
 	
 		Set( "GlobalAssetPicker.Background", new BOX_BRUSH( "Old/Menu_Background", FMargin(8.0f/64.0f) ) );
-		Set( "GlobalAssetPicker.OutermostMargin", FMargin(4, 4, 4, 4) );
+		Set( "GlobalAssetPicker.OutermostMargin", FMargin(4.f, 4.f, 4.f, 4.f) );
 
 		Set( "GlobalAssetPicker.TitleFont", FTextBlockStyle(NormalText)
-				.SetFont( DEFAULT_FONT( "Regular", 9 ) )
+			.SetFont( DEFAULT_FONT( "Regular", 9 ) )
 			.SetColorAndOpacity( FLinearColor::White )
-				.SetShadowOffset( FVector2D( 1,1 ) )
-				.SetShadowColorAndOpacity( FLinearColor::Black )
-			);
+			.SetShadowOffset( FVector2f::UnitVector )
+			.SetShadowColorAndOpacity( FLinearColor::Black )
+		);
 	}
 
 
@@ -1363,9 +1363,9 @@ void FStarshipEditorStyle::FStyle::SetupGeneralStyles()
 	// About screen
 	if (IncludeEditorSpecificStyles())
 	{
-		Set( "AboutScreen.Background", new IMAGE_BRUSH( "About/Background", FVector2D(688,317) ) );
-		Set( "AboutScreen.UnrealLogo", new IMAGE_BRUSH_SVG( "About/UnrealLogo", FVector2D(40,40) ) );
-		Set( "AboutScreen.EpicGamesLogo", new IMAGE_BRUSH_SVG( "About/EpicGamesLogo", FVector2D(40,40) ) );
+		Set( "AboutScreen.Background", new IMAGE_BRUSH( "About/Background", FVector2f(688.f,317.f) ) );
+		Set( "AboutScreen.UnrealLogo", new IMAGE_BRUSH_SVG( "About/UnrealLogo", Icon40x40 ) );
+		Set( "AboutScreen.EpicGamesLogo", new IMAGE_BRUSH_SVG( "About/EpicGamesLogo", Icon40x40 ) );
 		Set( "AboutScreen.TitleFont", DEFAULT_FONT( "Bold", 13) );
 	}
 #endif // WITH_EDITOR || (IS_PROGRAM && WITH_UNREAL_DEVELOPER_TOOLS)
@@ -1387,47 +1387,47 @@ void FStarshipEditorStyle::FStyle::SetupGeneralStyles()
 
 		FTextBlockStyle CreditsNormal = FTextBlockStyle(NormalText)
 			.SetFont(DEFAULT_FONT("Regular", 16))
-			.SetShadowOffset(FVector2D::UnitVector);
+			.SetShadowOffset(FVector2f::UnitVector);
 
 		Set("Credits.Normal", CreditsNormal);
 
 		Set("Credits.Strong", FTextBlockStyle(CreditsNormal)
 			.SetFont(DEFAULT_FONT("Bold", 16))
-			.SetShadowOffset(FVector2D::UnitVector));
+			.SetShadowOffset(FVector2f::UnitVector));
 
 		Set("Credits.H1", FTextBlockStyle(CreditsNormal)
 			.SetColorAndOpacity(EditorOrange)
 			.SetFont(DEFAULT_FONT("Bold", 36))
-			.SetShadowOffset(FVector2D::UnitVector));
+			.SetShadowOffset(FVector2f::UnitVector));
 
 		Set("Credits.H2", FTextBlockStyle(CreditsNormal)
 			.SetColorAndOpacity(EditorOrange)
 			.SetFont(DEFAULT_FONT("Bold", 30))
-			.SetShadowOffset(FVector2D::UnitVector));
+			.SetShadowOffset(FVector2f::UnitVector));
 
 		Set("Credits.H3", FTextBlockStyle(CreditsNormal)
 			.SetFont(DEFAULT_FONT("Bold", 24))
-			.SetShadowOffset(FVector2D::UnitVector));
+			.SetShadowOffset(FVector2f::UnitVector));
 
 		Set("Credits.H4", FTextBlockStyle(CreditsNormal)
 			.SetFont(DEFAULT_FONT("Bold", 18))
-			.SetShadowOffset(FVector2D::UnitVector));
+			.SetShadowOffset(FVector2f::UnitVector));
 
 		Set("Credits.H5", FTextBlockStyle(CreditsNormal)
 			.SetFont(DEFAULT_FONT("Bold", 12))
-			.SetShadowOffset(FVector2D::UnitVector));
+			.SetShadowOffset(FVector2f::UnitVector));
 
 		Set("Credits.H6", FTextBlockStyle(CreditsNormal)
 			.SetFont(DEFAULT_FONT("Bold", 6))
-			.SetShadowOffset(FVector2D::UnitVector));
+			.SetShadowOffset(FVector2f::UnitVector));
 
 		FTextBlockStyle LinkText = FTextBlockStyle(NormalText)
 			.SetColorAndOpacity(EditorOrange)
-			.SetShadowOffset(FVector2D::UnitVector);
+			.SetShadowOffset(FVector2f::UnitVector);
 		FButtonStyle HoverOnlyHyperlinkButton = FButtonStyle()
 			.SetNormal(FSlateNoResource())
 			.SetPressed(FSlateNoResource())
-			.SetHovered(BORDER_BRUSH("Old/HyperlinkUnderline", FMargin(0, 0, 0, 3 / 16.0f)));
+			.SetHovered(BORDER_BRUSH("Old/HyperlinkUnderline", FMargin(0.f, 0.f, 0.f, 3.f / 16.0f)));
 		FHyperlinkStyle HoverOnlyHyperlink = FHyperlinkStyle()
 			.SetUnderlineStyle(HoverOnlyHyperlinkButton)
 			.SetTextStyle(LinkText)
@@ -1450,15 +1450,15 @@ void FStarshipEditorStyle::FStyle::SetupGeneralStyles()
 		Set("HardwareTargets.Strong", FTextBlockStyle(TargetSettingsNormal)
 			.SetFont(DEFAULT_FONT("Bold", 8))
 			.SetColorAndOpacity(EditorOrange)
-			.SetShadowOffset(FVector2D::UnitVector));
+			.SetShadowOffset(FVector2f::UnitVector));
 	}
 #endif
 
 	// New Level Dialog
 #if WITH_EDITOR || (IS_PROGRAM && WITH_UNREAL_DEVELOPER_TOOLS)
 	{
-		Set( "NewLevelDialog.Blank", new IMAGE_BRUSH( "NewLevels/NewLevelBlank", FVector2D(256,256) ) );
-		Set( "NewLevelDialog.BlankWP", new IMAGE_BRUSH("NewLevels/NewLevelBlankWP", FVector2D(256, 256) ) );
+		Set( "NewLevelDialog.Blank", new IMAGE_BRUSH( "NewLevels/NewLevelBlank", FVector2f(256.f,256.f) ) );
+		Set( "NewLevelDialog.BlankWP", new IMAGE_BRUSH("NewLevels/NewLevelBlankWP", FVector2f(256.f, 256.f) ) );
 	}
 
 	// Build and Submit
@@ -1536,11 +1536,11 @@ void FStarshipEditorStyle::FStyle::SetupGeneralStyles()
 		Set( "FoliageEditMode.DeleteItem", new IMAGE_BRUSH( "Icons/FoliageEditMode/icon_FoliageEditMode_RemoveSettings_20x", Icon20x20 ) );
 		Set( "FoliageEditMode.SelectionBackground", new IMAGE_BRUSH( "Icons/FoliageEditMode/FoliageEditMode_SelectionBackground", Icon32x32 ) );
 		Set( "FoliageEditMode.ItemBackground", new IMAGE_BRUSH( "Icons/FoliageEditMode/FoliageEditMode_Background", Icon64x64 ) );
-		Set( "FoliageEditMode.BubbleBorder", new BOX_BRUSH( "Icons/FoliageEditMode/FoliageEditMode_BubbleBorder", FMargin(8/32.0f) ) );
+		Set( "FoliageEditMode.BubbleBorder", new BOX_BRUSH( "Icons/FoliageEditMode/FoliageEditMode_BubbleBorder", FMargin(8.f/32.0f) ) );
 
 		Set( "FoliageEditMode.TreeView.ScrollBorder", FScrollBorderStyle()
 			.SetTopShadowBrush(FSlateNoResource())
-			.SetBottomShadowBrush(BOX_BRUSH("Common/ScrollBorderShadowBottom", FVector2D(16, 8), FMargin(0.5, 0, 0.5, 1)))
+			.SetBottomShadowBrush(BOX_BRUSH("Common/ScrollBorderShadowBottom", FVector2f(16.f, 8.f), FMargin(0.5f, 0.f, 0.5f, 1.f)))
 			);
 
 		Set("FoliageEditMode.Splitter", FSplitterStyle()
@@ -1550,33 +1550,33 @@ void FStarshipEditorStyle::FStyle::SetupGeneralStyles()
 
 		Set("FoliageEditMode.ActiveToolName.Text", FTextBlockStyle(NormalText)
 			.SetFont(DEFAULT_FONT("Bold", 11))
-			.SetShadowOffset(FVector2D(1, 1))
+			.SetShadowOffset(FVector2f::UnitVector)
 			);
 
 		Set("FoliageEditMode.AddFoliageType.Text", FTextBlockStyle(NormalText)
 			.SetFont(DEFAULT_FONT("Bold", 10))
 			.SetColorAndOpacity(FLinearColor(1.0f, 1.0f, 1.0f))
 			.SetHighlightColor(FLinearColor(1.0f, 1.0f, 1.0f))
-			.SetShadowOffset(FVector2D(1, 1))
-			.SetShadowColorAndOpacity(FLinearColor(0, 0, 0, 0.9f)));
+			.SetShadowOffset(FVector2f::UnitVector)
+			.SetShadowColorAndOpacity(FLinearColor(0.f, 0.f, 0.f, 0.9f)));
 
 
 		FLinearColor InputA = FStyleColors::Input.GetSpecifiedColor();
-		InputA .A = .60;
+		InputA .A = .60f;
 
 		FLinearColor InputB = FStyleColors::Input.GetSpecifiedColor();
-		InputA .A = .70;
+		InputA .A = .70f;
 
 		Set("FoliageEditMode.FloatingButton", FButtonStyle()
-			.SetNormal(FSlateRoundedBoxBrush(InputA, 2))
-			.SetHovered(FSlateRoundedBoxBrush(InputB, 2))
-			.SetPressed(FSlateRoundedBoxBrush(InputB, 2))
+			.SetNormal(FSlateRoundedBoxBrush(InputA, 2.f))
+			.SetHovered(FSlateRoundedBoxBrush(InputB, 2.f))
+			.SetPressed(FSlateRoundedBoxBrush(InputB, 2.f))
 			.SetNormalForeground(FStyleColors::Foreground)
 			.SetHoveredForeground(FStyleColors::ForegroundHover)
 			.SetPressedForeground(FStyleColors::ForegroundHover)
 			.SetDisabledForeground(FStyleColors::White25)
-			.SetNormalPadding(FMargin(4))
-			.SetPressedPadding(FMargin(4))
+			.SetNormalPadding(FMargin(4.f))
+			.SetPressedPadding(FMargin(4.f))
 		 );
 
 	}
@@ -1601,7 +1601,7 @@ void FStarshipEditorStyle::FStyle::SetupGeneralStyles()
 	if (IncludeEditorSpecificStyles())
 	{
 		Set( "GameProjectDialog.BlankProjectThumbnail", new IMAGE_BRUSH( "GameProjectDialog/blank_project_thumbnail", Icon128x128 ) );
-		Set( "GameProjectDialog.BlankProjectPreview", new IMAGE_BRUSH( "GameProjectDialog/blank_project_preview", FVector2D(400, 200) ) );
+		Set( "GameProjectDialog.BlankProjectPreview", new IMAGE_BRUSH( "GameProjectDialog/blank_project_preview", FVector2f(400.f, 200.f) ) );
 
 	}
 
@@ -1635,13 +1635,13 @@ void FStarshipEditorStyle::FStyle::SetupGeneralStyles()
 
 	// Hardware Targeting
 	{
-		Set( "HardwareTargeting.MobilePlatform", new IMAGE_BRUSH( "/Icons/HardwareTargeting/Mobile", FVector2D(64, 64) ) );
-		Set( "HardwareTargeting.DesktopPlatform", new IMAGE_BRUSH( "/Icons/HardwareTargeting/Desktop", FVector2D(64, 64) ) );
-		Set( "HardwareTargeting.HardwareUnspecified", new IMAGE_BRUSH( "/Icons/HardwareTargeting/HardwareUnspecified", FVector2D(64, 64) ) );
+		Set( "HardwareTargeting.MobilePlatform", new IMAGE_BRUSH( "/Icons/HardwareTargeting/Mobile", Icon64x64 ) );
+		Set( "HardwareTargeting.DesktopPlatform", new IMAGE_BRUSH( "/Icons/HardwareTargeting/Desktop", Icon64x64) );
+		Set( "HardwareTargeting.HardwareUnspecified", new IMAGE_BRUSH( "/Icons/HardwareTargeting/HardwareUnspecified", Icon64x64) );
 
-		Set( "HardwareTargeting.MaximumQuality", new IMAGE_BRUSH( "/Icons/HardwareTargeting/MaximumQuality", FVector2D(64, 64) ) );
-		Set( "HardwareTargeting.ScalableQuality", new IMAGE_BRUSH( "/Icons/HardwareTargeting/ScalableQuality", FVector2D(64, 64) ) );
-		Set( "HardwareTargeting.GraphicsUnspecified", new IMAGE_BRUSH( "/Icons/HardwareTargeting/GraphicsUnspecified", FVector2D(64, 64) ) );
+		Set( "HardwareTargeting.MaximumQuality", new IMAGE_BRUSH( "/Icons/HardwareTargeting/MaximumQuality", Icon64x64) );
+		Set( "HardwareTargeting.ScalableQuality", new IMAGE_BRUSH( "/Icons/HardwareTargeting/ScalableQuality", Icon64x64) );
+		Set( "HardwareTargeting.GraphicsUnspecified", new IMAGE_BRUSH( "/Icons/HardwareTargeting/GraphicsUnspecified", Icon64x64) );
 	}
 
 #endif // WITH_EDITOR
@@ -1727,9 +1727,9 @@ void FStarshipEditorStyle::FStyle::SetupGeneralStyles()
 		SlimPaletteToolBarStyle.ButtonStyle.SetNormalPadding(FMargin(0.0f, 0.0f));
 		SlimPaletteToolBarStyle.ButtonStyle.SetPressedPadding(FMargin(0.0f, 0.0f));
 
-		SlimPaletteToolBarStyle.ButtonStyle.Normal = FSlateRoundedBoxBrush(FStyleColors::Dropdown, 4.f, FLinearColor(0, 0, 0, .8), 0.5);
-		SlimPaletteToolBarStyle.ButtonStyle.Hovered = FSlateRoundedBoxBrush(FStyleColors::Hover, 4.f, FLinearColor(0, 0, 0, .8), 0.5);
-		SlimPaletteToolBarStyle.ButtonStyle.Pressed = FSlateRoundedBoxBrush(FStyleColors::Hover, 4.f, FLinearColor(0, 0, 0, .8), 0.5);
+		SlimPaletteToolBarStyle.ButtonStyle.Normal = FSlateRoundedBoxBrush(FStyleColors::Dropdown, 4.f, FLinearColor(0.f, 0.f, 0.f, .8f), 0.5f);
+		SlimPaletteToolBarStyle.ButtonStyle.Hovered = FSlateRoundedBoxBrush(FStyleColors::Hover, 4.f, FLinearColor(0.f, 0.f, 0.f, .8f), 0.5f);
+		SlimPaletteToolBarStyle.ButtonStyle.Pressed = FSlateRoundedBoxBrush(FStyleColors::Hover, 4.f, FLinearColor(0.f, 0.f, 0.f, .8f), 0.5f);
 		SlimPaletteToolBarStyle.ButtonStyle.HoveredForeground = FStyleColors::ForegroundHover;
 		SlimPaletteToolBarStyle.ButtonStyle.PressedForeground = FStyleColors::ForegroundHover;
 		SlimPaletteToolBarStyle.SetUniformBlockHeight(33.f);
@@ -1737,10 +1737,10 @@ void FStarshipEditorStyle::FStyle::SetupGeneralStyles()
 		SlimPaletteToolBarStyle.SetNumColumns(2);
 		
 		FCheckBoxStyle CheckBoxStyle = FCheckBoxStyle(SlimPaletteToolBarStyle.ToggleButton)
-			.SetCheckedImage(FSlateRoundedBoxBrush(FStyleColors::Primary, 4.f, FLinearColor(0, 0, 0, .8), 0.5))
-			.SetCheckedHoveredImage(FSlateRoundedBoxBrush(FStyleColors::PrimaryHover, 4.f, FLinearColor(0, 0, 0, .8), 0.5))
-			.SetCheckedPressedImage(FSlateRoundedBoxBrush(FStyleColors::PrimaryHover, 4.f, FLinearColor(0, 0, 0, .8), 0.5))
-			.SetUncheckedImage(FSlateRoundedBoxBrush(FStyleColors::Dropdown, 4.f, FLinearColor(0, 0, 0, .8), 0.5))
+			.SetCheckedImage(FSlateRoundedBoxBrush(FStyleColors::Primary, 4.f, FLinearColor(0.f, 0.f, 0.f, .8f), 0.5f))
+			.SetCheckedHoveredImage(FSlateRoundedBoxBrush(FStyleColors::PrimaryHover, 4.f, FLinearColor(0.f, 0.f, 0.f, .8f), 0.5f))
+			.SetCheckedPressedImage(FSlateRoundedBoxBrush(FStyleColors::PrimaryHover, 4.f, FLinearColor(0.f, 0.f, 0.f, .8f), 0.5f))
+			.SetUncheckedImage(FSlateRoundedBoxBrush(FStyleColors::Dropdown, 4.f, FLinearColor(0.f, 0.f, 0.f, .8f), 0.5f))
 			.SetForegroundColor(FStyleColors::Foreground)
 			.SetPressedForegroundColor(FStyleColors::ForegroundHover)
 			.SetHoveredForegroundColor(FStyleColors::ForegroundHover)
@@ -1778,11 +1778,11 @@ void FStarshipEditorStyle::FStyle::SetupGeneralStyles()
 		FVerticalToolBarStyle.SetIconPadding(FMargin(8.0f, 8.0f));
 		FVerticalToolBarStyle.SetIconPaddingWithVisibleLabel(FMargin(8.0f, 8.0f, 8.0f, 2.0f));
  
-		FVerticalToolBarStyle.SetExpandBrush(CORE_IMAGE_BRUSH_SVG("Starship/Common/ellipsis-horizontal-narrow", FVector2D(16, 16)));
+		FVerticalToolBarStyle.SetExpandBrush(CORE_IMAGE_BRUSH_SVG("Starship/Common/ellipsis-horizontal-narrow", Icon16x16));
 		FCheckBoxStyle CheckBoxStyle = FCheckBoxStyle(FVerticalToolBarStyle.ToggleButton)
-			.SetCheckedImage(FSlateRoundedBoxBrush(FStyleColors::Primary, 4.f, FLinearColor(0, 0, 0, .8), 0.5))
-			.SetCheckedHoveredImage(FSlateRoundedBoxBrush(FStyleColors::PrimaryHover, 4.f, FLinearColor(0, 0, 0, .8), 0.5))
-			.SetCheckedPressedImage(FSlateRoundedBoxBrush(FStyleColors::PrimaryPress, 4.f, FLinearColor(0, 0, 0, .8), 0.5))
+			.SetCheckedImage(FSlateRoundedBoxBrush(FStyleColors::Primary, 4.f, FLinearColor(0.f, 0.f, 0.f, .8f), 0.5f))
+			.SetCheckedHoveredImage(FSlateRoundedBoxBrush(FStyleColors::PrimaryHover, 4.f, FLinearColor(0.f, 0.f, 0.f, .8f), 0.5f))
+			.SetCheckedPressedImage(FSlateRoundedBoxBrush(FStyleColors::PrimaryPress, 4.f, FLinearColor(0.f, 0.f, 0.f, .8f), 0.5f))
 			.SetForegroundColor(FStyleColors::Foreground)
 			.SetPressedForegroundColor(FStyleColors::ForegroundHover)
 			.SetHoveredForegroundColor(FStyleColors::ForegroundHover)
@@ -1802,20 +1802,20 @@ void FStarshipEditorStyle::FStyle::SetupGeneralStyles()
 	{
 		// FCategoryDrivenContentBuilder vertical toolbar style
 		FToolBarStyle CategoryDrivenContentBuilderToolBarStyle =  GetWidgetStyle<FToolBarStyle>("FVerticalToolBar");
-		CategoryDrivenContentBuilderToolBarStyle.SetIconPadding( FMargin(0, 9, 0, 4 ) );
-		CategoryDrivenContentBuilderToolBarStyle.SetBackgroundPadding( FMargin( 4 ) );
-		CategoryDrivenContentBuilderToolBarStyle.SetIconPaddingWithVisibleLabel( FMargin(0, 9, 0, 0 ) );
-		CategoryDrivenContentBuilderToolBarStyle.SetLabelPadding( FMargin(4, 5, 4, 9) );
+		CategoryDrivenContentBuilderToolBarStyle.SetIconPadding( FMargin(0.f, 9.f, 0.f, 4.f) );
+		CategoryDrivenContentBuilderToolBarStyle.SetBackgroundPadding( FMargin( 4.f ) );
+		CategoryDrivenContentBuilderToolBarStyle.SetIconPaddingWithVisibleLabel( FMargin(0.f, 9.f, 0.f, 0.f) );
+		CategoryDrivenContentBuilderToolBarStyle.SetLabelPadding( FMargin(4.f, 5.f, 4.f, 9.f) );
 		CategoryDrivenContentBuilderToolBarStyle.SetButtonContentMaxWidth( 56.0f );
 		FTextBlockStyle LabelTextStyle = FTextBlockStyle(NormalText)
 		                                 .SetOverflowPolicy(ETextOverflowPolicy::Ellipsis)
 		                                 .SetFont(DEFAULT_FONT("roboto", FCoreStyle::SmallTextSize));
 		CategoryDrivenContentBuilderToolBarStyle.SetLabelStyle( LabelTextStyle );
-		CategoryDrivenContentBuilderToolBarStyle.SetButtonPadding( FMargin( 0, 2, 0, 2 ));
+		CategoryDrivenContentBuilderToolBarStyle.SetButtonPadding( FMargin( 0.f, 2.f, 0.f, 2.f));
 		CategoryDrivenContentBuilderToolBarStyle.SetShowLabels( true );
 		Set("CategoryDrivenContentBuilderToolbarWithLabels", CategoryDrivenContentBuilderToolBarStyle);
 
-		CategoryDrivenContentBuilderToolBarStyle.SetIconPadding( FMargin(8, 8, 8, 8 ) );
+		CategoryDrivenContentBuilderToolBarStyle.SetIconPadding( FMargin(8.f) );
 		CategoryDrivenContentBuilderToolBarStyle.SetShowLabels( false );
 		CategoryDrivenContentBuilderToolBarStyle.SetButtonContentMaxWidth( 36.0f );
 		Set("CategoryDrivenContentBuilderToolbarWithoutLabels", CategoryDrivenContentBuilderToolBarStyle);
@@ -2071,8 +2071,8 @@ void FStarshipEditorStyle::FStyle::SetupWorldBrowserStyles()
 
 		Set("WorldBrowser.StatusBarText", FTextBlockStyle(NormalText)
 			.SetFont(DEFAULT_FONT("BoldCondensed", 12))
-			.SetColorAndOpacity(FLinearColor(0.9, 0.9f, 0.9f, 0.5f))
-			.SetShadowOffset(FVector2D::ZeroVector)
+			.SetColorAndOpacity(FLinearColor(0.9f, 0.9f, 0.9f, 0.5f))
+			.SetShadowOffset(FVector2f::ZeroVector)
 		);
 
 		Set("WorldBrowser.LabelFont", DEFAULT_FONT("Regular", 9));
@@ -2180,11 +2180,11 @@ void FStarshipEditorStyle::FStyle::SetupSequencerStyles()
 		Set("Sequencer.KeyTriangle20", new CORE_IMAGE_BRUSH_SVG("Starship/Common/curve-editor-append-key-20", Icon20x20));
 		Set("Sequencer.KeyLeft", new IMAGE_BRUSH("Sequencer/KeyLeft", Icon12x12));
 		Set("Sequencer.KeyRight", new IMAGE_BRUSH("Sequencer/KeyRight", Icon12x12));
-		Set("Sequencer.PartialKey", new IMAGE_BRUSH("Sequencer/PartialKey", FVector2D(11.f, 11.f)));
+		Set("Sequencer.PartialKey", new IMAGE_BRUSH("Sequencer/PartialKey", FVector2f(11.f, 11.f)));
 		Set("Sequencer.Star", new IMAGE_BRUSH("Sequencer/Star", Icon12x12));
 		Set("Sequencer.Empty", new IMAGE_BRUSH("Sequencer/Empty", Icon12x12));
-		Set("Sequencer.TangentHandle", new IMAGE_BRUSH("Sequencer/TangentHandle", FVector2D(7, 7)));
-		Set("Sequencer.GenericDivider", new IMAGE_BRUSH("Sequencer/GenericDivider", FVector2D(2.f, 2.f), FLinearColor::White, ESlateBrushTileType::Vertical));
+		Set("Sequencer.TangentHandle", new IMAGE_BRUSH("Sequencer/TangentHandle", FVector2f(7.f, 7.f)));
+		Set("Sequencer.GenericDivider", new IMAGE_BRUSH("Sequencer/GenericDivider", FVector2f(2.f, 2.f), FLinearColor::White, ESlateBrushTileType::Vertical));
 
 		Set("Sequencer.KeyBar.Dotted", new BORDER_BRUSH(TEXT("Sequencer/Keys/KeyBar_Dotted"), FMargin(0.f, 4.f, 0.f, 0.f)));
 		Set("Sequencer.KeyBar.Dashed", new BORDER_BRUSH(TEXT("Sequencer/Keys/KeyBar_Dashed"), FMargin(0.f, 4.f, 0.f, 0.f)));
@@ -2197,7 +2197,7 @@ void FStarshipEditorStyle::FStyle::SetupSequencerStyles()
 		Set("Sequencer.Timeline.RangeHandleLeft", new BOX_BRUSH("Sequencer/GenericGripLeft", FMargin(5.f / 16.f)));
 		Set("Sequencer.Timeline.RangeHandleRight", new BOX_BRUSH("Sequencer/GenericGripRight", FMargin(5.f / 16.f)));
 		Set("Sequencer.Timeline.RangeHandle", new BOX_BRUSH("Sequencer/GenericSectionBackground", FMargin(5.f / 16.f)));
-		Set("Sequencer.Timeline.NotifyAlignmentMarker", new IMAGE_BRUSH("Sequencer/NotifyAlignmentMarker", FVector2D(10, 19)));
+		Set("Sequencer.Timeline.NotifyAlignmentMarker", new IMAGE_BRUSH("Sequencer/NotifyAlignmentMarker", FVector2f(10.f, 19.f)));
 		Set("Sequencer.Timeline.PlayRange_Top_L", new BOX_BRUSH("Sequencer/PlayRange_Top_L", FMargin(1.f, 0.5f, 0.f, 0.5f)));
 		Set("Sequencer.Timeline.PlayRange_Top_R", new BOX_BRUSH("Sequencer/PlayRange_Top_R", FMargin(0.f, 0.5f, 1.f, 0.5f)));
 		Set("Sequencer.Timeline.PlayRange_L", new BOX_BRUSH("Sequencer/PlayRange_L", FMargin(1.f, 0.5f, 0.f, 0.5f)));
@@ -2207,7 +2207,7 @@ void FStarshipEditorStyle::FStyle::SetupSequencerStyles()
 
 		Set("Sequencer.Timeline.SubSequenceRangeHashL", new BORDER_BRUSH("Sequencer/SubSequenceRangeHashL", FMargin(1.f, 0.f, 0.f, 0.f)));
 		Set("Sequencer.Timeline.SubSequenceRangeHashR", new BORDER_BRUSH("Sequencer/SubSequenceRangeHashR", FMargin(1.f, 0.f, 0.f, 0.f)));
-		Set("Sequencer.Timeline.EaseInOut", new IMAGE_BRUSH("Sequencer/EaseInOut", FVector2D(128, 128)));
+		Set("Sequencer.Timeline.EaseInOut", new IMAGE_BRUSH("Sequencer/EaseInOut", Icon128x128));
 		Set("Sequencer.InterpLine", new BOX_BRUSH("Sequencer/InterpLine", FMargin(5.f / 7.f, 0.f, 0.f, 0.f)));
 
 		Set("Sequencer.Transport.JumpToPreviousKey", new IMAGE_BRUSH_SVG("Sequencer/PlaybackControls/PlayControlsJumpToPreviousKey", Icon20x20));
@@ -2241,14 +2241,14 @@ void FStarshipEditorStyle::FStyle::SetupSequencerStyles()
 		Set("Sequencer.Section.SelectedTrackTint", new BOX_BRUSH(TEXT("Sequencer/SelectedTrackTint"), FMargin(0.f, 0.5f)));
 		Set("Sequencer.Section.SelectionBorder", new BORDER_BRUSH(TEXT("Sequencer/SectionHighlight"), FMargin(7.f / 16.f)));
 		Set("Sequencer.Section.LockedBorder", new BORDER_BRUSH(TEXT("Sequencer/SectionLocked"), FMargin(7.f / 16.f)));
-		Set("Sequencer.Section.FilmBorder", new IMAGE_BRUSH(TEXT("Sequencer/SectionFilmBorder"), FVector2D(10, 7), FLinearColor::White, ESlateBrushTileType::Horizontal));
-		Set("Sequencer.Section.GripLeft", new FSlateRoundedBoxBrush(FStyleColors::White, FVector4(4, 0, 0, 4)));
-		Set("Sequencer.Section.GripRight", new FSlateRoundedBoxBrush(FStyleColors::White, FVector4(0, 4, 4, 0)));
+		Set("Sequencer.Section.FilmBorder", new IMAGE_BRUSH(TEXT("Sequencer/SectionFilmBorder"), FVector2f(10.f, 7.f), FLinearColor::White, ESlateBrushTileType::Horizontal));
+		Set("Sequencer.Section.GripLeft", new FSlateRoundedBoxBrush(FStyleColors::White, FVector4(4.f, 0.f, 0.f, 4.f)));
+		Set("Sequencer.Section.GripRight", new FSlateRoundedBoxBrush(FStyleColors::White, FVector4(0.f, 4.f, 4.f, 0.f)));
 		Set("Sequencer.Section.EasingHandle", new FSlateColorBrush(FStyleColors::White));
 
-		Set("Sequencer.Section.Background_Collapsed", new FSlateRoundedBoxBrush(FLinearColor::White, FVector4(4, 4, 4, 4)));
-		Set("Sequencer.Section.Background_Header", new FSlateRoundedBoxBrush(FLinearColor::White, FVector4(4, 4, 0, 0)));
-		Set("Sequencer.Section.Background_Contents", new FSlateRoundedBoxBrush(FLinearColor::White, FVector4(0, 0, 4, 4)));
+		Set("Sequencer.Section.Background_Collapsed", new FSlateRoundedBoxBrush(FLinearColor::White, FVector4(4.f, 4.f, 4.f, 4.f)));
+		Set("Sequencer.Section.Background_Header", new FSlateRoundedBoxBrush(FLinearColor::White, FVector4(4.f, 4.f, 0.f, 0.f)));
+		Set("Sequencer.Section.Background_Contents", new FSlateRoundedBoxBrush(FLinearColor::White, FVector4(0.f, 0.f, 4.f, 4.f)));
 
 		Set("Sequencer.Section.PreRoll", new BORDER_BRUSH(TEXT("Sequencer/PreRoll"), FMargin(0.f, .5f, 0.f, .5f)));
 
@@ -2258,8 +2258,8 @@ void FStarshipEditorStyle::FStyle::SetupSequencerStyles()
 		Set("Sequencer.Section.BackgroundText", DEFAULT_FONT("Bold", 24));
 		Set("Sequencer.Section.EmptySpace", new BOX_BRUSH(TEXT("Sequencer/EmptySpace"), FMargin(0.f, 7.f / 14.f)));
 
-		Set("Sequencer.MarkedFrame.LabelRight", new FSlateRoundedBoxBrush(FLinearColor::White, FVector4(0, 0, 4, 0)));
-		Set("Sequencer.MarkedFrame.LabelLeft", new FSlateRoundedBoxBrush(FLinearColor::White, FVector4(0, 0, 0, 4)));
+		Set("Sequencer.MarkedFrame.LabelRight", new FSlateRoundedBoxBrush(FLinearColor::White, FVector4(0.f, 0.f, 4.f, 0.f)));
+		Set("Sequencer.MarkedFrame.LabelLeft", new FSlateRoundedBoxBrush(FLinearColor::White, FVector4(0.f, 0.f, 0.f, 4.f)));
 
 		Set("Sequencer.LayerBar.Background",  new BOX_BRUSH(TEXT("Sequencer/LayerBarBackground"), FMargin(4.f / 16.f)));
 		Set("Sequencer.LayerBar.HandleLeft",  new FSlateColorBrush(FStyleColors::White)); // Maybe use a rounded box if we decide to round-out the layer bars
@@ -2292,8 +2292,8 @@ void FStarshipEditorStyle::FStyle::SetupSequencerStyles()
 
 		Set("Sequencer.Outliner.Indicators.TimeWarp", new IMAGE_BRUSH_SVG("Starship/Sequencer/TimeWarp_12", Icon12x12));
 
-		Set("Sequencer.ShotFilter", new IMAGE_BRUSH("Sequencer/FilteredArea", FVector2D(74, 74), FLinearColor::White, ESlateBrushTileType::Both));
-		Set("Sequencer.KeyMark", new IMAGE_BRUSH("Sequencer/KeyMark", FVector2D(3, 21), FLinearColor::White, ESlateBrushTileType::NoTile));
+		Set("Sequencer.ShotFilter", new IMAGE_BRUSH("Sequencer/FilteredArea", FVector2f(74.f, 74.f), FLinearColor::White, ESlateBrushTileType::Both));
+		Set("Sequencer.KeyMark", new IMAGE_BRUSH("Sequencer/KeyMark", FVector2f(3.f, 21.f), FLinearColor::White, ESlateBrushTileType::NoTile));
 		Set("Sequencer.ToggleAutoKeyEnabled", new IMAGE_BRUSH_SVG("Starship/Sequencer/AutoKey", Icon20x20));
 		Set("Sequencer.SetAutoKey", new IMAGE_BRUSH_SVG("Starship/Sequencer/AutoKey", Icon20x20));
 		Set("Sequencer.SetAutoTrack", new IMAGE_BRUSH_SVG("Starship/Sequencer/AutoTrack", Icon20x20));
@@ -2315,11 +2315,11 @@ void FStarshipEditorStyle::FStyle::SetupSequencerStyles()
 		Set("Sequencer.LockCamera", new IMAGE_BRUSH("Sequencer/Main_Icons/Icon_Sequencer_Look_Thru_24x", Icon16x16));
 		Set("Sequencer.UnlockCamera", new IMAGE_BRUSH("Sequencer/Main_Icons/Icon_Sequencer_Look_Thru_24x", Icon16x16, FLinearColor(1.f, 1.f, 1.f, 0.5f)));
 		Set("Sequencer.Thumbnail.SectionHandle", new IMAGE_BRUSH("Old/White", Icon16x16, FLinearColor::Black));
-		Set("Sequencer.TrackHoverHighlight_Top", new IMAGE_BRUSH(TEXT("Sequencer/TrackHoverHighlight_Top"), FVector2D(4, 4)));
-		Set("Sequencer.TrackHoverHighlight_Bottom", new IMAGE_BRUSH(TEXT("Sequencer/TrackHoverHighlight_Bottom"), FVector2D(4, 4)));
-		Set("Sequencer.SpawnableIconOverlay", new IMAGE_BRUSH(TEXT("Sequencer/SpawnableIconOverlay"), FVector2D(13, 13)));
-		Set("Sequencer.ReplaceableIconOverlay", new IMAGE_BRUSH(TEXT("Sequencer/ReplaceableIconOverlay"), FVector2D(13, 13)));
-		Set("Sequencer.MultipleIconOverlay", new IMAGE_BRUSH(TEXT("Sequencer/MultipleIconOverlay"), FVector2D(13, 13)));
+		Set("Sequencer.TrackHoverHighlight_Top", new IMAGE_BRUSH(TEXT("Sequencer/TrackHoverHighlight_Top"), FVector2f(4.f, 4.f)));
+		Set("Sequencer.TrackHoverHighlight_Bottom", new IMAGE_BRUSH(TEXT("Sequencer/TrackHoverHighlight_Bottom"), FVector2f(4.f, 4.f)));
+		Set("Sequencer.SpawnableIconOverlay", new IMAGE_BRUSH(TEXT("Sequencer/SpawnableIconOverlay"), FVector2f(13.f, 13.f)));
+		Set("Sequencer.ReplaceableIconOverlay", new IMAGE_BRUSH(TEXT("Sequencer/ReplaceableIconOverlay"), FVector2f(13.f, 13.f)));
+		Set("Sequencer.MultipleIconOverlay", new IMAGE_BRUSH(TEXT("Sequencer/MultipleIconOverlay"), FVector2f(13.f, 13.f)));
 		Set("Sequencer.DynamicBindingIconOverlay", new IMAGE_BRUSH(TEXT("Sequencer/DynamicBindingIconOverlay"), Icon16x16));
 		Set("Sequencer.SpawnableDynamicBindingIconOverlay", new IMAGE_BRUSH(TEXT("Sequencer/SpawnableDynamicBindingIconOverlay"), Icon16x16));
 		Set("Sequencer.LockSequence", new IMAGE_BRUSH("Sequencer/Main_Icons/Icon_Sequencer_Locked_16x", Icon16x16));
@@ -2360,12 +2360,12 @@ void FStarshipEditorStyle::FStyle::SetupSequencerStyles()
 		Set("Sequencer.BreadcrumbText", FTextBlockStyle(NormalText).SetFont(FStyleFonts::Get().NormalBold));
 		Set("Sequencer.BreadcrumbIcon", new IMAGE_BRUSH("Common/SmallArrowRight", Icon10x10));
 
-		Set("Sequencer.AddKey.Details", new IMAGE_BRUSH("Sequencer/AddKey_Details", FVector2D(11, 11)));
+		Set("Sequencer.AddKey.Details", new IMAGE_BRUSH("Sequencer/AddKey_Details", FVector2f(11.f, 11.f)));
 
-		Set("Sequencer.KeyedStatus.NotKeyed", new IMAGE_BRUSH_SVG("Sequencer/DetailsKeyUnkeyed", FVector2D(11, 11)));
-		Set("Sequencer.KeyedStatus.Keyed", new IMAGE_BRUSH_SVG("Sequencer/DetailsKeyKeyed", FVector2D(11, 11)));
-		Set("Sequencer.KeyedStatus.Animated", new IMAGE_BRUSH_SVG("Sequencer/DetailsKeyAnimated", FVector2D(11, 11)));
-		Set("Sequencer.KeyedStatus.PartialKey", new IMAGE_BRUSH_SVG("Sequencer/DetailsKeyPartialKey", FVector2D(11, 11)));
+		Set("Sequencer.KeyedStatus.NotKeyed", new IMAGE_BRUSH_SVG("Sequencer/DetailsKeyUnkeyed", FVector2f(11.f, 11.f)));
+		Set("Sequencer.KeyedStatus.Keyed", new IMAGE_BRUSH_SVG("Sequencer/DetailsKeyKeyed", FVector2f(11.f, 11.f)));
+		Set("Sequencer.KeyedStatus.Animated", new IMAGE_BRUSH_SVG("Sequencer/DetailsKeyAnimated", FVector2f(11.f, 11.f)));
+		Set("Sequencer.KeyedStatus.PartialKey", new IMAGE_BRUSH_SVG("Sequencer/DetailsKeyPartialKey", FVector2f(11.f, 11.f)));
 
 		const FSplitterStyle OutlinerSplitterStyle = FSplitterStyle()
 		.SetHandleNormalBrush(FSlateNoResource())
@@ -2373,8 +2373,8 @@ void FStarshipEditorStyle::FStyle::SetupSequencerStyles()
 		Set("Sequencer.AnimationOutliner.Splitter", OutlinerSplitterStyle);
 
 		Set("Sequencer.HyperlinkSpinBox", FSpinBoxStyle(GetWidgetStyle<FSpinBoxStyle>("SpinBox"))
-			.SetTextPadding(FMargin(0))
-			.SetBackgroundBrush(BORDER_BRUSH("Old/HyperlinkDotted", FMargin(0, 0, 0, 3 / 16.0f), FSlateColor::UseSubduedForeground()))
+			.SetTextPadding(FMargin(0.f))
+			.SetBackgroundBrush(BORDER_BRUSH("Old/HyperlinkDotted", FMargin(0.f, 0.f, 0.f, 3.f / 16.0f), FSlateColor::UseSubduedForeground()))
 			.SetHoveredBackgroundBrush(FSlateNoResource())
 			.SetInactiveFillBrush(FSlateNoResource())
 			.SetActiveFillBrush(FSlateNoResource())
@@ -2383,7 +2383,7 @@ void FStarshipEditorStyle::FStyle::SetupSequencerStyles()
 		);
 
 		Set("Sequencer.PlayTimeSpinBox", FSpinBoxStyle(GetWidgetStyle<FSpinBoxStyle>("SpinBox"))
-			.SetTextPadding(FMargin(0))
+			.SetTextPadding(FMargin(0.f))
 			.SetBackgroundBrush(FSlateNoResource())
 			.SetHoveredBackgroundBrush(FSlateNoResource())
 			.SetInactiveFillBrush(FSlateNoResource())
@@ -2425,11 +2425,11 @@ void FStarshipEditorStyle::FStyle::SetupSequencerStyles()
 		Set("Sequencer.UnboundEvent", new IMAGE_BRUSH("Sequencer/UnboundEvent", Icon16x16));
 
 		// Sequencer Blending Iconography
-		Set("EMovieSceneBlendType::Absolute", new IMAGE_BRUSH("Sequencer/EMovieSceneBlendType_Absolute", FVector2D(32, 16)));
-		Set("EMovieSceneBlendType::Relative", new IMAGE_BRUSH("Sequencer/EMovieSceneBlendType_Relative", FVector2D(32, 16)));
-		Set("EMovieSceneBlendType::Additive", new IMAGE_BRUSH("Sequencer/EMovieSceneBlendType_Additive", FVector2D(32, 16)));
-		Set("EMovieSceneBlendType::Override", new IMAGE_BRUSH_SVG("Sequencer/EMovieSceneBlendType_Override", FVector2D(32, 16)));
-		Set("EMovieSceneBlendType::AdditiveFromBase", new IMAGE_BRUSH("Sequencer/EMovieSceneBlendType_AdditiveFromBase", FVector2D(32, 16)));
+		Set("EMovieSceneBlendType::Absolute", new IMAGE_BRUSH("Sequencer/EMovieSceneBlendType_Absolute", FVector2f(32.f, 16.f)));
+		Set("EMovieSceneBlendType::Relative", new IMAGE_BRUSH("Sequencer/EMovieSceneBlendType_Relative", FVector2f(32.f, 16.f)));
+		Set("EMovieSceneBlendType::Additive", new IMAGE_BRUSH("Sequencer/EMovieSceneBlendType_Additive", FVector2f(32.f, 16.f)));
+		Set("EMovieSceneBlendType::Override", new IMAGE_BRUSH_SVG("Sequencer/EMovieSceneBlendType_Override", FVector2f(32.f, 16.f)));
+		Set("EMovieSceneBlendType::AdditiveFromBase", new IMAGE_BRUSH("Sequencer/EMovieSceneBlendType_AdditiveFromBase", FVector2f(32.f, 16.f)));
 
 		Set("Sequencer.TrackIsolate", new IMAGE_BRUSH_SVG("Sequencer/TrackIsolate", Icon16x16));
 		Set("Sequencer.TrackHide", new IMAGE_BRUSH_SVG("Sequencer/TrackHide", Icon16x16));
@@ -2477,17 +2477,17 @@ void FStarshipEditorStyle::FStyle::SetupViewportStyles()
 	{
 		FLinearColor ToolbarBackgroundColor = FStyleColors::Dropdown.GetSpecifiedColor();
 		ToolbarBackgroundColor.A = .90f;
-		const FSlateRoundedBoxBrush BackgroundBrush(ToolbarBackgroundColor, 5.f, ToolbarBackgroundColor, 1.0);
+		const FSlateRoundedBoxBrush BackgroundBrush(ToolbarBackgroundColor, 5.f, ToolbarBackgroundColor, 1.0f);
 
 		FLinearColor ToolbarHoveredColor = FStyleColors::Dropdown.GetSpecifiedColor();
-		ToolbarHoveredColor.R = FMath::Min(1.0f, ToolbarHoveredColor.R * 1.5);
-		ToolbarHoveredColor.G = FMath::Min(1.0f, ToolbarHoveredColor.G * 1.5);
-		ToolbarHoveredColor.B = FMath::Min(1.0f, ToolbarHoveredColor.B * 1.5);
-		const FSlateRoundedBoxBrush HoveredButtonBrush(ToolbarHoveredColor, 5.f, ToolbarHoveredColor, 1.0);
+		ToolbarHoveredColor.R = FMath::Min(1.0f, ToolbarHoveredColor.R * 1.5f);
+		ToolbarHoveredColor.G = FMath::Min(1.0f, ToolbarHoveredColor.G * 1.5f);
+		ToolbarHoveredColor.B = FMath::Min(1.0f, ToolbarHoveredColor.B * 1.5f);
+		const FSlateRoundedBoxBrush HoveredButtonBrush(ToolbarHoveredColor, 5.f, ToolbarHoveredColor, 1.0f);
 
 		FLinearColor ToolbarPressedColor = FStyleColors::Recessed.GetSpecifiedColor();
 		ToolbarPressedColor.A = .50f;
-		const FSlateRoundedBoxBrush PressedButtonBrush(ToolbarPressedColor, 5.f, ToolbarPressedColor, 1.0);
+		const FSlateRoundedBoxBrush PressedButtonBrush(ToolbarPressedColor, 5.f, ToolbarPressedColor, 1.0f);
 
 		FButtonStyle ButtonStyle;
 		{
@@ -2507,7 +2507,7 @@ void FStarshipEditorStyle::FStyle::SetupViewportStyles()
 
 		FCheckBoxStyle ToggleButtonStyle = SlimToolBarStyle.ToggleButton;
 		{
-			const FSlateRoundedBoxBrush TransparentBrush(FLinearColor(0, 0, 0, 0), 0, FLinearColor(0, 0, 0, 0), 0);
+			const FSlateRoundedBoxBrush TransparentBrush(FLinearColor(0.f, 0.f, 0.f, 0.f), 0.f, FLinearColor(0.f, 0.f, 0.f, 0.f), 0.f);
 
 			ToggleButtonStyle.SetCheckedImage(TransparentBrush)
 				.SetCheckedHoveredImage(HoveredButtonBrush)
@@ -2521,12 +2521,12 @@ void FStarshipEditorStyle::FStyle::SetupViewportStyles()
 
 		FComboButtonStyle ComboButtonStyle = SlimToolBarStyle.ComboButtonStyle;
 		{
-			FButtonStyle InnerButtonStyle = FButtonStyle(ButtonStyle).SetNormalPadding(0).SetPressedPadding(0);
+			FButtonStyle InnerButtonStyle = FButtonStyle(ButtonStyle).SetNormalPadding(0.f).SetPressedPadding(0.f);
 
 			ComboButtonStyle
 				.SetButtonStyle(InnerButtonStyle)
 				.SetDownArrowImage(IMAGE_BRUSH_SVG("Starship/EditorViewport/small-chevron-down", FVector2f(6.f, 6.f)))
-				.SetDownArrowPadding(FMargin(4.0f, 0, 0, 0));
+				.SetDownArrowPadding(FMargin(4.0f, 0.f, 0.f, 0.f));
 		}
 
 		FToolBarStyle ViewportToolbarStyle = SlimToolBarStyle;
@@ -2538,14 +2538,14 @@ void FStarshipEditorStyle::FStyle::SetupViewportStyles()
 			ViewportToolbarStyle.SetIconSize(Icon16x16)
 				.SetButtonStyle(ButtonStyle)
 				.SetToggleButtonStyle(ToggleButtonStyle)
-				.SetButtonPadding(0)
+				.SetButtonPadding(0.f)
 				.SetComboButtonStyle(ComboButtonStyle)
-				.SetComboButtonPadding(FMargin(4.0f, 0))
+				.SetComboButtonPadding(FMargin(4.0f, 0.f))
 				.SetSeparatorBrush(SeparatorBrush)
 				.SetSeparatorThickness(1.0f)
-				.SetSeparatorPadding(FMargin(4.0f, 0))
+				.SetSeparatorPadding(FMargin(4.0f, 0.f))
 				.SetBackgroundPadding(FMargin(2.0f))
-				.SetWrapButtonPadding(FMargin(4.0f, 0))
+				.SetWrapButtonPadding(FMargin(4.0f, 0.f))
 				.SetWrapButtonIndex(-1)
 				.SetRaisedChildrenRightPadding(16.0f);
 
@@ -2618,9 +2618,9 @@ void FStarshipEditorStyle::FStyle::SetupViewportStyles()
 				const FSlateRoundedBoxBrush WarningBrush(WarningColor, 5.0f, WarningColor, 1.0f);
 
 				FLinearColor WarningHoveredColor = FStyleColors::Warning.GetSpecifiedColor();
-				WarningHoveredColor.R = FMath::Min(1.0f, WarningHoveredColor.R * 1.5);
-				WarningHoveredColor.G = FMath::Min(1.0f, WarningHoveredColor.G * 1.5);
-				WarningHoveredColor.B = FMath::Min(1.0f, WarningHoveredColor.B * 1.5);
+				WarningHoveredColor.R = FMath::Min(1.0f, WarningHoveredColor.R * 1.5f);
+				WarningHoveredColor.G = FMath::Min(1.0f, WarningHoveredColor.G * 1.5f);
+				WarningHoveredColor.B = FMath::Min(1.0f, WarningHoveredColor.B * 1.5f);
 				const FSlateRoundedBoxBrush WarningHoveredBrush(WarningHoveredColor, 5.0f, WarningHoveredColor, 1.0f);
 
 				FLinearColor WarningPressedColor = FStyleColors::Warning.GetSpecifiedColor();
@@ -2673,7 +2673,7 @@ void FStarshipEditorStyle::FStyle::SetupViewportStyles()
 			const FComboButtonStyle ViewportSizingComboButtonStyle =
 				FComboButtonStyle(ComboButtonStyle)
 					.SetDownArrowPadding(FMargin(5.f, 0.f, 2.f, 0.f))
-					.SetDownArrowImage(CORE_IMAGE_BRUSH_SVG("Starship/Common/ellipsis-vertical-narrow", FVector2D(6, 20)))
+					.SetDownArrowImage(CORE_IMAGE_BRUSH_SVG("Starship/Common/ellipsis-vertical-narrow", FVector2f(6.f, 20.f)))
 					.SetButtonStyle(ViewportSizingButtonStyle);
 
 			const FToolBarStyle ViewportSizingStyle = FToolBarStyle(ViewportToolbarStyle)
@@ -2690,20 +2690,20 @@ void FStarshipEditorStyle::FStyle::SetupViewportStyles()
 				{
 					FLinearColor BackgroundColorRaised = FStyleColors::Dropdown.GetSpecifiedColor();
 					BackgroundColorRaised.A = .90f;
-					const FSlateRoundedBoxBrush BackgroundBrushRaised(BackgroundColorRaised, 5.f, BackgroundColorRaised, 1.0);
+					const FSlateRoundedBoxBrush BackgroundBrushRaised(BackgroundColorRaised, 5.f, BackgroundColorRaised, 1.0f);
 
 					ButtonStyleRaised.SetNormal(BackgroundBrushRaised)
 						.SetNormalForeground(FStyleColors::Foreground)
-						.SetNormal(FSlateRoundedBoxBrush(ToolbarBackgroundColor, FVector4(0, 4.0f, 4.0f, 0)))
-						.SetHovered(FSlateRoundedBoxBrush(ToolbarHoveredColor, FVector4(0, 4.0f, 4.0f, 0)))
-						.SetPressed(FSlateRoundedBoxBrush(ToolbarPressedColor, FVector4(0, 4.0f, 4.0f, 0)))
-						.SetDisabled(FSlateRoundedBoxBrush(FStyleColors::Dropdown, FVector4(0, 4.0f, 4.0f, 0)));
+						.SetNormal(FSlateRoundedBoxBrush(ToolbarBackgroundColor, FVector4(0.f, 4.0f, 4.0f, 0.f)))
+						.SetHovered(FSlateRoundedBoxBrush(ToolbarHoveredColor, FVector4(0.f, 4.0f, 4.0f, 0.f)))
+						.SetPressed(FSlateRoundedBoxBrush(ToolbarPressedColor, FVector4(0.f, 4.0f, 4.0f, 0.f)))
+						.SetDisabled(FSlateRoundedBoxBrush(FStyleColors::Dropdown, FVector4(0.f, 4.0f, 4.0f, 0.f)));
 				}
 
 				const FToolBarStyle ViewportSizingRaisedStyle = FToolBarStyle(ViewportToolbarStyle)
 																	.SetButtonStyle(ButtonStyleRaised)
-																	.SetButtonPadding(0)
-																	.SetComboButtonPadding(FMargin(2.0f, 0));
+																	.SetButtonPadding(0.f)
+																	.SetComboButtonPadding(FMargin(2.0f, 0.f));
 
 				Set("ViewportToolbarViewportSizingSubmenu.Raised", ViewportSizingRaisedStyle);
 			}
@@ -2768,10 +2768,10 @@ void FStarshipEditorStyle::FStyle::SetupViewportStyles()
 		FLinearColor ToolbarPressedColor = FStyleColors::Recessed.GetSpecifiedColor();
 		ToolbarPressedColor.A = .80f;
 
-		FSlateRoundedBoxBrush* ViewportGroupBrush = new FSlateRoundedBoxBrush(ToolbarBackgroundColor, 12.f, FLinearColor(0,0,0,.8), 1.0);
+		FSlateRoundedBoxBrush* ViewportGroupBrush = new FSlateRoundedBoxBrush(ToolbarBackgroundColor, 12.f, FLinearColor(0.f,0.f,0.f,.8f), 1.0f);
 		Set("EditorViewportToolBar.Group", ViewportGroupBrush);
 
-		FSlateRoundedBoxBrush* ViewportGroupPressedBrush = new FSlateRoundedBoxBrush(ToolbarPressedColor, 12.f, FLinearColor(0, 0, 0, .8), 1.0);
+		FSlateRoundedBoxBrush* ViewportGroupPressedBrush = new FSlateRoundedBoxBrush(ToolbarPressedColor, 12.f, FLinearColor(0.f, 0.f, 0.f, .8f), 1.0f);
 		Set("EditorViewportToolBar.Group.Pressed", ViewportGroupPressedBrush);
 
 		FButtonStyle ViewportMenuButton = FButtonStyle()
@@ -2882,9 +2882,9 @@ void FStarshipEditorStyle::FStyle::SetupViewportStyles()
 			.SetUncheckedImage(*ViewportGroupBrush)
 			.SetUncheckedPressedImage(*ViewportGroupPressedBrush)
 			.SetUncheckedHoveredImage(*ViewportGroupBrush)
-			.SetCheckedImage(FSlateRoundedBoxBrush(FStyleColors::Primary, 12.f, FLinearColor(0, 0, 0, .8), 1.0))
-			.SetCheckedHoveredImage(FSlateRoundedBoxBrush(FStyleColors::PrimaryHover, 12.f, FLinearColor(0, 0, 0, .8), 1.0))
-			.SetCheckedPressedImage(FSlateRoundedBoxBrush(FStyleColors::PrimaryPress, 12.f, FLinearColor(0, 0, 0, .8), 1.0))
+			.SetCheckedImage(FSlateRoundedBoxBrush(FStyleColors::Primary, 12.f, FLinearColor(0.f, 0.f, 0.f, .8f), 1.0f))
+			.SetCheckedHoveredImage(FSlateRoundedBoxBrush(FStyleColors::PrimaryHover, 12.f, FLinearColor(0.f, 0.f, 0.f, .8f), 1.0f))
+			.SetCheckedPressedImage(FSlateRoundedBoxBrush(FStyleColors::PrimaryPress, 12.f, FLinearColor(0.f, 0.f, 0.f, .8f), 1.0f))
 			.SetForegroundColor(FStyleColors::Foreground)
 			.SetPressedForegroundColor(FStyleColors::ForegroundHover)
 			.SetHoveredForegroundColor(FStyleColors::ForegroundHover)
@@ -2897,12 +2897,12 @@ void FStarshipEditorStyle::FStyle::SetupViewportStyles()
 		ViewportToolbarStyle
 			.SetBackground(FSlateNoResource())
 			.SetIconSize(Icon16x16)
-			.SetBackgroundPadding(FMargin(0))
-			.SetLabelPadding(FMargin(0))
+			.SetBackgroundPadding(FMargin(0.f))
+			.SetLabelPadding(FMargin(0.f))
 			.SetComboButtonPadding(FMargin(4.f, 0.0f))
 			.SetBlockPadding(FMargin(0.0f,0.0f))
-			.SetIndentedBlockPadding(FMargin(0))
-			.SetButtonPadding(FMargin(0))
+			.SetIndentedBlockPadding(FMargin(0.f))
+			.SetButtonPadding(FMargin(0.f))
 			.SetCheckBoxPadding(FMargin(4.0f, 0.0f))
 			.SetComboButtonStyle(ViewportComboButton)
 			.SetToggleButtonStyle(SoloToggleButton)
@@ -2925,8 +2925,8 @@ void FStarshipEditorStyle::FStyle::SetupViewportStyles()
 		Set("EditorViewportToolBar.Font", FStyleFonts::Get().Normal);
 
 		Set("EditorViewportToolBar.MenuButton", FButtonStyle(Button)
-			.SetNormal(BOX_BRUSH("Common/SmallRoundedButton", FMargin(7.f / 16.f), FLinearColor(1, 1, 1, 0.75f)))
-			.SetHovered(BOX_BRUSH("Common/SmallRoundedButton", FMargin(7.f / 16.f), FLinearColor(1, 1, 1, 1.0f)))
+			.SetNormal(BOX_BRUSH("Common/SmallRoundedButton", FMargin(7.f / 16.f), FLinearColor(1.f, 1.f, 1.f, 0.75f)))
+			.SetHovered(BOX_BRUSH("Common/SmallRoundedButton", FMargin(7.f / 16.f), FLinearColor(1.f, 1.f, 1.f, 1.0f)))
 			.SetPressed(BOX_BRUSH("Common/SmallRoundedButton", FMargin(7.f / 16.f)))
 		);
 
@@ -2939,11 +2939,11 @@ void FStarshipEditorStyle::FStyle::SetupViewportStyles()
 		FLinearColor ViewportOverlayColor = FStyleColors::Input.GetSpecifiedColor();
 		ViewportOverlayColor.A = 0.75f;
 
-		Set("EditorViewport.OverlayBrush", new FSlateRoundedBoxBrush(ViewportOverlayColor, 8.0, FStyleColors::Dropdown, 1.0));
+		Set("EditorViewport.OverlayBrush", new FSlateRoundedBoxBrush(ViewportOverlayColor, 8.0f, FStyleColors::Dropdown, 1.0f));
 
 		const FSlateColorBrush ActionableListViewBrush = FSlateColorBrush(FLinearColor(0.f, 0.f, 0.f, 0.f));
 		
-		Set("ActionableMessage.Border", new FSlateRoundedBoxBrush(ToolbarBackgroundColor, 4.0f, FLinearColor(0,0,0,.8), 1.0));
+		Set("ActionableMessage.Border", new FSlateRoundedBoxBrush(ToolbarBackgroundColor, 4.0f, FLinearColor(0.f,0.f,0.f,.8f), 1.0f));
 		Set("ActionableMessage.Warning", new IMAGE_BRUSH_SVG(TEXT("Starship/EditorViewport/alert-solid"), Icon16x16, FStyleColors::Warning));
 		Set("ActionableMessage.Update", new IMAGE_BRUSH_SVG(TEXT("Starship/EditorViewport/update"), Icon16x16));
 		Set("ActionableMessage.ListView", FTableViewStyle().SetBackgroundBrush(ActionableListViewBrush));
@@ -2966,16 +2966,16 @@ void FStarshipEditorStyle::FStyle::SetupViewportStyles()
 			FToolBarStyle()
 			.SetBackground(BOX_BRUSH("Old/Menu_Background", FMargin(8.0f / 64.0f), FLinearColor::Transparent))
 			.SetExpandBrush(IMAGE_BRUSH("Icons/toolbar_expand_16x", Icon8x8))
-			.SetComboButtonPadding(FMargin(0))
-			.SetButtonPadding(FMargin(0))
-			.SetCheckBoxPadding(FMargin(4))
+			.SetComboButtonPadding(FMargin(0.f))
+			.SetButtonPadding(FMargin(0.f))
+			.SetCheckBoxPadding(FMargin(4.f))
 			.SetSeparatorBrush(BOX_BRUSH("Old/Button", 8.0f / 32.0f, FLinearColor::Transparent))
 			.SetSeparatorPadding(FMargin(1.f, 0.f, 0.f, 0.f))
 			.SetIconSize(Icon16x16)
 			.SetLabelPadding(FMargin(0.0f, 0.0f, 3.0f, 0.0f))
 			.SetEditableTextStyle(FEditableTextBoxStyle(NormalEditableTextBoxStyle).SetFont(DEFAULT_FONT("Regular", 9)))
-			.SetIndentedBlockPadding(FMargin(0))
-			.SetBlockPadding(FMargin(0))
+			.SetIndentedBlockPadding(FMargin(0.f))
+			.SetBlockPadding(FMargin(0.f))
 			.SetLabelStyle(
 				FTextBlockStyle(NormalText)
 				.SetFont(DEFAULT_FONT("Bold", 9))
@@ -2987,8 +2987,8 @@ void FStarshipEditorStyle::FStyle::SetupViewportStyles()
 		const FString SmallRoundedButtonMiddle(TEXT("Common/SmallRoundedButtonCentre"));
 		const FString SmallRoundedButtonEnd(TEXT("Common/SmallRoundedButtonRight"));
 
-		const FLinearColor NormalColor(1, 1, 1, 0.75f);
-		const FLinearColor PressedColor(1, 1, 1, 1.f);
+		const FLinearColor NormalColor(1.f, 1.f, 1.f, 0.75f);
+		const FLinearColor PressedColor(1.f, 1.f, 1.f, 1.f);
 
 		/* Create style for "LegacyViewportMenu.ToggleButton" ... */
 		const FCheckBoxStyle ViewportMenuToggleButtonStyle = FCheckBoxStyle()
@@ -3128,14 +3128,14 @@ void FStarshipEditorStyle::FStyle::SetupWindowStyles()
 	WindowTitleOverride = new FSlateColorBrush(FStyleColors::Title);
 	Set("Brushes.Title", WindowTitleOverride);
 
-	Set("WindowSize.Small", FVector2D(480, 284));
-	Set("WindowSize.Medium", FVector2D(680, 492));
+	Set("WindowSize.Small", FVector2f(480.f, 284.f));
+	Set("WindowSize.Medium", FVector2f(680.f, 492.f));
 }
 
 void FStarshipEditorStyle::FStyle::SetupProjectBadgeStyle()
 {
 	Set("SProjectBadge.Text", FTextBlockStyle(NormalText));
-	Set("SProjectBadge.BadgeShape", new BOX_BRUSH("ProjectBadge/Badge", Icon16x16, FMargin(6.0f/16)));
+	Set("SProjectBadge.BadgeShape", new BOX_BRUSH("ProjectBadge/Badge", Icon16x16, FMargin(6.0f/16.f)));
 	Set("SProjectBadge.BadgePadding", FMargin(32.0f, 6.0f, 32.0f, 7.0f));
 }
 
@@ -3238,7 +3238,7 @@ void FStarshipEditorStyle::FStyle::SetupPropertyEditorStyles()
 		Set( "PropertyEditor.RemoveColumn", new IMAGE_BRUSH( "Common/PushPin_Down", Icon16x16, FColor( 96, 194, 253, 255 ).ReinterpretAsLinear() ) );
 		Set( "PropertyEditor.AddColumn", new IMAGE_BRUSH( "Common/PushPin_Up", Icon16x16, FColor( 96, 194, 253, 255 ).ReinterpretAsLinear() ) );
 
-		Set( "PropertyEditor.AddColumnOverlay",	new IMAGE_BRUSH( "Common/TinyChalkArrow", FVector2D( 71, 20 ), FColor( 96, 194, 253, 255 ).ReinterpretAsLinear() ) );
+		Set( "PropertyEditor.AddColumnOverlay",	new IMAGE_BRUSH( "Common/TinyChalkArrow", FVector2f( 71.f, 20.f), FColor( 96, 194, 253, 255 ).ReinterpretAsLinear() ) );
 		Set( "PropertyEditor.AddColumnMessage", FTextBlockStyle(NormalText)
 			.SetFont( DEFAULT_FONT( "BoldCondensedItalic", 10 ) )
 			.SetColorAndOpacity(FColor( 96, 194, 253, 255 ).ReinterpretAsLinear())
@@ -3254,20 +3254,20 @@ void FStarshipEditorStyle::FStyle::SetupPropertyEditorStyles()
 		Set( "PropertyEditor.AssetClass", FTextBlockStyle(NormalText)
 			.SetFont( DEFAULT_FONT( "Regular", 10 ) )
 			.SetColorAndOpacity( FLinearColor::White )
-			.SetShadowOffset( FVector2D(1,1) )
+			.SetShadowOffset( FVector2f::UnitVector )
 			.SetShadowColorAndOpacity( FLinearColor::Black )
 		);
 
 		const FButtonStyle AssetComboStyle = FButtonStyle()
-			.SetNormal( BOX_BRUSH( "Common/ButtonHoverHint", FMargin(4/16.0f), FLinearColor(1,1,1,0.15f) ) )
-			.SetHovered( BOX_BRUSH( "Common/ButtonHoverHint", FMargin(4/16.0f), FLinearColor(1,1,1,0.25f) ) )
-			.SetPressed( BOX_BRUSH( "Common/ButtonHoverHint", FMargin(4/16.0f), FLinearColor(1,1,1,0.30f) ) )
-			.SetNormalPadding( FMargin(0,0,0,1) )
-			.SetPressedPadding( FMargin(0,1,0,0) );
+			.SetNormal( BOX_BRUSH( "Common/ButtonHoverHint", FMargin(4.f /16.0f), FLinearColor(1.f,1.f,1.f,0.15f) ) )
+			.SetHovered( BOX_BRUSH( "Common/ButtonHoverHint", FMargin(4.f /16.0f), FLinearColor(1.f,1.f,1.f,0.25f) ) )
+			.SetPressed( BOX_BRUSH( "Common/ButtonHoverHint", FMargin(4.f /16.0f), FLinearColor(1.f,1.f,1.f,0.30f) ) )
+			.SetNormalPadding( FMargin(0.f,0.f,0.f,1.f) )
+			.SetPressedPadding( FMargin(0.f,1.f,0.f,0.f) );
 		Set( "PropertyEditor.AssetComboStyle", AssetComboStyle );
 
-		Set( "PropertyEditor.HorizontalDottedLine",		new IMAGE_BRUSH( "Common/HorizontalDottedLine_16x1px", FVector2D(16.0f, 1.0f), FLinearColor::White, ESlateBrushTileType::Horizontal ) );
-		Set( "PropertyEditor.VerticalDottedLine",		new IMAGE_BRUSH( "Common/VerticalDottedLine_1x16px", FVector2D(1.0f, 16.0f), FLinearColor::White, ESlateBrushTileType::Vertical ) );
+		Set( "PropertyEditor.HorizontalDottedLine",		new IMAGE_BRUSH( "Common/HorizontalDottedLine_16x1px", FVector2f(16.0f, 1.0f), FLinearColor::White, ESlateBrushTileType::Horizontal ) );
+		Set( "PropertyEditor.VerticalDottedLine",		new IMAGE_BRUSH( "Common/VerticalDottedLine_1x16px", FVector2f(1.0f, 16.0f), FLinearColor::White, ESlateBrushTileType::Vertical ) );
 		Set( "PropertyEditor.SlateBrushPreview",		new BOX_BRUSH( "PropertyView/SlateBrushPreview_32px", Icon32x32, FMargin(3.f/32.f, 3.f/32.f, 15.f/32.f, 13.f/32.f) ) );
 
 		Set( "PropertyTable.TableRow", GetWidgetStyle<FTableRowStyle>("TableView.AlternatingRow"));
@@ -3460,7 +3460,7 @@ void FStarshipEditorStyle::FStyle::SetupPropertyEditorStyles()
 		FComboButtonStyle CategoryComboButton = FComboButtonStyle(FStarshipCoreStyle::GetCoreStyle().GetWidgetStyle<FComboButtonStyle>("ComboButton"))
 		.SetButtonStyle(CategoryRowButton)
 		.SetDownArrowPadding(FMargin(2.f, 5.f, 3.f, 5.f))
-		.SetDownArrowImage(CORE_IMAGE_BRUSH_SVG("Starship/Common/ellipsis-vertical-narrow", FVector2D(6, 15)));
+		.SetDownArrowImage(CORE_IMAGE_BRUSH_SVG("Starship/Common/ellipsis-vertical-narrow", FVector2f(6.f, 15.f)));
 		CategoryComboButton.ButtonStyle = CategoryRowButton;
 		Set( "DetailsView.CategoryComboButton", CategoryComboButton);
 
@@ -3518,22 +3518,22 @@ void FStarshipEditorStyle::FStyle::SetupPropertyEditorStyles()
 
 		FToolBarStyle DetailsExtensionToolBarStyle = FToolBarStyle(FStarshipCoreStyle::GetCoreStyle().GetWidgetStyle<FToolBarStyle>("SlimToolBar"))
 			.SetButtonStyle(DetailsExtensionMenuButton)
-			.SetExpandBrush(CORE_IMAGE_BRUSH_SVG("Starship/Common/ellipsis-vertical-narrow", FVector2D(4, 16)))
+			.SetExpandBrush(CORE_IMAGE_BRUSH_SVG("Starship/Common/ellipsis-vertical-narrow", FVector2f(4.f, 16.f)))
 			.SetIconSize(Icon16x16)
 			.SetBackground(FSlateNoResource())
-			.SetLabelPadding(FMargin(0))
-			.SetComboButtonPadding(FMargin(0))
-			.SetBlockPadding(FMargin(0))
-			.SetIndentedBlockPadding(FMargin(0))
-			.SetBackgroundPadding(FMargin(0))
-			.SetButtonPadding(FMargin(0))
-			.SetCheckBoxPadding(FMargin(0))
+			.SetLabelPadding(FMargin(0.f))
+			.SetComboButtonPadding(FMargin(0.f))
+			.SetBlockPadding(FMargin(0.f))
+			.SetIndentedBlockPadding(FMargin(0.f))
+			.SetBackgroundPadding(FMargin(0.f))
+			.SetButtonPadding(FMargin(0.f))
+			.SetCheckBoxPadding(FMargin(0.f))
 			.SetSeparatorBrush(FSlateNoResource())
-			.SetSeparatorPadding(FMargin(0));
+			.SetSeparatorPadding(FMargin(0.f));
 
 		Set("DetailsView.ExtensionToolBar", DetailsExtensionToolBarStyle);
 
-		Set("DetailsView.ArrayDropShadow", new IMAGE_BRUSH("Common/ArrayDropShadow", FVector2D(32,2)));
+		Set("DetailsView.ArrayDropShadow", new IMAGE_BRUSH("Common/ArrayDropShadow", FVector2f(32.f,2.f)));
 
 		Set( "DetailsView.TreeView.TableRow", FTableRowStyle()
 			.SetEvenRowBackgroundBrush( FSlateNoResource() )
@@ -3592,21 +3592,21 @@ void FStarshipEditorStyle::FStyle::SetupProfilerStyle()
 		Set( "Profiler.CaptionBold", FTextBlockStyle(NormalText)
 			.SetFont( DEFAULT_FONT( "Bold", 10 ) )
 			.SetColorAndOpacity( FLinearColor::White )
-			.SetShadowOffset( FVector2D(1.0f, 1.0f) )
+			.SetShadowOffset( FVector2f::UnitVector )
 			.SetShadowColorAndOpacity( FLinearColor(0.f,0.f,0.f,0.8f) )
 		);
 
 		Set( "Profiler.TooltipBold", FTextBlockStyle(NormalText)
 			.SetFont( DEFAULT_FONT( "Bold", 8 ) )
 			.SetColorAndOpacity( FLinearColor(0.5f,0.5f,0.5f,1.0f) )
-			.SetShadowOffset( FVector2D(1.0f, 1.0f) )
+			.SetShadowOffset( FVector2f::UnitVector )
 			.SetShadowColorAndOpacity( FLinearColor(0.f,0.f,0.f,0.8f) )
 		);
 
 		Set( "Profiler.Tooltip", FTextBlockStyle(NormalText)
 			.SetFont( DEFAULT_FONT( "Regular", 8 ) )
 			.SetColorAndOpacity( FLinearColor::White )
-			.SetShadowOffset( FVector2D(1.0f, 1.0f) )
+			.SetShadowOffset( FVector2f::UnitVector )
 			.SetShadowColorAndOpacity( FLinearColor(0.f,0.f,0.f,0.8f) )
 		);
 
@@ -3654,7 +3654,7 @@ void FStarshipEditorStyle::FStyle::SetupProfilerStyle()
 		Set( "Profiler.EventGraph.DarkText", FTextBlockStyle(NormalText)
 			.SetFont( DEFAULT_FONT( "Regular", 8 ) )
 			.SetColorAndOpacity( FLinearColor::Black )
-			.SetShadowOffset( FVector2D(0.0f, 0.0f) )
+			.SetShadowOffset( FVector2f::ZeroVector )
 			);
 
 		// Thread-view
@@ -3705,7 +3705,7 @@ void FStarshipEditorStyle::FStyle::SetupGraphEditorStyles()
 	{
 		Set( "Graph.ForegroundColor", FLinearColor(218.0f/255.0f, 218.0f/255.0f, 218.0f/255.0f, 1.0f) );
 
-		Set( "Graph.TitleBackground", new BOX_BRUSH( "Old/Graph/GraphTitleBackground", FMargin(0) ) );
+		Set( "Graph.TitleBackground", new BOX_BRUSH( "Old/Graph/GraphTitleBackground", FMargin(0.f) ) );
 		Set( "Graph.Shadow", new BOX_BRUSH( "Old/Window/WindowBorder", 0.48f ) );
 		Set( "Graph.Arrow", new IMAGE_BRUSH( "Old/Graph/Arrow", Icon16x16 ) );
 		Set( "Graph.ExecutionBubble", new IMAGE_BRUSH( "Old/Graph/ExecutionBubble", Icon16x16 ) );
@@ -3713,35 +3713,35 @@ void FStarshipEditorStyle::FStyle::SetupGraphEditorStyles()
 		Set( "Graph.PlayInEditor", new BOX_BRUSH( "/Graph/RegularNode_shadow_selected", FMargin(18.0f/64.0f) ) );
 		Set( "Graph.ReadOnlyBorder", new BOX_BRUSH( "/Graph/Graph_readonly_border", FMargin(18.0f / 64.0f) ) );
 
-		Set( "Graph.Panel.SolidBackground", new IMAGE_BRUSH( "/Graph/GraphPanel_SolidBackground", FVector2D(16, 16), FLinearColor::White, ESlateBrushTileType::Both) );
-		Set( "Graph.Panel.GridLineColor",   FLinearColor(0.024f, 0.024, 0.024f) );
+		Set( "Graph.Panel.SolidBackground", new IMAGE_BRUSH( "/Graph/GraphPanel_SolidBackground", Icon16x16, FLinearColor::White, ESlateBrushTileType::Both) );
+		Set( "Graph.Panel.GridLineColor",   FLinearColor(0.024f, 0.024f, 0.024f) );
 		Set( "Graph.Panel.GridRuleColor",   FLinearColor(0.010f, 0.010f, 0.010f) );
 		Set( "Graph.Panel.GridCenterColor", FLinearColor(0.005f, 0.005f, 0.005f) );
 		
 		Set( "Graph.Panel.GridRulePeriod", 8.0f ); // should be a strictly positive integral value
 
-		Set( "Graph.Node.Separator", new IMAGE_BRUSH( "Old/Graph/NodeVerticalSeparator", FVector2D(8,8) ) );
+		Set( "Graph.Node.Separator", new IMAGE_BRUSH( "Old/Graph/NodeVerticalSeparator", Icon8x8 ) );
 		Set( "Graph.Node.TitleBackground", new BOX_BRUSH( "Old/Graph/NodeTitleBackground", FMargin(12.0f/64) ) );
 		Set( "Graph.Node.NodeBackground", new BOX_BRUSH( "Old/Graph/NodeBackground", FMargin(12.0f/64) ) );
 
 		Set( "Graph.Node.Body", new BOX_BRUSH( "/Graph/RegularNode_body", FMargin(16.f/64.f, 25.f/64.f, 16.f/64.f, 16.f/64.f) ) );
 		Set( "Graph.Node.TintedBody", new BOX_BRUSH( "/Graph/TintedNode_body", FMargin(16.f/64.f, 25.f/64.f, 16.f/64.f, 16.f/64.f) ) );
-		Set( "Graph.Node.DisabledBanner", new IMAGE_BRUSH( "/Graph/GraphPanel_StripesBackground", FVector2D(64, 64), FLinearColor(0.5f, 0.5f, 0.5f, 0.3f), ESlateBrushTileType::Both ) );
-		Set( "Graph.Node.DevelopmentBanner", new IMAGE_BRUSH( "/Graph/GraphPanel_StripesBackground", FVector2D(64, 64), FLinearColor::Yellow * FLinearColor(1.f, 1.f, 1.f, 0.3f), ESlateBrushTileType::Both ) );
+		Set( "Graph.Node.DisabledBanner", new IMAGE_BRUSH( "/Graph/GraphPanel_StripesBackground", Icon64x64, FLinearColor(0.5f, 0.5f, 0.5f, 0.3f), ESlateBrushTileType::Both ) );
+		Set( "Graph.Node.DevelopmentBanner", new IMAGE_BRUSH( "/Graph/GraphPanel_StripesBackground", Icon64x64, FLinearColor::Yellow * FLinearColor(1.f, 1.f, 1.f, 0.3f), ESlateBrushTileType::Both ) );
 		Set( "Graph.Node.TitleGloss", new BOX_BRUSH( "/Graph/RegularNode_title_gloss", FMargin(12.0f/64.0f) ) );
-		Set( "Graph.Node.ColorSpill", new BOX_BRUSH( "/Graph/RegularNode_color_spill", FMargin(8.0f/64.0f, 3.0f/32.0f, 0, 0) ) );
+		Set( "Graph.Node.ColorSpill", new BOX_BRUSH( "/Graph/RegularNode_color_spill", FMargin(8.0f/64.0f, 3.0f/32.0f, 0.f, 0.f) ) );
 		Set( "Graph.Node.TitleHighlight", new BOX_BRUSH( "/Graph/RegularNode_title_highlight", FMargin(16.0f/64.0f, 1.0f, 16.0f/64.0f, 0.0f) ) );
-		Set( "Graph.Node.IndicatorOverlay", new IMAGE_BRUSH( "/Graph/IndicatorOverlay_color_spill", FVector2D(128.f, 32.f) ) );
+		Set( "Graph.Node.IndicatorOverlay", new IMAGE_BRUSH( "/Graph/IndicatorOverlay_color_spill", FVector2f(128.f, 32.f) ) );
 
-		Set( "Graph.Node.ShadowSize", FVector2D(12,12) );
+		Set( "Graph.Node.ShadowSize", Icon12x12 );
 		Set( "Graph.Node.ShadowSelected", new BOX_BRUSH( "/Graph/RegularNode_shadow_selected", FMargin(18.0f/64.0f) ) );
 		Set( "Graph.Node.Shadow", new BOX_BRUSH( "/Graph/RegularNode_shadow", FMargin(18.0f/64.0f) ) );
 		
 		Set( "Graph.Node.DiffHighlight", new BOX_BRUSH( "/Graph/RegularNode_DiffHighlight", FMargin(18.0f/64.0f) ) );
 		Set( "Graph.Node.DiffHighlightShading", new BOX_BRUSH( "/Graph/RegularNode_DiffHighlightShading", FMargin(18.0f/64.0f) ) );
 
-		Set( "Graph.Node.RerouteShadow", new IMAGE_BRUSH( "/Graph/RerouteNode_shadow", FVector2D(64.0f, 64.0f) ) );
-		Set( "Graph.Node.RerouteShadowSelected", new IMAGE_BRUSH( "/Graph/RerouteNode_shadow_selected", FVector2D(64.0f, 64.0f) ) );
+		Set( "Graph.Node.RerouteShadow", new IMAGE_BRUSH( "/Graph/RerouteNode_shadow", Icon64x64 ) );
+		Set( "Graph.Node.RerouteShadowSelected", new IMAGE_BRUSH( "/Graph/RerouteNode_shadow_selected", Icon64x64 ) );
 		
 		Set( "Graph.Node.RerouteDiffHighlight", new BOX_BRUSH( "/Graph/RerouteNode_DiffHighlight", FMargin(18.0f/64.0f) ) );
 		Set( "Graph.Node.RerouteDiffHighlightShading", new BOX_BRUSH( "/Graph/RerouteNode_DiffHighlightShading", FMargin(18.0f/64.0f) ) );
@@ -3751,50 +3751,50 @@ void FStarshipEditorStyle::FStyle::SetupGraphEditorStyles()
 		Set( "Graph.CompactNode.DiffHighlight", new BOX_BRUSH( "/Graph/MathNode_DiffHighlight", FMargin(18.0f/64.0f) ) );
 		Set( "Graph.CompactNode.DiffHighlightShading", new BOX_BRUSH( "/Graph/MathNode_DiffHighlightShading", FMargin(18.0f/64.0f) ) );
 
-		Set( "Graph.Node.CommentBubble", new BOX_BRUSH( "Old/Graph/CommentBubble", FMargin(8/32.0f) ) );
-		Set( "Graph.Node.CommentArrow", new IMAGE_BRUSH( "Old/Graph/CommentBubbleArrow", FVector2D(8,8) ) );
+		Set( "Graph.Node.CommentBubble", new BOX_BRUSH( "Old/Graph/CommentBubble", FMargin(8.f/32.0f) ) );
+		Set( "Graph.Node.CommentArrow", new IMAGE_BRUSH( "Old/Graph/CommentBubbleArrow", Icon8x8 ) );
 		Set( "Graph.Node.CommentFont", DEFAULT_FONT( "Regular", 10 ) );
-		Set( "Graph.Node.Comment.BubbleOffset", FMargin(8,0,0,0) );
-		Set( "Graph.Node.Comment.PinIconPadding", FMargin(0,2,0,0) );
-		Set("Graph.Node.Comment.Handle", new IMAGE_BRUSH_SVG("Starship/GraphEditors/Comment_Handle", FVector2D(16.0f, 16.0f)));
-		Set("Graph.Node.Comment.BubblePadding", FVector2D(3, 3));
-		Set("Graph.Node.Comment.BubbleWidgetMargin", FMargin(4, 4));
+		Set( "Graph.Node.Comment.BubbleOffset", FMargin(8.f,0.f,0.f,0.f) );
+		Set( "Graph.Node.Comment.PinIconPadding", FMargin(0.f,2.f,0.f,0.f) );
+		Set("Graph.Node.Comment.Handle", new IMAGE_BRUSH_SVG("Starship/GraphEditors/Comment_Handle", Icon16x16));
+		Set("Graph.Node.Comment.BubblePadding", FVector2f(3.f, 3.f));
+		Set("Graph.Node.Comment.BubbleWidgetMargin", FMargin(4.f, 4.f));
 
 		const FCheckBoxStyle CommentTitleButton = FCheckBoxStyle()
 			.SetCheckBoxType( ESlateCheckBoxType::CheckBox )
-			.SetUncheckedImage( IMAGE_BRUSH( "Icons/icon_Blueprint_CommentBubbleOn_16x", FVector2D(16,16), FLinearColor(1.f, 1.f, 1.f, 0.8f)))
-			.SetUncheckedHoveredImage( IMAGE_BRUSH( "Icons/icon_Blueprint_CommentBubbleOn_16x", FVector2D(16,16), FLinearColor(1.f, 1.f, 1.f, 0.9f)))
-			.SetUncheckedPressedImage( IMAGE_BRUSH( "Icons/icon_Blueprint_CommentBubbleOn_16x", FVector2D(16,16), FLinearColor(1.f, 1.f, 1.f, 1.f)))
-			.SetCheckedImage( IMAGE_BRUSH( "Icons/icon_Blueprint_CommentBubbleOn_16x", FVector2D(16,16), FLinearColor(1.f, 1.f, 1.f, 0.8f)))
-			.SetCheckedHoveredImage( IMAGE_BRUSH( "Icons/icon_Blueprint_CommentBubbleOn_16x", FVector2D(16,16), FLinearColor(1.f, 1.f, 1.f, 1.f)))
-			.SetCheckedPressedImage( IMAGE_BRUSH( "Icons/icon_Blueprint_CommentBubbleOff_16x", FVector2D(16,16), FLinearColor(1.f, 1.f, 1.f, 0.6f)));
+			.SetUncheckedImage( IMAGE_BRUSH( "Icons/icon_Blueprint_CommentBubbleOn_16x", Icon16x16, FLinearColor(1.f, 1.f, 1.f, 0.8f)))
+			.SetUncheckedHoveredImage(IMAGE_BRUSH("Icons/icon_Blueprint_CommentBubbleOn_16x", Icon16x16, FLinearColor(1.f, 1.f, 1.f, 0.9f)))
+			.SetUncheckedPressedImage(IMAGE_BRUSH("Icons/icon_Blueprint_CommentBubbleOn_16x", Icon16x16, FLinearColor(1.f, 1.f, 1.f, 1.f)))
+			.SetCheckedImage(IMAGE_BRUSH("Icons/icon_Blueprint_CommentBubbleOn_16x", Icon16x16, FLinearColor(1.f, 1.f, 1.f, 0.8f)))
+			.SetCheckedHoveredImage(IMAGE_BRUSH("Icons/icon_Blueprint_CommentBubbleOn_16x", Icon16x16, FLinearColor(1.f, 1.f, 1.f, 1.f)))
+			.SetCheckedPressedImage(IMAGE_BRUSH("Icons/icon_Blueprint_CommentBubbleOff_16x", Icon16x16, FLinearColor(1.f, 1.f, 1.f, 0.6f)));
 		Set( "CommentTitleButton", CommentTitleButton );
 
 		const FCheckBoxStyle CommentBubbleButton = FCheckBoxStyle()
 			.SetCheckBoxType( ESlateCheckBoxType::CheckBox )
-			.SetUncheckedImage( IMAGE_BRUSH( "Icons/icon_Blueprint_CommentBubbleOn_16x", FVector2D(10,10), FLinearColor(1.f, 1.f, 1.f, 0.5f)))
-			.SetUncheckedHoveredImage( IMAGE_BRUSH( "Icons/icon_Blueprint_CommentBubbleOn_16x", FVector2D(10,10), FLinearColor(1.f, 1.f, 1.f, 0.9f)))
-			.SetUncheckedPressedImage( IMAGE_BRUSH( "Icons/icon_Blueprint_CommentBubbleOn_16x", FVector2D(10,10), FLinearColor(1.f, 1.f, 1.f, 1.f)))
-			.SetCheckedImage( IMAGE_BRUSH( "Icons/icon_Blueprint_CommentBubbleOn_16x", FVector2D(10,10), FLinearColor(1.f, 1.f, 1.f, 0.8f)))
-			.SetCheckedHoveredImage( IMAGE_BRUSH( "Icons/icon_Blueprint_CommentBubbleOn_16x", FVector2D(10,10), FLinearColor(1.f, 1.f, 1.f, 1.f)))
-			.SetCheckedPressedImage( IMAGE_BRUSH( "Icons/icon_Blueprint_CommentBubbleOn_16x", FVector2D(10,10), FLinearColor(1.f, 1.f, 1.f, 0.6f)));
+			.SetUncheckedImage( IMAGE_BRUSH( "Icons/icon_Blueprint_CommentBubbleOn_16x", Icon10x10, FLinearColor(1.f, 1.f, 1.f, 0.5f)))
+			.SetUncheckedHoveredImage( IMAGE_BRUSH( "Icons/icon_Blueprint_CommentBubbleOn_16x", Icon10x10, FLinearColor(1.f, 1.f, 1.f, 0.9f)))
+			.SetUncheckedPressedImage(IMAGE_BRUSH("Icons/icon_Blueprint_CommentBubbleOn_16x", Icon10x10, FLinearColor(1.f, 1.f, 1.f, 1.f)))
+			.SetCheckedImage(IMAGE_BRUSH("Icons/icon_Blueprint_CommentBubbleOn_16x", Icon10x10, FLinearColor(1.f, 1.f, 1.f, 0.8f)))
+			.SetCheckedHoveredImage(IMAGE_BRUSH("Icons/icon_Blueprint_CommentBubbleOn_16x", Icon10x10, FLinearColor(1.f, 1.f, 1.f, 1.f)))
+			.SetCheckedPressedImage(IMAGE_BRUSH("Icons/icon_Blueprint_CommentBubbleOn_16x", Icon10x10, FLinearColor(1.f, 1.f, 1.f, 0.6f)));
 		Set( "CommentBubbleButton", CommentBubbleButton );
 
 		const FCheckBoxStyle CommentBubblePin = FCheckBoxStyle()
 			.SetCheckBoxType( ESlateCheckBoxType::CheckBox )
-			.SetUncheckedImage( IMAGE_BRUSH( "Icons/icon_Blueprint_CommentBubbleUnPin_16x", FVector2D(10,10), FLinearColor(1.f, 1.f, 1.f, 0.5f)))
-			.SetUncheckedHoveredImage( IMAGE_BRUSH( "Icons/icon_Blueprint_CommentBubbleUnPin_16x", FVector2D(10,10), FLinearColor(1.f, 1.f, 1.f, 0.9f)))
-			.SetUncheckedPressedImage( IMAGE_BRUSH( "Icons/icon_Blueprint_CommentBubblePin_16x", FVector2D(10,10), FLinearColor(1.f, 1.f, 1.f, 1.f)))
-			.SetCheckedImage( IMAGE_BRUSH( "Icons/icon_Blueprint_CommentBubblePin_16x", FVector2D(10,10), FLinearColor(1.f, 1.f, 1.f, 0.8f)))
-			.SetCheckedHoveredImage( IMAGE_BRUSH( "Icons/icon_Blueprint_CommentBubblePin_16x", FVector2D(10,10), FLinearColor(1.f, 1.f, 1.f, 1.f)))
-			.SetCheckedPressedImage( IMAGE_BRUSH( "Icons/icon_Blueprint_CommentBubbleUnPin_16x", FVector2D(10,10), FLinearColor(1.f, 1.f, 1.f, 0.6f)));
+			.SetUncheckedImage( IMAGE_BRUSH( "Icons/icon_Blueprint_CommentBubbleUnPin_16x", Icon10x10, FLinearColor(1.f, 1.f, 1.f, 0.5f)))
+			.SetUncheckedHoveredImage(IMAGE_BRUSH("Icons/icon_Blueprint_CommentBubbleUnPin_16x", Icon10x10, FLinearColor(1.f, 1.f, 1.f, 0.9f)))
+			.SetUncheckedPressedImage(IMAGE_BRUSH("Icons/icon_Blueprint_CommentBubblePin_16x", Icon10x10, FLinearColor(1.f, 1.f, 1.f, 1.f)))
+			.SetCheckedImage(IMAGE_BRUSH("Icons/icon_Blueprint_CommentBubblePin_16x", Icon10x10, FLinearColor(1.f, 1.f, 1.f, 0.8f)))
+			.SetCheckedHoveredImage(IMAGE_BRUSH("Icons/icon_Blueprint_CommentBubblePin_16x", Icon10x10, FLinearColor(1.f, 1.f, 1.f, 1.f)))
+			.SetCheckedPressedImage(IMAGE_BRUSH("Icons/icon_Blueprint_CommentBubbleUnPin_16x", Icon10x10, FLinearColor(1.f, 1.f, 1.f, 0.6f)));
 		Set( "CommentBubblePin", CommentBubblePin );
 
 
 		Set( "Graph.VarNode.Body", new BOX_BRUSH( "/Graph/VarNode_body", FMargin(16.f/64.f, 12.f/28.f) ) );
-		Set( "Graph.VarNode.ColorSpill", new IMAGE_BRUSH( "/Graph/VarNode_color_spill", FVector2D(132,28) ) );
+		Set( "Graph.VarNode.ColorSpill", new IMAGE_BRUSH( "/Graph/VarNode_color_spill", FVector2f(132.f,28.f) ) );
 		Set( "Graph.VarNode.Gloss", new BOX_BRUSH( "/Graph/VarNode_gloss", FMargin(16.f/64.f, 16.f/28.f, 16.f/64.f, 4.f/28.f) ) );
-		Set( "Graph.VarNode.IndicatorOverlay", new IMAGE_BRUSH("/Graph/IndicatorOverlay_color_spill", FVector2D(64.f, 28.f)));
+		Set( "Graph.VarNode.IndicatorOverlay", new IMAGE_BRUSH("/Graph/IndicatorOverlay_color_spill", FVector2f(64.f, 28.f)));
 		
 		Set( "Graph.VarNode.ShadowSelected", new BOX_BRUSH( "/Graph/VarNode_shadow_selected", FMargin(26.0f/64.0f) ) );
 		Set( "Graph.VarNode.Shadow", new BOX_BRUSH( "/Graph/VarNode_shadow", FMargin(26.0f/64.0f) ) );
@@ -3823,7 +3823,7 @@ void FStarshipEditorStyle::FStyle::SetupGraphEditorStyles()
 				FTextBlockStyle GraphStateNodeTitle = FTextBlockStyle(NormalText)
 					.SetFont( DEFAULT_FONT( "Bold", 14 ) )
 					.SetColorAndOpacity( FLinearColor(230.0f/255.0f,230.0f/255.0f,230.0f/255.0f) )
-					.SetShadowOffset( FVector2D( 2,2 ) )
+					.SetShadowOffset( FVector2f( 2.f,2.f) )
 					.SetShadowColorAndOpacity( FLinearColor(0.f,0.f,0.f, 0.7f) );
 				Set( "Graph.StateNode.NodeTitle", GraphStateNodeTitle );
 
@@ -3848,8 +3848,8 @@ void FStarshipEditorStyle::FStyle::SetupGraphEditorStyles()
 			{
 				FMargin TestMargin(16.f/64.f, 16.f/28.f, 16.f/64.f, 4.f/28.f);
 				Set( "Graph.TransitionNode.ColorSpill", new BOX_BRUSH( "/Persona/StateMachineEditor/Trans_Node_ColorSpill", TestMargin ) );
-				Set( "Graph.TransitionNode.Icon", new IMAGE_BRUSH( "/Persona/StateMachineEditor/Trans_Node_Icon", FVector2D(25,25) ) );
-				Set( "Graph.TransitionNode.Icon_Inertialization", new IMAGE_BRUSH( "/Persona/StateMachineEditor/Trans_Node_Icon_Inertialization", FVector2D(25,25) ) );
+				Set( "Graph.TransitionNode.Icon", new IMAGE_BRUSH( "/Persona/StateMachineEditor/Trans_Node_Icon", FVector2f(25.f,25.f) ) );
+				Set( "Graph.TransitionNode.Icon_Inertialization", new IMAGE_BRUSH( "/Persona/StateMachineEditor/Trans_Node_Icon_Inertialization", FVector2f(25.f,25.f) ) );
 			}
 
 			// Transition rule tooltip name 
@@ -3857,7 +3857,7 @@ void FStarshipEditorStyle::FStyle::SetupGraphEditorStyles()
 				Set( "Graph.TransitionNode.TooltipName", FTextBlockStyle(NormalText)
 					.SetFont( DEFAULT_FONT( "Bold", 12 ) )
 					.SetColorAndOpacity( FLinearColor(218.0f/255.0f,218.0f/255.0f,218.0f/255.0f) )
-					.SetShadowOffset( FVector2D(1.0f, 1.0f) )
+					.SetShadowOffset( FVector2f::UnitVector )
 					.SetShadowColorAndOpacity( FLinearColor(0.f,0.f,0.f, 0.7f) )
 				);
 			}
@@ -3867,7 +3867,7 @@ void FStarshipEditorStyle::FStyle::SetupGraphEditorStyles()
 				Set( "Graph.TransitionNode.TooltipRule", FTextBlockStyle(NormalText)
 					.SetFont( DEFAULT_FONT( "Bold", 8 ) )
 					.SetColorAndOpacity( FLinearColor(180.0f/255.0f,180.0f/255.0f,180.0f/255.0f) )
-					.SetShadowOffset( FVector2D(1.0f, 1.0f) )
+					.SetShadowOffset( FVector2f::UnitVector )
 					.SetShadowColorAndOpacity( FLinearColor(0.f,0.f,0.f, 0.7f) )
 				);
 			}
@@ -3880,8 +3880,8 @@ void FStarshipEditorStyle::FStyle::SetupGraphEditorStyles()
 			Set("Persona.RetargetManager.ImportantText", FTextBlockStyle(NormalText)
 				.SetFont(DEFAULT_FONT("Bold", 11))
 				.SetColorAndOpacity(FLinearColor(1.0f, 1.0f, 1.0f, 1.0f))
-				.SetShadowOffset(FVector2D(1, 1))
-				.SetShadowColorAndOpacity(FLinearColor(0, 0, 0, 0.9f)));
+				.SetShadowOffset(FVector2f::UnitVector)
+				.SetShadowColorAndOpacity(FLinearColor(0.f, 0.f, 0.f, 0.9f)));
 		}
 
 		// Behavior Tree Editor
@@ -3960,10 +3960,10 @@ void FStarshipEditorStyle::FStyle::SetupGraphEditorStyles()
 
 			Set( "BTEditor.DebuggerOverlay.Breakpoint.Disabled", new IMAGE_BRUSH_SVG( "Starship/Blueprints/Breakpoint_Disabled", Icon32x32 ) );
 			Set( "BTEditor.DebuggerOverlay.Breakpoint.Enabled", new IMAGE_BRUSH_SVG( "Starship/Blueprints/Breakpoint_Valid", Icon32x32 ) );
-			Set( "BTEditor.DebuggerOverlay.ActiveNodePointer", new IMAGE_BRUSH( "Old/Kismet2/IP_Normal", FVector2D(128,96)) );
-			Set( "BTEditor.DebuggerOverlay.SearchTriggerPointer", new IMAGE_BRUSH( "/BehaviorTree/SearchTriggerPointer", FVector2D(48,64)) );
-			Set( "BTEditor.DebuggerOverlay.FailedTriggerPointer", new IMAGE_BRUSH( "/BehaviorTree/FailedTriggerPointer", FVector2D(48,64)) );
-			Set( "BTEditor.DebuggerOverlay.BreakOnBreakpointPointer", new IMAGE_BRUSH( "Old/Kismet2/IP_Breakpoint", FVector2D(128,96)) );
+			Set( "BTEditor.DebuggerOverlay.ActiveNodePointer", new IMAGE_BRUSH( "Old/Kismet2/IP_Normal", FVector2f(128.f,96.f)) );
+			Set( "BTEditor.DebuggerOverlay.SearchTriggerPointer", new IMAGE_BRUSH( "/BehaviorTree/SearchTriggerPointer", FVector2f(48.f,64.f)) );
+			Set( "BTEditor.DebuggerOverlay.FailedTriggerPointer", new IMAGE_BRUSH( "/BehaviorTree/FailedTriggerPointer", FVector2f(48.f,64.f)) );
+			Set( "BTEditor.DebuggerOverlay.BreakOnBreakpointPointer", new IMAGE_BRUSH( "Old/Kismet2/IP_Breakpoint", FVector2f(128.f,96.f)) );
 
 			Set( "BTEditor.Blackboard.NewEntry", new IMAGE_BRUSH_SVG( "Starship/BehaviorTree/BlackboardNewKey", Icon20x20 ) );
 
@@ -3971,18 +3971,18 @@ void FStarshipEditorStyle::FStyle::SetupGraphEditorStyles()
 			Set( "BTEditor.SwitchToBlackboardMode", new IMAGE_BRUSH_SVG( "Starship/BehaviorTree/Blackboard_20", Icon20x20));
 
 			// Blackboard classes
-			Set( "ClassIcon.BlackboardKeyType_Bool", new FSlateRoundedBoxBrush( FLinearColor(0.300000f, 0.0f, 0.0f, 1.0f), 2.5f, FVector2D(16.f, 5.f) ) );
-			Set( "ClassIcon.BlackboardKeyType_Class", new FSlateRoundedBoxBrush( FLinearColor(0.1f, 0.0f, 0.5f, 1.0f), 2.5f, FVector2D(16.f, 5.f) ) );
-			Set( "ClassIcon.BlackboardKeyType_Enum", new FSlateRoundedBoxBrush( FLinearColor(0.0f, 0.160000f, 0.131270f, 1.0f), 2.5f, FVector2D(16.f, 5.f) ) );
-			Set( "ClassIcon.BlackboardKeyType_Float", new FSlateRoundedBoxBrush( FLinearColor(0.357667f, 1.0f, 0.060000f, 1.0f), 2.5f, FVector2D(16.f, 5.f) ) );
-			Set( "ClassIcon.BlackboardKeyType_Int", new FSlateRoundedBoxBrush( FLinearColor(0.013575f, 0.770000f, 0.429609f, 1.0f), 2.5f, FVector2D(16.f, 5.f) ) );
-			Set( "ClassIcon.BlackboardKeyType_Name", new FSlateRoundedBoxBrush( FLinearColor(0.607717f, 0.224984f, 1.0f, 1.0f), 2.5f, FVector2D(16.f, 5.f) ) );
-			Set( "ClassIcon.BlackboardKeyType_NativeEnum", new FSlateRoundedBoxBrush( FLinearColor(0.0f, 0.160000f, 0.131270f, 1.0f), 2.5f, FVector2D(16.f, 5.f) ) );
-			Set( "ClassIcon.BlackboardKeyType_Object", new FSlateRoundedBoxBrush( FLinearColor(0.0f, 0.4f, 0.910000f, 1.0f), 2.5f, FVector2D(16.f, 5.f) ) );
-			Set( "ClassIcon.BlackboardKeyType_Rotator", new FSlateRoundedBoxBrush( FLinearColor(0.353393f, 0.454175f, 1.0f, 1.0f), 2.5f, FVector2D(16.f, 5.f) ) );
-			Set( "ClassIcon.BlackboardKeyType_String", new FSlateRoundedBoxBrush( FLinearColor(1.0f, 0.0f, 0.660537f, 1.0f), 2.5f, FVector2D(16.f, 5.f) ) );
-			Set( "ClassIcon.BlackboardKeyType_Struct", new FSlateRoundedBoxBrush( FLinearColor(0.f, 0.349019f, 0.796070f, 1.0f), 2.5f, FVector2D(16.f, 5.f)));
-			Set( "ClassIcon.BlackboardKeyType_Vector", new FSlateRoundedBoxBrush( FLinearColor(1.0f, 0.591255f, 0.016512f, 1.0f), 2.5f, FVector2D(16.f, 5.f) ) );
+			Set( "ClassIcon.BlackboardKeyType_Bool", new FSlateRoundedBoxBrush( FLinearColor(0.300000f, 0.0f, 0.0f, 1.0f), 2.5f, FVector2f(16.f, 5.f) ) );
+			Set( "ClassIcon.BlackboardKeyType_Class", new FSlateRoundedBoxBrush( FLinearColor(0.1f, 0.0f, 0.5f, 1.0f), 2.5f, FVector2f(16.f, 5.f) ) );
+			Set( "ClassIcon.BlackboardKeyType_Enum", new FSlateRoundedBoxBrush( FLinearColor(0.0f, 0.160000f, 0.131270f, 1.0f), 2.5f, FVector2f(16.f, 5.f) ) );
+			Set( "ClassIcon.BlackboardKeyType_Float", new FSlateRoundedBoxBrush( FLinearColor(0.357667f, 1.0f, 0.060000f, 1.0f), 2.5f, FVector2f(16.f, 5.f) ) );
+			Set( "ClassIcon.BlackboardKeyType_Int", new FSlateRoundedBoxBrush( FLinearColor(0.013575f, 0.770000f, 0.429609f, 1.0f), 2.5f, FVector2f(16.f, 5.f) ) );
+			Set( "ClassIcon.BlackboardKeyType_Name", new FSlateRoundedBoxBrush( FLinearColor(0.607717f, 0.224984f, 1.0f, 1.0f), 2.5f, FVector2f(16.f, 5.f) ) );
+			Set( "ClassIcon.BlackboardKeyType_NativeEnum", new FSlateRoundedBoxBrush( FLinearColor(0.0f, 0.160000f, 0.131270f, 1.0f), 2.5f, FVector2f(16.f, 5.f) ) );
+			Set( "ClassIcon.BlackboardKeyType_Object", new FSlateRoundedBoxBrush( FLinearColor(0.0f, 0.4f, 0.910000f, 1.0f), 2.5f, FVector2f(16.f, 5.f) ) );
+			Set( "ClassIcon.BlackboardKeyType_Rotator", new FSlateRoundedBoxBrush( FLinearColor(0.353393f, 0.454175f, 1.0f, 1.0f), 2.5f, FVector2f(16.f, 5.f) ) );
+			Set( "ClassIcon.BlackboardKeyType_String", new FSlateRoundedBoxBrush( FLinearColor(1.0f, 0.0f, 0.660537f, 1.0f), 2.5f, FVector2f(16.f, 5.f) ) );
+			Set( "ClassIcon.BlackboardKeyType_Struct", new FSlateRoundedBoxBrush( FLinearColor(0.f, 0.349019f, 0.796070f, 1.0f), 2.5f, FVector2f(16.f, 5.f)));
+			Set( "ClassIcon.BlackboardKeyType_Vector", new FSlateRoundedBoxBrush( FLinearColor(1.0f, 0.591255f, 0.016512f, 1.0f), 2.5f, FVector2f(16.f, 5.f) ) );
 
 			Set( "BTEditor.Common.NewBlackboard", new IMAGE_BRUSH_SVG( "Starship/BehaviorTree/Blackboard_20", Icon20x20));
 			Set( "BTEditor.Graph.NewTask", new IMAGE_BRUSH_SVG( "Starship/Common/Tasks", Icon20x20));
@@ -4001,38 +4001,38 @@ void FStarshipEditorStyle::FStyle::SetupGraphEditorStyles()
 		// Visible on hover button for transition node
 		{
 			Set( "TransitionNodeButton.Normal", new FSlateNoResource() );
-			Set( "TransitionNodeButton.Hovered", new IMAGE_BRUSH( "/Persona/StateMachineEditor/Trans_Button_Hovered", FVector2D(12,25) ) );
-			Set( "TransitionNodeButton.Pressed", new IMAGE_BRUSH( "/Persona/StateMachineEditor/Trans_Button_Pressed", FVector2D(12,25) ) );
+			Set( "TransitionNodeButton.Hovered", new IMAGE_BRUSH( "/Persona/StateMachineEditor/Trans_Button_Hovered", FVector2f(12.f,25.f) ) );
+			Set( "TransitionNodeButton.Pressed", new IMAGE_BRUSH( "/Persona/StateMachineEditor/Trans_Button_Pressed", FVector2f(12.f,25.f) ) );
 		}
 
 		{
-			Set( "Graph.AnimationResultNode.Body", new IMAGE_BRUSH( "/Graph/Animation/AnimationNode_Result_128x", FVector2D(128, 128) ) );
+			Set( "Graph.AnimationResultNode.Body", new IMAGE_BRUSH( "/Graph/Animation/AnimationNode_Result_128x", Icon128x128 ) );
 			Set( "Graph.AnimationFastPathIndicator", new IMAGE_BRUSH( "/Graph/Animation/AnimationNode_FastPath", Icon32x32 ) );
 		}
 
 		// SoundCueEditor Graph Nodes
 		{
-			Set( "Graph.SoundResultNode.Body", new IMAGE_BRUSH( "/Graph/SoundCue_SpeakerIcon", FVector2D(144, 144) ) );
+			Set( "Graph.SoundResultNode.Body", new IMAGE_BRUSH( "/Graph/SoundCue_SpeakerIcon", FVector2f(144.f, 144.f) ) );
 		}
 
-		Set( "Graph.Node.NodeEntryTop", new IMAGE_BRUSH( "Old/Graph/NodeEntryTop", FVector2D(64,12) ) );
-		Set( "Graph.Node.NodeEntryBottom", new IMAGE_BRUSH( "Old/Graph/NodeEntryBottom", FVector2D(64,12) ) );
-		Set( "Graph.Node.NodeExitTop", new IMAGE_BRUSH( "Old/Graph/NodeExitTop", FVector2D(64,12) ) );
-		Set( "Graph.Node.NodeExitBottom", new IMAGE_BRUSH( "Old/Graph/NodeExitBottom", FVector2D(64,12) ) );
+		Set( "Graph.Node.NodeEntryTop", new IMAGE_BRUSH( "Old/Graph/NodeEntryTop", FVector2f(64.f,12.f) ) );
+		Set( "Graph.Node.NodeEntryBottom", new IMAGE_BRUSH( "Old/Graph/NodeEntryBottom", FVector2f(64.f,12.f) ) );
+		Set( "Graph.Node.NodeExitTop", new IMAGE_BRUSH( "Old/Graph/NodeExitTop", FVector2f(64.f,12.f) ) );
+		Set( "Graph.Node.NodeExitBottom", new IMAGE_BRUSH( "Old/Graph/NodeExitBottom", FVector2f(64.f,12.f) ) );
 
 		Set( "Graph.Node.NodeEntryShadow", new BOX_BRUSH( "Old/Graph/NodeEntryShadow", FMargin(5.f/80, 21.f/52) ) );
 		Set( "Graph.Node.NodeEntryShadowSelected", new BOX_BRUSH( "Old/Graph/NodeEntryShadowSelected", FMargin(5.f/80, 21.f/52) ) );
 		Set( "Graph.Node.NodeExitShadow", new BOX_BRUSH( "Old/Graph/NodeExitShadow", FMargin(5.f/80, 21.f/52) ) );
 		Set( "Graph.Node.NodeExitShadowSelected", new BOX_BRUSH( "Old/Graph/NodeExitShadowSelected", FMargin(5.f/80, 21.f/52) ) );
 
-		Set( "Graph.Node.Autoplay", new IMAGE_BRUSH( "Graph/Icons/Overlay_Autoplay", FVector2D(22,22) ) );
-		Set( "Graph.Node.Loop", new IMAGE_BRUSH( "Graph/Icons/Overlay_Loop", FVector2D(22,22) ) );
+		Set( "Graph.Node.Autoplay", new IMAGE_BRUSH( "Graph/Icons/Overlay_Autoplay", FVector2f(22.f,22.f) ) );
+		Set( "Graph.Node.Loop", new IMAGE_BRUSH( "Graph/Icons/Overlay_Loop", FVector2f(22.f,22.f) ) );
 
 		{
 			FTextBlockStyle GraphNodeTitle = FTextBlockStyle(NormalText)
 				.SetFont( DEFAULT_FONT( "Bold", 10 ) )
 				.SetColorAndOpacity( FStyleColors::Foreground )
-				.SetShadowOffset( FVector2D::ZeroVector )
+				.SetShadowOffset( FVector2f::ZeroVector )
 				.SetShadowColorAndOpacity( FLinearColor(0.f,0.f,0.f, 0.7f) );
 			Set( "Graph.Node.NodeTitle", GraphNodeTitle );
 
@@ -4058,7 +4058,7 @@ void FStarshipEditorStyle::FStyle::SetupGraphEditorStyles()
 			Set( "Graph.Node.NodeTitleExtraLines", FTextBlockStyle(NormalText)
 				.SetFont( DEFAULT_FONT( "Italic", 9 ) )
 				.SetColorAndOpacity( FLinearColor(218.0f/255.0f,218.0f/255.0f,96.0f/255.0f, 0.5f) )
-				.SetShadowOffset( FVector2D::ZeroVector )
+				.SetShadowOffset( FVector2f::ZeroVector )
 				.SetShadowColorAndOpacity( FLinearColor(0.f,0.f,0.f, 0.7f) )
 			);
 
@@ -4079,7 +4079,7 @@ void FStarshipEditorStyle::FStyle::SetupGraphEditorStyles()
 			FTextBlockStyle GraphCommentBlockTitle = FTextBlockStyle(NormalText)
 				.SetFont( DEFAULT_FONT( "Bold", 18 ) )
 				.SetColorAndOpacity( FLinearColor(218.0f/255.0f,218.0f/255.0f,218.0f/255.0f) )
-				.SetShadowOffset( FVector2D(1.5f, 1.5f) )
+				.SetShadowOffset( FVector2f(1.5f, 1.5f) )
 				.SetShadowColorAndOpacity( FLinearColor(0.f,0.f,0.f, 0.7f) );
 			Set( "Graph.CommentBlock.Title", GraphCommentBlockTitle );
 
@@ -4100,21 +4100,21 @@ void FStarshipEditorStyle::FStyle::SetupGraphEditorStyles()
 			Set( "Graph.CompactNode.Title", FTextBlockStyle(NormalText)
 				.SetFont( DEFAULT_FONT( "BoldCondensed", 20 ) )
 				.SetColorAndOpacity( FLinearColor(1.0f, 1.0f, 1.0f, 0.5f) )
-				.SetShadowOffset( FVector2D::ZeroVector )
+				.SetShadowOffset( FVector2f::ZeroVector )
 				.SetShadowColorAndOpacity( FLinearColor::White )
 			);
 
 			Set( "Graph.ArrayCompactNode.Title", FTextBlockStyle(NormalText)
 				.SetFont( DEFAULT_FONT( "BoldCondensed", 20 ) )
 				.SetColorAndOpacity( FLinearColor(1.0f, 1.0f, 1.0f, 0.5f) ) //218.0f/255.0f, 218.0f/255.0f, 218.0f/255.0f, 0.25f) )
-				.SetShadowOffset( FVector2D::ZeroVector )
+				.SetShadowOffset( FVector2f::ZeroVector )
 				.SetShadowColorAndOpacity( FLinearColor::White )
 				);
 
 			Set( "Graph.Node.PinName", FTextBlockStyle(NormalText)
 				.SetFont( DEFAULT_FONT( "Regular", 9 ) )
 				.SetColorAndOpacity( FLinearColor(218.0f/255.0f,218.0f/255.0f,218.0f/255.0f) )
-				.SetShadowOffset( FVector2D::ZeroVector )
+				.SetShadowOffset( FVector2f::ZeroVector )
 				.SetShadowColorAndOpacity( FLinearColor(0.8f,0.8f,0.8f, 0.5) )
 			);
 
@@ -4123,7 +4123,7 @@ void FStarshipEditorStyle::FStyle::SetupGraphEditorStyles()
 				FTextBlockStyle InlineEditableTextBlockReadOnly = FTextBlockStyle(NormalText)
 					.SetFont(DEFAULT_FONT("Regular", 9))
 					.SetColorAndOpacity(FLinearColor(218.0f / 255.0f, 218.0f / 255.0f, 218.0f / 255.0f))
-					.SetShadowOffset(FVector2D::ZeroVector)
+					.SetShadowOffset(FVector2f::ZeroVector)
 					.SetShadowColorAndOpacity(FLinearColor(0.8f, 0.8f, 0.8f, 0.5));
 
 				FEditableTextBoxStyle InlineEditableTextBlockEditable = FEditableTextBoxStyle()
@@ -4166,52 +4166,52 @@ void FStarshipEditorStyle::FStyle::SetupGraphEditorStyles()
 			Set( "Graph.Node.DefaultPinName", FTextBlockStyle(NormalText)
 				.SetFont( DEFAULT_FONT( "Italic", 9 ) )
 				.SetColorAndOpacity( FLinearColor(218.0f/255.0f,218.0f/255.0f,218.0f/255.0f) )
-				.SetShadowOffset( FVector2D::ZeroVector )
+				.SetShadowOffset( FVector2f::ZeroVector )
 				.SetShadowColorAndOpacity( FLinearColor(0.8f,0.8f,0.8f, 0.5) )
 			);
 		}
-		Set( "Graph.Pin.DefaultPinSeparator", new IMAGE_BRUSH( "/Graph/PinSeparator", FVector2D(64,8) ) );
+		Set( "Graph.Pin.DefaultPinSeparator", new IMAGE_BRUSH( "/Graph/PinSeparator", FVector2f(64.f,8.f) ) );
 
 		/** Original Pin Styles */
-		Set( "Graph.Pin.Connected", new IMAGE_BRUSH( "/Graph/Pin_connected", FVector2D(11,11) ) );
-		Set( "Graph.Pin.Disconnected", new IMAGE_BRUSH( "/Graph/Pin_disconnected", FVector2D(11,11) ) );
-		Set( "Graph.ArrayPin.Connected", new IMAGE_BRUSH( "/Graph/ArrayPin_connected", FVector2D(11,11) ) );
-		Set( "Graph.ArrayPin.Disconnected", new IMAGE_BRUSH( "/Graph/ArrayPin_disconnected", FVector2D(11,11) ) );
-		Set( "Graph.RefPin.Connected", new IMAGE_BRUSH( "/Graph/RefPin_connected", FVector2D(11,11) ) );
-		Set( "Graph.RefPin.Disconnected", new IMAGE_BRUSH( "/Graph/RefPin_disconnected", FVector2D(11,11) ) );
+		Set( "Graph.Pin.Connected", new IMAGE_BRUSH( "/Graph/Pin_connected", FVector2f(11.f,11.f) ) );
+		Set( "Graph.Pin.Disconnected", new IMAGE_BRUSH( "/Graph/Pin_disconnected", FVector2f(11.f,11.f) ) );
+		Set( "Graph.ArrayPin.Connected", new IMAGE_BRUSH( "/Graph/ArrayPin_connected", FVector2f(11.f,11.f) ) );
+		Set( "Graph.ArrayPin.Disconnected", new IMAGE_BRUSH( "/Graph/ArrayPin_disconnected", FVector2f(11.f,11.f) ) );
+		Set( "Graph.RefPin.Connected", new IMAGE_BRUSH( "/Graph/RefPin_connected", FVector2f(11.f,11.f) ) );
+		Set( "Graph.RefPin.Disconnected", new IMAGE_BRUSH( "/Graph/RefPin_disconnected", FVector2f(11.f,11.f) ) );
 
-		Set("Graph.Pin.CopyNodePinLeft_Connected", new IMAGE_BRUSH("/Graph/CopyNodePinLeft_connected", FVector2D(12, 24)));
-		Set("Graph.Pin.CopyNodePinLeft_Disconnected", new IMAGE_BRUSH("/Graph/CopyNodePinLeft_disconnected", FVector2D(12, 24)));
+		Set("Graph.Pin.CopyNodePinLeft_Connected", new IMAGE_BRUSH("/Graph/CopyNodePinLeft_connected", FVector2f(12.f, 24.f)));
+		Set("Graph.Pin.CopyNodePinLeft_Disconnected", new IMAGE_BRUSH("/Graph/CopyNodePinLeft_disconnected", FVector2f(12.f, 24.f)));
 
-		Set("Graph.Pin.CopyNodePinRight_Connected", new IMAGE_BRUSH("/Graph/CopyNodePinRight_connected", FVector2D(12, 24)));
-		Set("Graph.Pin.CopyNodePinRight_Disconnected", new IMAGE_BRUSH("/Graph/CopyNodePinRight_disconnected", FVector2D(12, 24)));
+		Set("Graph.Pin.CopyNodePinRight_Connected", new IMAGE_BRUSH("/Graph/CopyNodePinRight_connected", FVector2f(12.f, 24.f)));
+		Set("Graph.Pin.CopyNodePinRight_Disconnected", new IMAGE_BRUSH("/Graph/CopyNodePinRight_disconnected", FVector2f(12.f, 24.f)));
 
 		/** Variant A Pin Styles */
-		Set( "Graph.Pin.Connected_VarA", new IMAGE_BRUSH( "/Graph/Pin_connected_VarA", FVector2D(15,11)) );
-		Set( "Graph.Pin.Disconnected_VarA", new IMAGE_BRUSH( "/Graph/Pin_disconnected_VarA", FVector2D(15,11)) );
+		Set( "Graph.Pin.Connected_VarA", new IMAGE_BRUSH( "/Graph/Pin_connected_VarA", FVector2f(15.f,11.f)) );
+		Set( "Graph.Pin.Disconnected_VarA", new IMAGE_BRUSH( "/Graph/Pin_disconnected_VarA", FVector2f(15.f,11.f)) );
 
-		Set( "Graph.DelegatePin.Connected", new IMAGE_BRUSH( "/Graph/DelegatePin_Connected", FVector2D(11,11) ) );
-		Set( "Graph.DelegatePin.Disconnected", new IMAGE_BRUSH( "/Graph/DelegatePin_Disconnected", FVector2D(11,11) ) );
+		Set( "Graph.DelegatePin.Connected", new IMAGE_BRUSH( "/Graph/DelegatePin_Connected", FVector2f(11.f,11.f) ) );
+		Set( "Graph.DelegatePin.Disconnected", new IMAGE_BRUSH( "/Graph/DelegatePin_Disconnected", FVector2f(11.f,11.f) ) );
 
-		Set( "Graph.Replication.AuthorityOnly", new IMAGE_BRUSH( "/Graph/AuthorityOnly", FVector2D(32,32) ) );
-		Set( "Graph.Replication.ClientEvent", new IMAGE_BRUSH( "/Graph/ClientEvent", FVector2D(32,32) ) );
-		Set( "Graph.Replication.Replicated", new IMAGE_BRUSH( "/Graph/Replicated", FVector2D(32,32) ) );
+		Set( "Graph.Replication.AuthorityOnly", new IMAGE_BRUSH( "/Graph/AuthorityOnly", Icon32x32 ) );
+		Set("Graph.Replication.ClientEvent", new IMAGE_BRUSH("/Graph/ClientEvent", Icon32x32));
+		Set( "Graph.Replication.Replicated", new IMAGE_BRUSH( "/Graph/Replicated", Icon32x32 ) );
 
-		Set("Graph.Editor.EditorOnlyIcon", new IMAGE_BRUSH("/Graph/EditorOnly", FVector2D(32, 32)));
+		Set("Graph.Editor.EditorOnlyIcon", new IMAGE_BRUSH("/Graph/EditorOnly", Icon32x32));
 
-		Set( "Graph.Event.InterfaceEventIcon", new IMAGE_BRUSH("/Graph/InterfaceEventIcon", FVector2D(32,32) ) );
+		Set( "Graph.Event.InterfaceEventIcon", new IMAGE_BRUSH("/Graph/InterfaceEventIcon", Icon32x32 ) );
 
-		Set( "Graph.Latent.LatentIcon", new IMAGE_BRUSH("/Graph/LatentIcon", FVector2D(32,32) ) );
-		Set( "Graph.Message.MessageIcon", new IMAGE_BRUSH("/Graph/MessageIcon", FVector2D(32,32) ) );
-		Set( "Graph.Function.FunctionParameterIcon", new IMAGE_BRUSH_SVG("/Starship/GraphEditors/FunctionInputParameter", FVector2D(20.0f, 20.0f) ) );
-		Set( "Graph.Function.FunctionLocalVariableIcon", new IMAGE_BRUSH_SVG("/Starship/GraphEditors/FunctionLocalVariable", FVector2D(20.0f, 20.0f) ) );
+		Set( "Graph.Latent.LatentIcon", new IMAGE_BRUSH("/Graph/LatentIcon", Icon32x32 ) );
+		Set( "Graph.Message.MessageIcon", new IMAGE_BRUSH("/Graph/MessageIcon", Icon32x32 ) );
+		Set( "Graph.Function.FunctionParameterIcon", new IMAGE_BRUSH_SVG("/Starship/GraphEditors/FunctionInputParameter", FVector2f(20.0f, 20.0f) ) );
+		Set( "Graph.Function.FunctionLocalVariableIcon", new IMAGE_BRUSH_SVG("/Starship/GraphEditors/FunctionLocalVariable", FVector2f(20.0f, 20.0f) ) );
 
 		Set( "Graph.ExecPin.Connected", new IMAGE_BRUSH( "Old/Graph/ExecPin_Connected", Icon12x16 ) );
 		Set( "Graph.ExecPin.Disconnected", new IMAGE_BRUSH( "Old/Graph/ExecPin_Disconnected", Icon12x16 ) );
 		Set( "Graph.ExecPin.ConnectedHovered", new IMAGE_BRUSH( "Old/Graph/ExecPin_Connected", Icon12x16, FLinearColor(0.8f,0.8f,0.8f) ) );
 		Set( "Graph.ExecPin.DisconnectedHovered", new IMAGE_BRUSH( "Old/Graph/ExecPin_Disconnected", Icon12x16, FLinearColor(0.8f,0.8f,0.8f) ) );
 
-		const FVector2D Icon15x28(15.0f, 28.0f);
+		const FVector2f Icon15x28(15.0f, 28.0f);
 		Set("Graph.PosePin.Connected", new IMAGE_BRUSH_SVG("Starship/AnimationBlueprintEditor/AnimationGraphPose", Icon16x16));
 		Set("Graph.PosePin.Disconnected", new IMAGE_BRUSH_SVG("Starship/AnimationBlueprintEditor/AnimationGraphPoseDisconnected", Icon16x16));
 		Set("Graph.PosePin.ConnectedHovered", new IMAGE_BRUSH_SVG("Starship/AnimationBlueprintEditor/AnimationGraphPose", Icon16x16, FLinearColor(0.8f, 0.8f, 0.8f)));
@@ -4225,14 +4225,14 @@ void FStarshipEditorStyle::FStyle::SetupGraphEditorStyles()
 
 		Set( "Graph.WatchedPinIcon_Pinned", new IMAGE_BRUSH( "Old/Graph/WatchedPinIcon_Pinned", Icon16x16 ) );
 
-		Set( "Graph.Pin.BackgroundHovered", new IMAGE_BRUSH( "/Graph/Pin_hover_cue", FVector2D(32,8)));
+		Set( "Graph.Pin.BackgroundHovered", new IMAGE_BRUSH( "/Graph/Pin_hover_cue", FVector2f(32.f,8.f)));
 		Set( "Graph.Pin.Background", new FSlateNoResource() );
-		Set( "Graph.Pin.DiffHighlight", new IMAGE_BRUSH( "/Graph/Pin_DiffHighlight", FVector2D(32,8)));
+		Set( "Graph.Pin.DiffHighlight", new IMAGE_BRUSH( "/Graph/Pin_DiffHighlight", FVector2f(32.f,8.f)));
 
 		Set( "Graph.Pin.ObjectSet", new IMAGE_BRUSH( "Old/Graph/Pin_ObjectSet", Icon12x12 ) );
 		Set( "Graph.Pin.ObjectEmpty", new IMAGE_BRUSH( "Old/Graph/Pin_ObjectEmpty", Icon12x12 ) );
 
-		Set("Graph.Pin.Dummy", new IMAGE_BRUSH("/Graph/Pin_dummy", FVector2D(15, 11)));
+		Set("Graph.Pin.Dummy", new IMAGE_BRUSH("/Graph/Pin_dummy", FVector2f(15.f, 11.f)));
 
 		Set( "Graph.ConnectorFeedback.Border", new BOX_BRUSH( "Old/Menu_Background", FMargin(8.0f/64.0f) ) );
 		Set( "Graph.ConnectorFeedback.OK", new CORE_IMAGE_BRUSH_SVG("Starship/Common/check-circle", Icon16x16, FStyleColors::AccentGreen));
@@ -4245,32 +4245,32 @@ void FStarshipEditorStyle::FStyle::SetupGraphEditorStyles()
 		{
 			Set( "Graph.CornerText", FTextBlockStyle(NormalText)
 				.SetFont( DEFAULT_FONT( "BoldCondensed", 48 ) )
-				.SetColorAndOpacity( FLinearColor(0.8, 0.8f, 0.8f, 0.2f) )
-				.SetShadowOffset( FVector2D::ZeroVector )
+				.SetColorAndOpacity( FLinearColor(0.8f, 0.8f, 0.8f, 0.2f) )
+				.SetShadowOffset( FVector2f::ZeroVector )
 			);
 
 			Set("Graph.WarningText", FTextBlockStyle(NormalText)
 				.SetFont(DEFAULT_FONT("BoldCondensed", 20))
-				.SetColorAndOpacity(FLinearColor(0.8, 0.5f, 0.07f, 1.0f))
-				.SetShadowOffset(FVector2D(1, 1))
+				.SetColorAndOpacity(FLinearColor(0.8f, 0.5f, 0.07f, 1.0f))
+				.SetShadowOffset(FVector2f::UnitVector)
 			);
 
 			Set( "Graph.SimulatingText", FTextBlockStyle(NormalText)
 				.SetFont( DEFAULT_FONT( "BoldCondensed", 48 ) )
-				.SetColorAndOpacity( FLinearColor(0.8, 0.8f, 0.0f, 0.2f) )
-				.SetShadowOffset( FVector2D::ZeroVector )
+				.SetColorAndOpacity( FLinearColor(0.8f, 0.8f, 0.0f, 0.2f) )
+				.SetShadowOffset( FVector2f::ZeroVector )
 			);
 
 			Set( "GraphPreview.CornerText", FTextBlockStyle(NormalText)
 				.SetFont( DEFAULT_FONT( "BoldCondensed", 16 ) )
-				.SetColorAndOpacity( FLinearColor(0.8, 0.8f, 0.8f, 0.2f) )
-				.SetShadowOffset( FVector2D::ZeroVector )
+				.SetColorAndOpacity( FLinearColor(0.8f, 0.8f, 0.8f, 0.2f) )
+				.SetShadowOffset( FVector2f::ZeroVector )
 			);
 
 			Set( "Graph.InstructionText", FTextBlockStyle(NormalText)
 				.SetFont( DEFAULT_FONT( "BoldCondensed", 24 ) )
 				.SetColorAndOpacity( FLinearColor(1.f, 1.f, 1.f, 0.6f) )
-				.SetShadowOffset( FVector2D::ZeroVector )
+				.SetShadowOffset( FVector2f::ZeroVector )
 			);
 
 			Set( "Graph.InstructionBackground", new BOX_BRUSH("Common/GroupBorder", FMargin(4.0f / 16.0f), FLinearColor(0.1f, 0.1f, 0.1f, 0.7f)) );
@@ -4466,7 +4466,7 @@ void FStarshipEditorStyle::FStyle::SetupGraphEditorStyles()
 			.SetFont(DEFAULT_FONT("Bold", 10))
 			.SetColorAndOpacity(FLinearColor(1.0f, 1.0f, 1.0f))
 			.SetHighlightColor(FLinearColor(1.0f, 1.0f, 1.0f))
-			.SetShadowOffset(FVector2D(1, 1))
+			.SetShadowOffset(FVector2f::UnitVector)
 			.SetShadowColorAndOpacity(FLinearColor(0, 0, 0, 0.9f)));
 
 		Set("SCSEditor.TreePanel", new FSlateNoResource());
@@ -4483,13 +4483,13 @@ void FStarshipEditorStyle::FStyle::SetupGraphEditorStyles()
 		Set("SCSEditor.ComponentTooltip.Label",
 			FTextBlockStyle(NormalText)
 			.SetColorAndOpacity(FLinearColor(0.075f, 0.075f, 0.075f))
-			.SetShadowOffset(FVector2D(1.0f, 1.0f))
+			.SetShadowOffset(FVector2f::UnitVector)
 			.SetShadowColorAndOpacity(FLinearColor(0.8f, 0.8f, 0.8f))
 		);
 		Set("SCSEditor.ComponentTooltip.ImportantLabel",
 			FTextBlockStyle(NormalText)
 			.SetColorAndOpacity(FLinearColor(0.05f, 0.05f, 0.05f))
-			.SetShadowOffset(FVector2D(1.0f, 1.0f))
+			.SetShadowOffset(FVector2f::UnitVector)
 			.SetShadowColorAndOpacity(FLinearColor(0.8f, 0.8f, 0.8f))
 		);
 
@@ -4498,14 +4498,14 @@ void FStarshipEditorStyle::FStyle::SetupGraphEditorStyles()
 			FTextBlockStyle(NormalText)
 			.SetFont(DEFAULT_FONT("Bold", 10))
 			.SetColorAndOpacity(FLinearColor(0.0f, 0.0f, 0.0f))
-			.SetShadowOffset(FVector2D(1.0f, 1.0f))
+			.SetShadowOffset(FVector2f::UnitVector)
 			.SetShadowColorAndOpacity(FLinearColor(0.8f, 0.8f, 0.8f))
 		);
 		Set("SCSEditor.ComponentTooltip.ImportantValue",
 			FTextBlockStyle(NormalText)
 			.SetFont(DEFAULT_FONT("Bold", 10))
 			.SetColorAndOpacity(FLinearColor(0.3f, 0.0f, 0.0f))
-			.SetShadowOffset(FVector2D(1.0f, 1.0f))
+			.SetShadowOffset(FVector2f::UnitVector)
 			.SetShadowColorAndOpacity(FLinearColor(0.8f, 0.8f, 0.8f))
 		);
 
@@ -4513,7 +4513,7 @@ void FStarshipEditorStyle::FStyle::SetupGraphEditorStyles()
 			FTextBlockStyle(NormalText)
 			.SetFont(DEFAULT_FONT("Italic", 10))
 			.SetColorAndOpacity(FLinearColor(0.1f, 0.1f, 0.1f))
-			.SetShadowOffset(FVector2D(1.0f, 1.0f))
+			.SetShadowOffset(FVector2f::UnitVector)
 			.SetShadowColorAndOpacity(FLinearColor(0.8f, 0.8f, 0.8f))
 		);
 	}
@@ -4543,7 +4543,7 @@ void FStarshipEditorStyle::FStyle::SetupGraphEditorStyles()
 		Set("Persona.MontageEditor.ChildMontageInstruction", FTextBlockStyle(NormalText)
 			.SetFont(DEFAULT_FONT("BoldCondensed", 14))
 			.SetColorAndOpacity(FLinearColor(1.f, 1.f, 1.f, 1.0f))
-			.SetShadowOffset(FVector2D::ZeroVector)
+			.SetShadowOffset(FVector2f::ZeroVector)
 		);
 	}
 #endif // WITH_EDITOR || (IS_PROGRAM && WITH_UNREAL_DEVELOPER_TOOLS)
@@ -4784,8 +4784,8 @@ void FStarshipEditorStyle::FStyle::SetupLevelEditorStyle()
 
 		Set( "LevelEditor.Create", new IMAGE_BRUSH( "Icons/icon_Mode_Placement_40px", Icon40x40 ) );
 		Set( "LevelEditor.Create.Small", new IMAGE_BRUSH( "Icons/icon_Mode_Placement_40px", Icon20x20 ) );
-		Set( "LevelEditor.Create.OutlineHoriz", new IMAGE_BRUSH( "Common/WorkingFrame_Marquee", FVector2D(34.0f, 3.0f), FLinearColor::White, ESlateBrushTileType::Horizontal) );
-		Set( "LevelEditor.Create.OutlineVert", new IMAGE_BRUSH( "Common/WorkingFrame_Marquee_Vert", FVector2D(3.0f, 34.0f), FLinearColor::White, ESlateBrushTileType::Vertical) );
+		Set( "LevelEditor.Create.OutlineHoriz", new IMAGE_BRUSH( "Common/WorkingFrame_Marquee", FVector2f(34.0f, 3.0f), FLinearColor::White, ESlateBrushTileType::Horizontal) );
+		Set( "LevelEditor.Create.OutlineVert", new IMAGE_BRUSH( "Common/WorkingFrame_Marquee_Vert", FVector2f(3.0f, 34.0f), FLinearColor::White, ESlateBrushTileType::Vertical) );
 
 		Set( "LevelEditor.Tab", new IMAGE_BRUSH_SVG("Starship/AssetIcons/World_16", Icon16x16));
 		Set( "LevelEditor.AssetColor", FColor(255, 156, 0));
@@ -4807,7 +4807,7 @@ void FStarshipEditorStyle::FStyle::SetupLevelEditorStyle()
 
 		Set( "ToolPalette.DockingLabel", FTextBlockStyle(NormalText)
 			.SetFont( DEFAULT_FONT( "Regular", 9 ) ) 
-			.SetShadowOffset(FVector2D(1, 1))
+			.SetShadowOffset(FVector2f::UnitVector)
 			.SetShadowColorAndOpacity(FLinearColor(0, 0, 0, 0.9f))
 		);
 
@@ -4887,7 +4887,7 @@ void FStarshipEditorStyle::FStyle::SetupLevelEditorStyle()
 
 		Set( "PlacementBrowser.Asset.Background", new FSlateRoundedBoxBrush(FStyleColors::Recessed, 6.f));
 		Set( "PlacementBrowser.Asset.LabelBack", new BOX_BRUSH("Starship/PlacementBrowser/LabelBack_18x", 6.f/18.f, FStyleColors::Dropdown));
-		Set( "PlacementBrowser.Asset.ThumbnailBackground", new FSlateRoundedBoxBrush(FStyleColors::Dropdown, 6 ));
+		Set( "PlacementBrowser.Asset.ThumbnailBackground", new FSlateRoundedBoxBrush(FStyleColors::Dropdown, 6.f ));
 
 		FLinearColor DimBackground = FLinearColor( FColor( 64, 64, 64 ) );
 		FLinearColor DimBackgroundHover = FLinearColor( FColor( 50, 50, 50 ) );
@@ -4906,8 +4906,8 @@ void FStarshipEditorStyle::FStyle::SetupLevelEditorStyle()
 		Set( "PlacementBrowser.Tab.Text", FTextBlockStyle( NormalText )
 			.SetFont( DEFAULT_FONT( "Bold", 10 ) )
 			.SetColorAndOpacity( FLinearColor( 1.0f, 1.0f, 1.0f, 0.9f ) )
-			.SetShadowOffset( FVector2D( 1, 1 ) )
-			.SetShadowColorAndOpacity( FLinearColor( 0, 0, 0, 0.9f ) ) );
+			.SetShadowOffset( FVector2f::UnitVector )
+			.SetShadowColorAndOpacity( FLinearColor( 0.f, 0.f, 0.f, 0.9f ) ) );
 
 		Set( "PlacementBrowser.Asset.Name", FTextBlockStyle( NormalText )
 			.SetFont( DEFAULT_FONT( "Regular", 10 ) )
@@ -4916,11 +4916,11 @@ void FStarshipEditorStyle::FStyle::SetupLevelEditorStyle()
 		Set( "PlacementBrowser.Asset.Type", FTextBlockStyle( NormalText )
 			.SetFont( DEFAULT_FONT( "Regular", 8 ) )
 			.SetColorAndOpacity( FLinearColor( 0.8f, 0.8f, 0.8f, 0.9f ) )
-			.SetShadowOffset( FVector2D( 1, 1 ) )
-			.SetShadowColorAndOpacity( FLinearColor( 0, 0, 0, 0.9f ) ) );
+			.SetShadowOffset( FVector2f::UnitVector )
+			.SetShadowColorAndOpacity( FLinearColor( 0.f, 0.f, 0.f, 0.9f ) ) );
 
 		Set( "PlacementBrowser.ActiveTabNub", new IMAGE_BRUSH( "Icons/TabTriangle_24x", Icon24x24, FLinearColor( FColor( 42, 42, 42 ) ) ) );
-		Set( "PlacementBrowser.ActiveTabBar", new IMAGE_BRUSH( "Common/Selection", FVector2D(2.0f, 2.0f), SelectionColor ) );
+		Set( "PlacementBrowser.ActiveTabBar", new IMAGE_BRUSH( "Common/Selection", FVector2f(2.0f, 2.0f), SelectionColor ) );
 
 		Set( "PlacementBrowser.ShowAllContent", new IMAGE_BRUSH( "Icons/icon_Placement_AllContent_20px", Icon20x20 ) );
 		Set( "PlacementBrowser.ShowAllContent.Small", new IMAGE_BRUSH( "Icons/icon_Placement_AllContent_20px", Icon20x20 ) );
@@ -5005,7 +5005,7 @@ void FStarshipEditorStyle::FStyle::SetupLevelEditorStyle()
 
 	{
 
-		Set( "AssetDeleteDialog.Background", new IMAGE_BRUSH( "Common/Selection", Icon8x8, FLinearColor( 0.016, 0.016, 0.016 ) ) );
+		Set( "AssetDeleteDialog.Background", new IMAGE_BRUSH( "Common/Selection", Icon8x8, FLinearColor( 0.016f, 0.016f, 0.016f ) ) );
 	}
 
 	// Level editor tool box icons
@@ -5015,8 +5015,8 @@ void FStarshipEditorStyle::FStyle::SetupLevelEditorStyle()
 
 	// Editor viewport layout command icons
 	{
-		const FVector2D IconLayoutSize(47.0f, 37.0f);
-		const FVector2D IconLayoutSizeSmall(47.0f, 37.0f);		// small version set to same size as these are in their own menu and don't clutter the UI
+		const FVector2f IconLayoutSize(47.0f, 37.0f);
+		const FVector2f IconLayoutSizeSmall(47.0f, 37.0f);		// small version set to same size as these are in their own menu and don't clutter the UI
 
 		Set("EditorViewport.ViewportConfig_OnePane", new IMAGE_BRUSH_SVG("Starship/ViewportLayout/Layouts1Pane", IconLayoutSize));
 		Set("EditorViewport.ViewportConfig_OnePane.Small", new IMAGE_BRUSH_SVG("Starship/ViewportLayout/Layouts1Pane", IconLayoutSizeSmall));
@@ -5046,8 +5046,8 @@ void FStarshipEditorStyle::FStyle::SetupLevelEditorStyle()
 
 	// Level viewport layout command icons
 	{
-		const FVector2D IconLayoutSize(47.0f, 37.0f);
-		const FVector2D IconLayoutSizeSmall(47.0f, 37.0f);		// small version set to same size as these are in their own menu and don't clutter the UI
+		const FVector2f IconLayoutSize(47.0f, 37.0f);
+		const FVector2f IconLayoutSizeSmall(47.0f, 37.0f);		// small version set to same size as these are in their own menu and don't clutter the UI
 
 		Set("LevelViewport.ViewportConfig_OnePane", new IMAGE_BRUSH_SVG("Starship/ViewportLayout/Layouts1Pane", IconLayoutSize));
 		Set("LevelViewport.ViewportConfig_OnePane.Small", new IMAGE_BRUSH_SVG("Starship/ViewportLayout/Layouts1Pane", IconLayoutSizeSmall));
@@ -5084,12 +5084,12 @@ void FStarshipEditorStyle::FStyle::SetupLevelEditorStyle()
 
 	// Level editor status bar
 	{
-		Set( "TransformSettings.RelativeCoordinateSettings", new IMAGE_BRUSH( "Icons/icon_axis_16px", FVector2D( 16, 16 ) ) );
+		Set( "TransformSettings.RelativeCoordinateSettings", new IMAGE_BRUSH( "Icons/icon_axis_16px", Icon16x16 ) );
 	}
 
 	// Mesh Proxy Window
 	{
-		Set("MeshProxy.SimplygonLogo", new IMAGE_BRUSH( "Icons/SimplygonBanner_Sml", FVector2D(174, 36) ) );
+		Set("MeshProxy.SimplygonLogo", new IMAGE_BRUSH( "Icons/SimplygonBanner_Sml", FVector2f(174.f, 36.f) ) );
 	}
 #endif // WITH_EDITOR
 
@@ -5097,7 +5097,7 @@ void FStarshipEditorStyle::FStyle::SetupLevelEditorStyle()
 #if WITH_EDITOR || (IS_PROGRAM && WITH_UNREAL_DEVELOPER_TOOLS)
 	{
 
-		Set( "LevelViewport.DebugBorder", new BOX_BRUSH( "Old/Window/ViewportDebugBorder", 0.8f, FLinearColor(.7,0,0,.5) ) );
+		Set( "LevelViewport.DebugBorder", new BOX_BRUSH( "Old/Window/ViewportDebugBorder", 0.8f, FLinearColor(.7f,0.f,0.f,.5f) ) );
 		Set( "LevelViewport.BlackBackground", new FSlateColorBrush( FLinearColor::Red ) ); 
 		Set( "LevelViewport.StartingPlayInEditorBorder", new BOX_BRUSH( "Old/Window/ViewportDebugBorder", 0.8f, FLinearColor(0.1f,1.0f,0.1f,1.0f) ) );
 		Set( "LevelViewport.StartingSimulateBorder", new BOX_BRUSH( "Old/Window/ViewportDebugBorder", 0.8f, FLinearColor(1.0f,1.0f,0.1f,1.0f) ) );
@@ -5176,7 +5176,7 @@ void FStarshipEditorStyle::FStyle::SetupLevelEditorStyle()
 		const FString SmallRoundedButtonMiddle(TEXT("Common/SmallRoundedToggleCenter"));
 		const FString SmallRoundedButtonEnd(TEXT("Common/SmallRoundedToggleRight"));
 
-		const FLinearColor NormalColor(0.15, 0.15, 0.15, 1);
+		const FLinearColor NormalColor(0.15f, 0.15f, 0.15f, 1.f);
 
 		Set("Property.ToggleButton", FCheckBoxStyle()
 			.SetCheckBoxType(ESlateCheckBoxType::ToggleButton)
@@ -5284,7 +5284,7 @@ void FStarshipEditorStyle::FStyle::SetupLevelEditorStyle()
 
 	// Scalability (Performance Warning)
 	{
-		Set( "Scalability.ScalabilitySettings", new IMAGE_BRUSH("Scalability/ScalabilitySettings", FVector2D(473.0f, 266.0f) ) );
+		Set( "Scalability.ScalabilitySettings", new IMAGE_BRUSH("Scalability/ScalabilitySettings", FVector2f(473.0f, 266.0f) ) );
 	}
 
 	Set("WorkspaceMenu.AdditionalUI", new IMAGE_BRUSH("Icons/icon_ViewMode_LODColoration_16px", Icon16x16));
@@ -5302,8 +5302,8 @@ void FStarshipEditorStyle::FStyle::SetupPersonaStyle()
 		Set( "AnimViewportMenu.RotateMode.Small", new IMAGE_BRUSH( "Icons/icon_rotate_40x", Icon16x16 ) );
 		Set( "AnimViewportMenu.CameraFollow", new IMAGE_BRUSH( "Persona/Viewport/Camera_FollowBounds_40px", Icon32x32) );
 		Set( "AnimViewportMenu.CameraFollow.Small", new IMAGE_BRUSH( "Persona/Viewport/Camera_FollowBounds_40px", Icon16x16 ) );
-		Set( "AnimViewport.LocalSpaceEditing", new IMAGE_BRUSH( "Icons/icon_axis_local_16px", FVector2D( 16, 16 ) ) );
-		Set( "AnimViewport.WorldSpaceEditing", new IMAGE_BRUSH( "Icons/icon_axis_world_16px", FVector2D( 16, 16 ) ) );
+		Set( "AnimViewport.LocalSpaceEditing", new IMAGE_BRUSH( "Icons/icon_axis_local_16px", Icon16x16 ) );
+		Set( "AnimViewport.WorldSpaceEditing", new IMAGE_BRUSH( "Icons/icon_axis_world_16px", Icon16x16 ) );
 		Set( "AnimViewportMenu.SetShowNormals", new IMAGE_BRUSH( TEXT("Icons/icon_StaticMeshEd_Normals_40x"), Icon40x40 ) );
 		Set( "AnimViewportMenu.SetShowNormals.Small", new IMAGE_BRUSH( TEXT("Icons/icon_StaticMeshEd_Normals_40x"), Icon20x20 ) );
 		Set( "AnimViewportMenu.SetShowTangents", new IMAGE_BRUSH( TEXT("Icons/icon_StaticMeshEd_Tangents_40x"), Icon40x40 ) );
@@ -5321,17 +5321,17 @@ void FStarshipEditorStyle::FStyle::SetupPersonaStyle()
 
 		Set("AnimViewport.MessageText", FTextBlockStyle(NormalText)
 			.SetFont(DEFAULT_FONT("Bold", FStarshipCoreStyle::RegularTextSize))
-			.SetShadowOffset(FVector2D(1.0f, 1.0f))
+			.SetShadowOffset(FVector2f::UnitVector)
 			.SetShadowColorAndOpacity(FLinearColor(0.0f, 0.0f, 0.0f)));
 		Set("AnimViewport.WarningText", FTextBlockStyle(NormalText)
 			.SetFont(DEFAULT_FONT("Bold", FStarshipCoreStyle::RegularTextSize))
 			.SetColorAndOpacity(FLinearColor::Yellow)
-			.SetShadowOffset(FVector2D(1.0f, 1.0f))
+			.SetShadowOffset(FVector2f::UnitVector)
 			.SetShadowColorAndOpacity(FLinearColor(0.0f, 0.0f, 0.0f)));
 		Set("AnimViewport.ErrorText", FTextBlockStyle(NormalText)
 			.SetFont(DEFAULT_FONT("Bold", FStarshipCoreStyle::RegularTextSize))
 			.SetColorAndOpacity(FLinearColor::Red)
-			.SetShadowOffset(FVector2D(1.0f, 1.0f))
+			.SetShadowOffset(FVector2f::UnitVector)
 			.SetShadowColorAndOpacity(FLinearColor(0.0f, 0.0f, 0.0f)));
 		
 		// Viewport notifications 
@@ -5452,8 +5452,8 @@ void FStarshipEditorStyle::FStyle::SetupPersonaStyle()
 			.SetDownArrowPadding(FMargin(2.0f, 0.0f, 0.0f, 0.0f)));
 
 		Set("SkeletonTree.BlendProfile", new IMAGE_BRUSH(TEXT("Icons/icon_Persona_NewBlendSpace_16x"), Icon16x16));
-		Set("SkeletonTree.InlineEditorShadowTop", new IMAGE_BRUSH(TEXT("Common/ScrollBoxShadowTop"), FVector2D(64, 8)));
-		Set("SkeletonTree.InlineEditorShadowBottom", new IMAGE_BRUSH(TEXT("Common/ScrollBoxShadowBottom"), FVector2D(64, 8)));
+		Set("SkeletonTree.InlineEditorShadowTop", new IMAGE_BRUSH(TEXT("Common/ScrollBoxShadowTop"), FVector2f(64.f, 8.f)));
+		Set("SkeletonTree.InlineEditorShadowBottom", new IMAGE_BRUSH(TEXT("Common/ScrollBoxShadowBottom"), FVector2f(64.f, 8.f)));
 	}
 
 	// Animation editor
@@ -5476,12 +5476,12 @@ void FStarshipEditorStyle::FStyle::SetupPersonaStyle()
 
 		Set("AnimTimeline.Outliner.Label", FTextBlockStyle(NormalText)
 			.SetColorAndOpacity(FLinearColor(1.0f, 1.0f, 1.0f))
-			.SetShadowOffset(FVector2D(1.0f, 1.0f)));
+			.SetShadowOffset(FVector2f::UnitVector));
 
 		Set("AnimTimeline.Outliner.ItemColor", FLinearColor(0.04f, 0.04f, 0.04f, 0.8f));
 		Set("AnimTimeline.Outliner.HeaderColor", FLinearColor(0.03f, 0.03f, 0.03f, 1.0f));
 
-		Set("AnimTimeline.SectionMarker", new IMAGE_BRUSH(TEXT("Sequencer/SectionMarker"), FVector2D(11, 12)));
+		Set("AnimTimeline.SectionMarker", new IMAGE_BRUSH(TEXT("Sequencer/SectionMarker"), FVector2f(11.f, 12.f)));
 	}
 
 	// Skeletal mesh editor
@@ -5623,7 +5623,7 @@ void FStarshipEditorStyle::FStyle::SetupPersonaStyle()
 		Set( "Kismet.CreateBlueprint", new IMAGE_BRUSH( "/Icons/CreateBlueprint", Icon16x16) );
 		Set( "Kismet.HarvestBlueprintFromActors", new IMAGE_BRUSH_SVG( "Starship/Blueprints/HarvestBlueprintFromActors", Icon16x16) );
 
-		Set( "Kismet.Comment.Background", new IMAGE_BRUSH( "Old/Kismet2/Comment_Background", FVector2D(100.0f, 68.0f)) );
+		Set( "Kismet.Comment.Background", new IMAGE_BRUSH( "Old/Kismet2/Comment_Background", FVector2f(100.0f, 68.0f)) );
 
 		Set( "Kismet.AllClasses.VariableIcon", new IMAGE_BRUSH_SVG("Starship/Blueprints/pill", Icon16x16));
 		Set( "Kismet.AllClasses.ArrayVariableIcon", new IMAGE_BRUSH_SVG("Starship/Blueprints/pillarray", Icon16x16));
@@ -5653,11 +5653,11 @@ void FStarshipEditorStyle::FStyle::SetupPersonaStyle()
 		// blend space
 		Set("BlendSpaceEditor.ToggleTriangulation", new IMAGE_BRUSH(TEXT("Persona/BlendSpace/triangulation_16"), Icon16x16));
 		Set("BlendSpaceEditor.ToggleLabels", new IMAGE_BRUSH(TEXT("Persona/BlendSpace/label_16"), Icon16x16));
-		Set("BlendSpaceEditor.ArrowDown", new IMAGE_BRUSH(TEXT("Persona/BlendSpace/arrow_down_12x"), FVector2D(13.0f, 25.0f)));
-		Set("BlendSpaceEditor.ArrowUp", new IMAGE_BRUSH(TEXT("Persona/BlendSpace/arrow_up_12x"), FVector2D(13.0f, 25.0f)));
-		Set("BlendSpaceEditor.ArrowRight", new IMAGE_BRUSH(TEXT("Persona/BlendSpace/arrow_right_12x"), FVector2D(25.0f, 13.0f)));
-		Set("BlendSpaceEditor.ArrowLeft", new IMAGE_BRUSH(TEXT("Persona/BlendSpace/arrow_left_12x"), FVector2D(25.0f, 13.0f)));
-		Set("BlendSpaceEditor.PreviewIcon", new IMAGE_BRUSH(TEXT("Persona/BlendSpace/preview_21x"), FVector2D(21.0f, 21.0f)));
+		Set("BlendSpaceEditor.ArrowDown", new IMAGE_BRUSH(TEXT("Persona/BlendSpace/arrow_down_12x"), FVector2f(13.0f, 25.0f)));
+		Set("BlendSpaceEditor.ArrowUp", new IMAGE_BRUSH(TEXT("Persona/BlendSpace/arrow_up_12x"), FVector2f(13.0f, 25.0f)));
+		Set("BlendSpaceEditor.ArrowRight", new IMAGE_BRUSH(TEXT("Persona/BlendSpace/arrow_right_12x"), FVector2f(25.0f, 13.0f)));
+		Set("BlendSpaceEditor.ArrowLeft", new IMAGE_BRUSH(TEXT("Persona/BlendSpace/arrow_left_12x"), FVector2f(25.0f, 13.0f)));
+		Set("BlendSpaceEditor.PreviewIcon", new IMAGE_BRUSH(TEXT("Persona/BlendSpace/preview_21x"), FVector2f(21.0f, 21.0f)));
 		Set("BlendSpaceEditor.LabelBackground", new FSlateRoundedBoxBrush(FStyleColors::Background, FStyleColors::Foreground, 1.0f));
 		Set("BlendSpaceEditor.ZoomToFit", new IMAGE_BRUSH("GenericCurveEditor/Icons/FramingSelected_48x", Icon16x16));
 
@@ -5798,8 +5798,8 @@ void FStarshipEditorStyle::FStyle::SetupPersonaStyle()
 
 		// Specialized Play Button (Left button with green color)
 		FLinearColor GreenHSV = FStyleColors::AccentGreen.GetSpecifiedColor().LinearRGBToHSV();
-		FLinearColor GreenHover = FLinearColor(GreenHSV.R, GreenHSV.G * .5, GreenHSV.B, GreenHSV.A).HSVToLinearRGB();
-		FLinearColor GreenPress = FLinearColor(GreenHSV.R, GreenHSV.G, GreenHSV.B*.5, GreenHSV.A).HSVToLinearRGB(); 
+		FLinearColor GreenHover = FLinearColor(GreenHSV.R, GreenHSV.G * .5f, GreenHSV.B, GreenHSV.A).HSVToLinearRGB();
+		FLinearColor GreenPress = FLinearColor(GreenHSV.R, GreenHSV.G, GreenHSV.B*.5f, GreenHSV.A).HSVToLinearRGB(); 
 
 		FToolBarStyle MainToolbarPlayButton = MainToolbarLeftButton;
 
@@ -5832,8 +5832,8 @@ void FStarshipEditorStyle::FStyle::SetupPersonaStyle()
 
 		FLinearColor RedHSV = FStyleColors::AccentRed.GetSpecifiedColor().LinearRGBToHSV();
 
-		FLinearColor RedHover = FLinearColor(RedHSV.R, RedHSV.G * .5, RedHSV.B, RedHSV.A).HSVToLinearRGB();
-		FLinearColor RedPress = FLinearColor(RedHSV.R, RedHSV.G, RedHSV.B * .5, RedHSV.A).HSVToLinearRGB();
+		FLinearColor RedHover = FLinearColor(RedHSV.R, RedHSV.G * .5f, RedHSV.B, RedHSV.A).HSVToLinearRGB();
+		FLinearColor RedPress = FLinearColor(RedHSV.R, RedHSV.G, RedHSV.B * .5f, RedHSV.A).HSVToLinearRGB();
 
 		FToolBarStyle MainToolbarStopButton = MainToolbarCenterButton;
 
@@ -5877,7 +5877,7 @@ void FStarshipEditorStyle::FStyle::SetupPersonaStyle()
 		FComboButtonStyle PlayToolbarComboButton =
 			FComboButtonStyle(FStarshipCoreStyle::GetCoreStyle().GetWidgetStyle<FComboButtonStyle>("ComboButton"))
 				.SetDownArrowPadding(FMargin(1.f, 0.f, 2.f, 0.f))
-				.SetDownArrowImage(CORE_IMAGE_BRUSH_SVG("Starship/Common/ellipsis-vertical-narrow", FVector2D(6, 24)));
+				.SetDownArrowImage(CORE_IMAGE_BRUSH_SVG("Starship/Common/ellipsis-vertical-narrow", FVector2f(6.f, 24.f)));
 		PlayToolbarComboButton.ButtonStyle = RightToolbarComboButton;
 
 		MainToolbarRightComboButton.SetButtonStyle(RightToolbarComboButton);
@@ -5952,8 +5952,8 @@ void FStarshipEditorStyle::FStyle::SetupPersonaStyle()
 		Set( "Kismet.DebuggerOverlay.Breakpoint.EnabledAndInvalid", new IMAGE_BRUSH_SVG( "Starship/Blueprints/Breakpoint_Invalid", Icon32x32, FStyleColors::AccentRed) );
 		Set( "Kismet.DebuggerOverlay.Breakpoint.EnabledAndValid", new IMAGE_BRUSH_SVG( "Starship/Blueprints/Breakpoint_Valid", Icon32x32, FStyleColors::AccentRed) );
 
-		Set("Kismet.DebuggerOverlay.InstructionPointer", new IMAGE_BRUSH_SVG( "Starship/Blueprints/IP_Breakpoint", FVector2D(128,96)) );
-		Set("Kismet.DebuggerOverlay.InstructionPointerBreakpoint", new IMAGE_BRUSH_SVG( "Starship/Blueprints/IP_Breakpoint", FVector2D(128,96), FStyleColors::AccentRed) );
+		Set("Kismet.DebuggerOverlay.InstructionPointer", new IMAGE_BRUSH_SVG( "Starship/Blueprints/IP_Breakpoint", FVector2f(128.f,96.f)) );
+		Set("Kismet.DebuggerOverlay.InstructionPointerBreakpoint", new IMAGE_BRUSH_SVG( "Starship/Blueprints/IP_Breakpoint", FVector2f(128.f,96.f), FStyleColors::AccentRed) );
 
 		Set("Kismet.CallStackViewer.CurrentStackFrame", new IMAGE_BRUSH_SVG( "Starship/Blueprints/DebuggerArrow", Icon12x12 ));
 		Set("Kismet.CallStackViewer.CurrentStackFrameColor", FLinearColor(0.728f, 0.364f, 0.003f) );
@@ -6006,7 +6006,7 @@ void FStarshipEditorStyle::FStyle::SetupPersonaStyle()
 
 		Set( "PinnedCommandList.Heading", FTextBlockStyle(NormalText)
 			.SetFont(DEFAULT_FONT( "Regular", 8))
-			.SetColorAndOpacity(FLinearColor(0.4f, 0.4, 0.4f, 1.0f)));
+			.SetColorAndOpacity(FLinearColor(0.4f, 0.4f, 0.4f, 1.0f)));
 
 		/* Set images for various SCheckBox states associated with menu check box items... */
 		const FCheckBoxStyle BasicMenuCheckBoxStyle = FCheckBoxStyle()
@@ -6106,7 +6106,7 @@ void FStarshipEditorStyle::FStyle::SetupPersonaStyle()
 
 		Set( "ViewportPinnedCommandList.Heading", FTextBlockStyle(NormalText)
 			.SetFont(DEFAULT_FONT( "Regular", 8))
-			.SetColorAndOpacity(FLinearColor(0.4f, 0.4, 0.4f, 1.0f)));
+			.SetColorAndOpacity(FLinearColor(0.4f, 0.4f, 0.4f, 1.0f)));
 
 		/* Set images for various SCheckBox states associated with menu check box items... */
 		const FCheckBoxStyle BasicMenuCheckBoxStyle = FCheckBoxStyle()
@@ -6681,7 +6681,7 @@ void FStarshipEditorStyle::FStyle::SetupContentBrowserStyle()
 		Set( "ContentBrowser.SortDown", new IMAGE_BRUSH( "Common/SortDownArrow", Icon8x4 ) );
 
 		// Filter List - These are aliases for SBasicFilterBar styles in StarshipCoreStyle for backwards compatibility
-		Set("ContentBrowser.FilterImage", new CORE_IMAGE_BRUSH_SVG("Starship/CoreWidgets/FilterBar/FilterColorSegment", FVector2D(8, 22)));
+		Set("ContentBrowser.FilterImage", new CORE_IMAGE_BRUSH_SVG("Starship/CoreWidgets/FilterBar/FilterColorSegment", FVector2f(8.f, 22.f)));
 		Set("ContentBrowser.FilterBackground", new FSlateRoundedBoxBrush(FStyleColors::Secondary, 3.0f));
 
 		Set("ContentBrowser.FilterButton", FStarshipCoreStyle::GetCoreStyle().GetWidgetStyle<FCheckBoxStyle>("FilterBar.FilterButton"));
@@ -6693,19 +6693,19 @@ void FStarshipEditorStyle::FStyle::SetupContentBrowserStyle()
 		Set("ContentBrowser.Sources.Collections.Compact", new IMAGE_BRUSH("ContentBrowser/Sources_Collections_Compact_16x", Icon16x16));
 
 		// Asset tags (common)
-		Set("ContentBrowser.AssetTagBackground", new FSlateRoundedBoxBrush(FStyleColors::White, 2.0));
+		Set("ContentBrowser.AssetTagBackground", new FSlateRoundedBoxBrush(FStyleColors::White, 2.0f));
 
 		// Asset tags (standard)
 		Set("ContentBrowser.AssetTagButton", FCheckBoxStyle()
-			.SetUncheckedImage(IMAGE_BRUSH("ContentBrowser/AssetTagCheckbox_Flat", FVector2D(14.0f, 28.0f)))
-			.SetUncheckedHoveredImage(IMAGE_BRUSH("ContentBrowser/AssetTagCheckbox_Flat", FVector2D(14.0f, 28.0f), FLinearColor(0.5f, 0.5f, 0.5f, 1.0f)))
-			.SetUncheckedPressedImage(IMAGE_BRUSH("ContentBrowser/AssetTagCheckbox_Flat", FVector2D(14.0f, 28.0f), FLinearColor(0.5f, 0.5f, 0.5f, 1.0f)))
-			.SetUndeterminedImage(IMAGE_BRUSH("ContentBrowser/AssetTagCheckbox_Flat", FVector2D(14.0f, 28.0f)))
-			.SetUndeterminedHoveredImage(IMAGE_BRUSH("ContentBrowser/AssetTagCheckbox_Flat", FVector2D(14.0f, 28.0f), FLinearColor(0.5f, 0.5f, 0.5f, 1.0f)))
-			.SetUndeterminedPressedImage(IMAGE_BRUSH("ContentBrowser/AssetTagCheckbox_Flat", FVector2D(14.0f, 28.0f), FLinearColor(0.5f, 0.5f, 0.5f, 1.0f)))
-			.SetCheckedImage(IMAGE_BRUSH("ContentBrowser/AssetTagCheckbox_Flat", FVector2D(14.0f, 28.0f)))
-			.SetCheckedHoveredImage(IMAGE_BRUSH("ContentBrowser/AssetTagCheckbox_Flat", FVector2D(14.0f, 28.0f), FLinearColor(0.5f, 0.5f, 0.5f, 1.0f)))
-			.SetCheckedPressedImage(IMAGE_BRUSH("ContentBrowser/AssetTagCheckbox_Flat", FVector2D(14.0f, 28.0f), FLinearColor(0.5f, 0.5f, 0.5f, 1.0f)))
+			.SetUncheckedImage(IMAGE_BRUSH("ContentBrowser/AssetTagCheckbox_Flat", FVector2f(14.0f, 28.0f)))
+			.SetUncheckedHoveredImage(IMAGE_BRUSH("ContentBrowser/AssetTagCheckbox_Flat", FVector2f(14.0f, 28.0f), FLinearColor(0.5f, 0.5f, 0.5f, 1.0f)))
+			.SetUncheckedPressedImage(IMAGE_BRUSH("ContentBrowser/AssetTagCheckbox_Flat", FVector2f(14.0f, 28.0f), FLinearColor(0.5f, 0.5f, 0.5f, 1.0f)))
+			.SetUndeterminedImage(IMAGE_BRUSH("ContentBrowser/AssetTagCheckbox_Flat", FVector2f(14.0f, 28.0f)))
+			.SetUndeterminedHoveredImage(IMAGE_BRUSH("ContentBrowser/AssetTagCheckbox_Flat", FVector2f(14.0f, 28.0f), FLinearColor(0.5f, 0.5f, 0.5f, 1.0f)))
+			.SetUndeterminedPressedImage(IMAGE_BRUSH("ContentBrowser/AssetTagCheckbox_Flat", FVector2f(14.0f, 28.0f), FLinearColor(0.5f, 0.5f, 0.5f, 1.0f)))
+			.SetCheckedImage(IMAGE_BRUSH("ContentBrowser/AssetTagCheckbox_Flat", FVector2f(14.0f, 28.0f)))
+			.SetCheckedHoveredImage(IMAGE_BRUSH("ContentBrowser/AssetTagCheckbox_Flat", FVector2f(14.0f, 28.0f), FLinearColor(0.5f, 0.5f, 0.5f, 1.0f)))
+			.SetCheckedPressedImage(IMAGE_BRUSH("ContentBrowser/AssetTagCheckbox_Flat", FVector2f(14.0f, 28.0f), FLinearColor(0.5f, 0.5f, 0.5f, 1.0f)))
 			.SetPadding(0.0f)
 			);
 
@@ -6714,15 +6714,15 @@ void FStarshipEditorStyle::FStyle::SetupContentBrowserStyle()
 
 		// Asset tags (compact)
 		Set("ContentBrowser.AssetTagButton.Compact", FCheckBoxStyle()
-			.SetUncheckedImage(IMAGE_BRUSH("ContentBrowser/AssetTagCheckbox_Flat_Compact", FVector2D(10.0f, 20.0f)))
-			.SetUncheckedHoveredImage(IMAGE_BRUSH("ContentBrowser/AssetTagCheckbox_Flat_Compact", FVector2D(10.0f, 20.0f), FLinearColor(0.5f, 0.5f, 0.5f, 1.0f)))
-			.SetUncheckedPressedImage(IMAGE_BRUSH("ContentBrowser/AssetTagCheckbox_Flat_Compact", FVector2D(10.0f, 20.0f), FLinearColor(0.5f, 0.5f, 0.5f, 1.0f)))
-			.SetUndeterminedImage(IMAGE_BRUSH("ContentBrowser/AssetTagCheckbox_Flat_Compact", FVector2D(10.0f, 20.0f)))
-			.SetUndeterminedHoveredImage(IMAGE_BRUSH("ContentBrowser/AssetTagCheckbox_Flat_Compact", FVector2D(10.0f, 20.0f), FLinearColor(0.5f, 0.5f, 0.5f, 1.0f)))
-			.SetUndeterminedPressedImage(IMAGE_BRUSH("ContentBrowser/AssetTagCheckbox_Flat_Compact", FVector2D(10.0f, 20.0f), FLinearColor(0.5f, 0.5f, 0.5f, 1.0f)))
-			.SetCheckedImage(IMAGE_BRUSH("ContentBrowser/AssetTagCheckbox_Flat_Compact", FVector2D(10.0f, 20.0f)))
-			.SetCheckedHoveredImage(IMAGE_BRUSH("ContentBrowser/AssetTagCheckbox_Flat_Compact", FVector2D(10.0f, 20.0f), FLinearColor(0.5f, 0.5f, 0.5f, 1.0f)))
-			.SetCheckedPressedImage(IMAGE_BRUSH("ContentBrowser/AssetTagCheckbox_Flat_Compact", FVector2D(10.0f, 20.0f), FLinearColor(0.5f, 0.5f, 0.5f, 1.0f)))
+			.SetUncheckedImage(IMAGE_BRUSH("ContentBrowser/AssetTagCheckbox_Flat_Compact", FVector2f(10.0f, 20.0f)))
+			.SetUncheckedHoveredImage(IMAGE_BRUSH("ContentBrowser/AssetTagCheckbox_Flat_Compact", FVector2f(10.0f, 20.0f), FLinearColor(0.5f, 0.5f, 0.5f, 1.0f)))
+			.SetUncheckedPressedImage(IMAGE_BRUSH("ContentBrowser/AssetTagCheckbox_Flat_Compact", FVector2f(10.0f, 20.0f), FLinearColor(0.5f, 0.5f, 0.5f, 1.0f)))
+			.SetUndeterminedImage(IMAGE_BRUSH("ContentBrowser/AssetTagCheckbox_Flat_Compact", FVector2f(10.0f, 20.0f)))
+			.SetUndeterminedHoveredImage(IMAGE_BRUSH("ContentBrowser/AssetTagCheckbox_Flat_Compact", FVector2f(10.0f, 20.0f), FLinearColor(0.5f, 0.5f, 0.5f, 1.0f)))
+			.SetUndeterminedPressedImage(IMAGE_BRUSH("ContentBrowser/AssetTagCheckbox_Flat_Compact", FVector2f(10.0f, 20.0f), FLinearColor(0.5f, 0.5f, 0.5f, 1.0f)))
+			.SetCheckedImage(IMAGE_BRUSH("ContentBrowser/AssetTagCheckbox_Flat_Compact", FVector2f(10.0f, 20.0f)))
+			.SetCheckedHoveredImage(IMAGE_BRUSH("ContentBrowser/AssetTagCheckbox_Flat_Compact", FVector2f(10.0f, 20.0f), FLinearColor(0.5f, 0.5f, 0.5f, 1.0f)))
+			.SetCheckedPressedImage(IMAGE_BRUSH("ContentBrowser/AssetTagCheckbox_Flat_Compact", FVector2f(10.0f, 20.0f), FLinearColor(0.5f, 0.5f, 0.5f, 1.0f)))
 			.SetPadding(0.0f)
 			);
 
@@ -6741,8 +6741,8 @@ void FStarshipEditorStyle::FStyle::SetupContentBrowserStyle()
 			.SetFont( DEFAULT_FONT( "Bold", 11 ) )
 			.SetColorAndOpacity( FLinearColor( 1.0f, 1.0f, 1.0f ) )
 			.SetHighlightColor( FLinearColor( 1.0f, 1.0f, 1.0f ) )
-			.SetShadowOffset( FVector2D( 1, 1 ) )
-			.SetShadowColorAndOpacity( FLinearColor( 0, 0, 0, 0.9f ) ) );
+			.SetShadowOffset( FVector2f::UnitVector )
+			.SetShadowColorAndOpacity( FLinearColor( 0.f, 0.f, 0.f, 0.9f ) ) );
 
 		Set("ContentBrowser.ClassFont", FTextBlockStyle(NormalText)
 			.SetFont(DEFAULT_FONT("Regular", 7)));
@@ -6800,12 +6800,12 @@ void FStarshipEditorStyle::FStyle::SetupContentBrowserStyle()
 		Set( "ContentBrowser.ColumnViewFolderIcon", new CORE_IMAGE_BRUSH_SVG( "Starship/Common/folder-closed", Icon16x16 ) );
 		Set( "ContentBrowser.ColumnViewDeveloperFolderIcon", new IMAGE_BRUSH_SVG( "Starship/ContentBrowser/folder-developer", Icon16x16 ) );
 
-		Set("ContentBrowser.ListViewFolderIcon", new IMAGE_BRUSH_SVG("Starship/ContentBrowser/folder", FVector2D(64, 64)));
-		Set("ContentBrowser.ListViewVirtualFolderIcon", new IMAGE_BRUSH_SVG("Starship/ContentBrowser/FolderLargeVirtual", FVector2D(64, 64)));
-		Set("ContentBrowser.ListViewVirtualFolderShadow", new IMAGE_BRUSH("Starship/ContentBrowser/FolderLargeVirtualShadow", FVector2D(256, 256)));
-		Set("ContentBrowser.ListViewDeveloperFolderIcon", new IMAGE_BRUSH_SVG("Starship/ContentBrowser/FolderLargeDeveloper", FVector2D(64, 64)));
-		Set("ContentBrowser.ListViewCodeFolderIcon", new IMAGE_BRUSH_SVG("Starship/ContentBrowser/FolderLargeCode", FVector2D(64, 64)));
-		Set("ContentBrowser.ListViewPluginFolderIcon", new IMAGE_BRUSH_SVG("Starship/ContentBrowser/FolderLargePlugin", FVector2D(64, 64)));
+		Set("ContentBrowser.ListViewFolderIcon", new IMAGE_BRUSH_SVG("Starship/ContentBrowser/folder", Icon64x64));
+		Set("ContentBrowser.ListViewVirtualFolderIcon", new IMAGE_BRUSH_SVG("Starship/ContentBrowser/FolderLargeVirtual", Icon64x64));
+		Set("ContentBrowser.ListViewVirtualFolderShadow", new IMAGE_BRUSH("Starship/ContentBrowser/FolderLargeVirtualShadow", FVector2f(256.f, 256.f)));
+		Set("ContentBrowser.ListViewDeveloperFolderIcon", new IMAGE_BRUSH_SVG("Starship/ContentBrowser/FolderLargeDeveloper", Icon64x64));
+		Set("ContentBrowser.ListViewCodeFolderIcon", new IMAGE_BRUSH_SVG("Starship/ContentBrowser/FolderLargeCode", Icon64x64));
+		Set("ContentBrowser.ListViewPluginFolderIcon", new IMAGE_BRUSH_SVG("Starship/ContentBrowser/FolderLargePlugin", Icon64x64));
 
 		Set("ContentBrowser.AssetTileItem.FolderAreaHoveredBackground", new FSlateRoundedBoxBrush(FStyleColors::Secondary, 4.0f));
 		Set("ContentBrowser.AssetTileItem.FolderAreaSelectedBackground", new FSlateRoundedBoxBrush(FStyleColors::Secondary, 4.0f));
@@ -6847,7 +6847,7 @@ void FStarshipEditorStyle::FStyle::SetupContentBrowserStyle()
 		}
 
 		Set("ContentBrowser.AssetTileItem.DropShadow", new BOX_BRUSH("Starship/ContentBrowser/drop-shadow", FMargin(4.0f / 64.0f)));
-		Set("ContentBrowser.FolderItem.DropShadow", new IMAGE_BRUSH("Starship/ContentBrowser/folder-drop-shadow", FVector2D(256, 256)));
+		Set("ContentBrowser.FolderItem.DropShadow", new IMAGE_BRUSH("Starship/ContentBrowser/folder-drop-shadow", FVector2f(256.f, 256.f)));
 
 
 
@@ -7037,7 +7037,7 @@ void FStarshipEditorStyle::FStyle::SetupLandscapeEditorStyle()
 		Set("LandscapeEditor.Layer.Sync", new IMAGE_BRUSH_SVG("Icons/icon_Landscape_Layers_From_Material_16x", Icon16x16));
 		
 		Set("LandscapeEditor.SpinBox", FSpinBoxStyle(GetWidgetStyle<FSpinBoxStyle>("SpinBox"))
-			.SetTextPadding(FMargin(0))
+			.SetTextPadding(FMargin(0.f))
 			.SetBackgroundBrush(FSlateNoResource())
 			.SetHoveredBackgroundBrush(FSlateNoResource())
 			.SetInactiveFillBrush(FSlateNoResource())
@@ -7057,7 +7057,7 @@ void FStarshipEditorStyle::FStyle::SetupToolkitStyles()
 	{
 		Set("ProjectBrowser.VersionOverlayText", FTextBlockStyle(NormalText)
 			.SetFont(DEFAULT_FONT("Bold", 14))
-			.SetShadowOffset(FVector2D(0, 1)));
+			.SetShadowOffset(FVector2f(0.f, 1.f)));
 
 
 		const FTableRowStyle ProjectBrowserTableRowStyle = FTableRowStyle()
@@ -7138,8 +7138,8 @@ void FStarshipEditorStyle::FStyle::SetupToolkitStyles()
 		Set("CurveTableEditor.Cell.Text", FEditableTextBoxStyle()
 			.SetTextStyle(NormalText)
 			.SetBackgroundImageNormal(FSlateNoResource())
-			.SetBackgroundImageHovered(FSlateRoundedBoxBrush(FStyleColors::Input, 0.0, FStyleColors::Hover, 1.0))
-			.SetBackgroundImageFocused(FSlateRoundedBoxBrush(FStyleColors::Input, 0.0, FStyleColors::Primary, 1.0))
+			.SetBackgroundImageHovered(FSlateRoundedBoxBrush(FStyleColors::Input, 0.0f, FStyleColors::Hover, 1.0f))
+			.SetBackgroundImageFocused(FSlateRoundedBoxBrush(FStyleColors::Input, 0.0f, FStyleColors::Primary, 1.0f))
 			.SetBackgroundImageReadOnly(FSlateNoResource())
 			.SetPadding(FMargin(4.0f))
 			.SetForegroundColor(FStyleColors::Foreground)
@@ -7190,10 +7190,10 @@ void FStarshipEditorStyle::FStyle::SetupToolkitStyles()
 		);
 
 		Set("DataTableEditor.CellListViewRow", FTableRowStyle(NormalTableRowStyle)
-			.SetEvenRowBackgroundBrush(IMAGE_BRUSH("PropertyView/DetailCategoryMiddle", FVector2D(16, 16), FLinearColor(0.5f, 0.5f, 0.5f)))
-			.SetEvenRowBackgroundHoveredBrush(IMAGE_BRUSH("PropertyView/DetailCategoryMiddle_Hovered", FVector2D(16, 16), FLinearColor(0.5f, 0.5f, 0.5f)))
-			.SetOddRowBackgroundBrush(IMAGE_BRUSH("PropertyView/DetailCategoryMiddle", FVector2D(16, 16), FLinearColor(0.2f, 0.2f, 0.2f)))
-			.SetOddRowBackgroundHoveredBrush(IMAGE_BRUSH("PropertyView/DetailCategoryMiddle_Hovered", FVector2D(16, 16), FLinearColor(0.2f, 0.2f, 0.2f)))
+			.SetEvenRowBackgroundBrush(IMAGE_BRUSH("PropertyView/DetailCategoryMiddle", Icon16x16, FLinearColor(0.5f, 0.5f, 0.5f)))
+			.SetEvenRowBackgroundHoveredBrush(IMAGE_BRUSH("PropertyView/DetailCategoryMiddle_Hovered", Icon16x16, FLinearColor(0.5f, 0.5f, 0.5f)))
+			.SetOddRowBackgroundBrush(IMAGE_BRUSH("PropertyView/DetailCategoryMiddle", Icon16x16, FLinearColor(0.2f, 0.2f, 0.2f)))
+			.SetOddRowBackgroundHoveredBrush(IMAGE_BRUSH("PropertyView/DetailCategoryMiddle_Hovered", Icon16x16, FLinearColor(0.2f, 0.2f, 0.2f)))
 			.SetActiveBrush(IMAGE_BRUSH("Common/Selection", Icon8x8, FLinearColor(0.075f, 0.075f, 0.075f)))
 			.SetActiveHoveredBrush(IMAGE_BRUSH("Common/Selection", Icon8x8, FLinearColor(0.075f, 0.075f, 0.075f)))
 			.SetInactiveBrush(IMAGE_BRUSH("Common/Selection", Icon8x8, FLinearColor(0.075f, 0.075f, 0.075f)))
@@ -7242,10 +7242,10 @@ void FStarshipEditorStyle::FStyle::SetupToolkitStyles()
 
 
 		Set("MaterialEditor.CellListViewRow", FTableRowStyle(NormalTableRowStyle)
-			.SetEvenRowBackgroundBrush(IMAGE_BRUSH("PropertyView/DetailCategoryMiddle", FVector2D(16, 16), FLinearColor(0.5f, 0.5f, 0.5f)))
-			.SetEvenRowBackgroundHoveredBrush(IMAGE_BRUSH("PropertyView/DetailCategoryMiddle_Hovered", FVector2D(16, 16), FLinearColor(0.5f, 0.5f, 0.5f)))
-			.SetOddRowBackgroundBrush(IMAGE_BRUSH("PropertyView/DetailCategoryMiddle", FVector2D(16, 16), FLinearColor(0.35f, 0.35f, 0.35f)))
-			.SetOddRowBackgroundHoveredBrush(IMAGE_BRUSH("PropertyView/DetailCategoryMiddle_Hovered", FVector2D(16, 16), FLinearColor(0.35f, 0.35f, 0.35f)))
+			.SetEvenRowBackgroundBrush(IMAGE_BRUSH("PropertyView/DetailCategoryMiddle", Icon16x16, FLinearColor(0.5f, 0.5f, 0.5f)))
+			.SetEvenRowBackgroundHoveredBrush(IMAGE_BRUSH("PropertyView/DetailCategoryMiddle_Hovered", Icon16x16, FLinearColor(0.5f, 0.5f, 0.5f)))
+			.SetOddRowBackgroundBrush(IMAGE_BRUSH("PropertyView/DetailCategoryMiddle", Icon16x16, FLinearColor(0.35f, 0.35f, 0.35f)))
+			.SetOddRowBackgroundHoveredBrush(IMAGE_BRUSH("PropertyView/DetailCategoryMiddle_Hovered", Icon16x16, FLinearColor(0.35f, 0.35f, 0.35f)))
 			.SetActiveBrush(IMAGE_BRUSH("Common/Selection", Icon8x8, FLinearColor(0.075f, 0.075f, 0.075f)))
 			.SetActiveHoveredBrush(IMAGE_BRUSH("Common/Selection", Icon8x8, FLinearColor(0.075f, 0.075f, 0.075f)))
 			.SetInactiveBrush(IMAGE_BRUSH("Common/Selection", Icon8x8, FLinearColor(0.075f, 0.075f, 0.075f)))
@@ -7568,12 +7568,12 @@ void FStarshipEditorStyle::FStyle::SetupToolkitStyles()
 		Set("GenericCurveEditor.OpenUserImplementableFilterWindow", new IMAGE_BRUSH_SVG("Starship/GenericCurveEditor/MiscFilters", Icon20x20));
 
 		// Key Types
-		Set("GenericCurveEditor.Key", new IMAGE_BRUSH("GenericCurveEditor/Keys/Diamond_Filled", FVector2D(11.0f, 11.0f))); 
-		Set("GenericCurveEditor.ConstantKey", new IMAGE_BRUSH("GenericCurveEditor/Keys/Square_Filled", FVector2D(11.0f, 11.0f))); 
-		Set("GenericCurveEditor.LinearKey", new IMAGE_BRUSH("GenericCurveEditor/Keys/Triangle_Filled", FVector2D(11.0f, 11.0f)));
-		Set("GenericCurveEditor.CubicKey", new IMAGE_BRUSH("GenericCurveEditor/Keys/Diamond_Filled", FVector2D(11.0f, 11.0f)));
-		Set("GenericCurveEditor.TangentHandle", new IMAGE_BRUSH("GenericCurveEditor/Keys/TangentHandle", FVector2D(8.0f, 8.0f)));
-		Set("GenericCurveEditor.WeightedTangentCubicKey", new IMAGE_BRUSH("GenericCurveEditor/Keys/Trapezoid_Filled", FVector2D(11.0f, 11.0f)));
+		Set("GenericCurveEditor.Key", new IMAGE_BRUSH("GenericCurveEditor/Keys/Diamond_Filled", FVector2f(11.0f, 11.0f))); 
+		Set("GenericCurveEditor.ConstantKey", new IMAGE_BRUSH("GenericCurveEditor/Keys/Square_Filled", FVector2f(11.0f, 11.0f))); 
+		Set("GenericCurveEditor.LinearKey", new IMAGE_BRUSH("GenericCurveEditor/Keys/Triangle_Filled", FVector2f(11.0f, 11.0f)));
+		Set("GenericCurveEditor.CubicKey", new IMAGE_BRUSH("GenericCurveEditor/Keys/Diamond_Filled", FVector2f(11.0f, 11.0f)));
+		Set("GenericCurveEditor.TangentHandle", new IMAGE_BRUSH("GenericCurveEditor/Keys/TangentHandle", Icon8x8));
+		Set("GenericCurveEditor.WeightedTangentCubicKey", new IMAGE_BRUSH("GenericCurveEditor/Keys/Trapezoid_Filled", FVector2f(11.0f, 11.0f)));
 
 		// Pre-Infinity
 		Set("GenericCurveEditor.SetPreInfinityExtrapConstant", new IMAGE_BRUSH_SVG("Starship/GenericCurveEditor/PreInfinityConstant_20", Icon20x20));
@@ -7730,15 +7730,15 @@ void FStarshipEditorStyle::FStyle::SetupToolkitStyles()
 			.SetFont(DEFAULT_FONT("Bold", 11))
 			.SetColorAndOpacity(FLinearColor(1.0f, 1.0f, 1.0f))
 			.SetHighlightColor(FLinearColor(1.0f, 1.0f, 1.0f))
-			.SetShadowOffset(FVector2D(1, 1))
+			.SetShadowOffset(FVector2f::UnitVector)
 			.SetShadowColorAndOpacity(FLinearColor(0, 0, 0, 0.9f)));
 
 		Set("PhysicsAssetEditor.Tools.Font", FTextBlockStyle(NormalText)
 			.SetFont(DEFAULT_FONT("Bold", 11))
 			.SetColorAndOpacity(FLinearColor(1.0f, 1.0f, 1.0f))
 			.SetHighlightColor(FLinearColor(1.0f, 1.0f, 1.0f))
-			.SetShadowOffset(FVector2D(1, 1))
-			.SetShadowColorAndOpacity(FLinearColor(0, 0, 0, 0.9f)));
+			.SetShadowOffset(FVector2f::UnitVector)
+			.SetShadowColorAndOpacity(FLinearColor(0.f, 0.f, 0.f, 0.9f)));
 
 		FLinearColor Red = FLinearColor::Red;
 		FLinearColor Red_Selected = FLinearColor::Red.Desaturate(0.75f);
@@ -8012,7 +8012,7 @@ void FStarshipEditorStyle::FStyle::SetupAutomationStyles()
 			FTextBlockStyle(NormalText)
 			.SetFont(DEFAULT_FONT("Regular", 12))
 			.SetColorAndOpacity(FSlateColor::UseForeground())
-			.SetShadowOffset(FVector2D(1, 1))
+			.SetShadowOffset(FVector2f::UnitVector)
 			.SetShadowColorAndOpacity(FLinearColor::Black);
 
 		Set("FBXLargeFont", FBXLargeFont);
@@ -8021,7 +8021,7 @@ void FStarshipEditorStyle::FStyle::SetupAutomationStyles()
 			FTextBlockStyle(NormalText)
 			.SetFont(DEFAULT_FONT("Regular", 11))
 			.SetColorAndOpacity(FSlateColor::UseForeground())
-			.SetShadowOffset(FVector2D(1, 1))
+			.SetShadowOffset(FVector2f::UnitVector)
 			.SetShadowColorAndOpacity(FLinearColor::Black);
 
 		Set("FBXMediumFont", FBXMediumFont);
@@ -8030,7 +8030,7 @@ void FStarshipEditorStyle::FStyle::SetupAutomationStyles()
 			FTextBlockStyle(NormalText)
 			.SetFont(DEFAULT_FONT("Regular", 10))
 			.SetColorAndOpacity(FSlateColor::UseForeground())
-			.SetShadowOffset(FVector2D(1, 1))
+			.SetShadowOffset(FVector2f::UnitVector)
 			.SetShadowColorAndOpacity(FLinearColor::Black);
 
 		Set("FBXSmallFont", FBXSmallFont);
@@ -8067,32 +8067,32 @@ void FStarshipEditorStyle::FStyle::SetupUMGEditorStyles()
 	Set("UMGEditor.ResizeAreaHandle", new IMAGE_BRUSH("Icons/UMG/ResizeAreaHandle", Icon20x20));
 
 	Set("UMGEditor.AnchorGizmo.Center", new IMAGE_BRUSH("Icons/UMG/AnchorGizmo/center", Icon16x16));
-	Set("UMGEditor.AnchorGizmo.Center.Hovered", new IMAGE_BRUSH("Icons/UMG/AnchorGizmo/center", Icon16x16, FLinearColor(0, 1, 0)));
+	Set("UMGEditor.AnchorGizmo.Center.Hovered", new IMAGE_BRUSH("Icons/UMG/AnchorGizmo/center", Icon16x16, FLinearColor::Green));
 	
-	Set("UMGEditor.AnchorGizmo.Left", new IMAGE_BRUSH("Icons/UMG/AnchorGizmo/left", FVector2D(32, 16)));
-	Set("UMGEditor.AnchorGizmo.Left.Hovered", new IMAGE_BRUSH("Icons/UMG/AnchorGizmo/left", FVector2D(32, 16), FLinearColor(0, 1, 0)));
-	Set("UMGEditor.AnchorGizmo.Right", new IMAGE_BRUSH("Icons/UMG/AnchorGizmo/right", FVector2D(32, 16)));
-	Set("UMGEditor.AnchorGizmo.Right.Hovered", new IMAGE_BRUSH("Icons/UMG/AnchorGizmo/right", FVector2D(32, 16), FLinearColor(0, 1, 0)));
+	Set("UMGEditor.AnchorGizmo.Left", new IMAGE_BRUSH("Icons/UMG/AnchorGizmo/left", FVector2f(32.f, 16.f)));
+	Set("UMGEditor.AnchorGizmo.Left.Hovered", new IMAGE_BRUSH("Icons/UMG/AnchorGizmo/left", FVector2f(32.f, 16.f), FLinearColor::Green));
+	Set("UMGEditor.AnchorGizmo.Right", new IMAGE_BRUSH("Icons/UMG/AnchorGizmo/right", FVector2f(32.f, 16.f)));
+	Set("UMGEditor.AnchorGizmo.Right.Hovered", new IMAGE_BRUSH("Icons/UMG/AnchorGizmo/right", FVector2f(32.f, 16.f), FLinearColor::Green));
 	
-	Set("UMGEditor.AnchorGizmo.Top", new IMAGE_BRUSH("Icons/UMG/AnchorGizmo/top", FVector2D(16, 32)));
-	Set("UMGEditor.AnchorGizmo.Top.Hovered", new IMAGE_BRUSH("Icons/UMG/AnchorGizmo/top", FVector2D(16, 32), FLinearColor(0, 1, 0)));
-	Set("UMGEditor.AnchorGizmo.Bottom", new IMAGE_BRUSH("Icons/UMG/AnchorGizmo/bottom", FVector2D(16, 32)));
-	Set("UMGEditor.AnchorGizmo.Bottom.Hovered", new IMAGE_BRUSH("Icons/UMG/AnchorGizmo/bottom", FVector2D(16, 32), FLinearColor(0, 1, 0)));
+	Set("UMGEditor.AnchorGizmo.Top", new IMAGE_BRUSH("Icons/UMG/AnchorGizmo/top", FVector2f(16.f, 32.f)));
+	Set("UMGEditor.AnchorGizmo.Top.Hovered", new IMAGE_BRUSH("Icons/UMG/AnchorGizmo/top", FVector2f(16.f, 32.f), FLinearColor::Green));
+	Set("UMGEditor.AnchorGizmo.Bottom", new IMAGE_BRUSH("Icons/UMG/AnchorGizmo/bottom", FVector2f(16.f, 32.f)));
+	Set("UMGEditor.AnchorGizmo.Bottom.Hovered", new IMAGE_BRUSH("Icons/UMG/AnchorGizmo/bottom", FVector2f(16.f, 32.f), FLinearColor::Green));
 
-	Set("UMGEditor.AnchorGizmo.TopLeft", new IMAGE_BRUSH("Icons/UMG/AnchorGizmo/topleft", FVector2D(24, 24)));
-	Set("UMGEditor.AnchorGizmo.TopLeft.Hovered", new IMAGE_BRUSH("Icons/UMG/AnchorGizmo/topleft", FVector2D(24, 24), FLinearColor(0, 1, 0)));
+	Set("UMGEditor.AnchorGizmo.TopLeft", new IMAGE_BRUSH("Icons/UMG/AnchorGizmo/topleft", Icon24x24));
+	Set("UMGEditor.AnchorGizmo.TopLeft.Hovered", new IMAGE_BRUSH("Icons/UMG/AnchorGizmo/topleft", Icon24x24, FLinearColor::Green));
 
-	Set("UMGEditor.AnchorGizmo.TopRight", new IMAGE_BRUSH("Icons/UMG/AnchorGizmo/topright", FVector2D(24, 24)));
-	Set("UMGEditor.AnchorGizmo.TopRight.Hovered", new IMAGE_BRUSH("Icons/UMG/AnchorGizmo/topright", FVector2D(24, 24), FLinearColor(0, 1, 0)));
+	Set("UMGEditor.AnchorGizmo.TopRight", new IMAGE_BRUSH("Icons/UMG/AnchorGizmo/topright", Icon24x24));
+	Set("UMGEditor.AnchorGizmo.TopRight.Hovered", new IMAGE_BRUSH("Icons/UMG/AnchorGizmo/topright", Icon24x24, FLinearColor::Green));
 
-	Set("UMGEditor.AnchorGizmo.BottomLeft", new IMAGE_BRUSH("Icons/UMG/AnchorGizmo/bottomleft", FVector2D(24, 24)));
-	Set("UMGEditor.AnchorGizmo.BottomLeft.Hovered", new IMAGE_BRUSH("Icons/UMG/AnchorGizmo/bottomleft", FVector2D(24, 24), FLinearColor(0, 1, 0)));
+	Set("UMGEditor.AnchorGizmo.BottomLeft", new IMAGE_BRUSH("Icons/UMG/AnchorGizmo/bottomleft", Icon24x24));
+	Set("UMGEditor.AnchorGizmo.BottomLeft.Hovered", new IMAGE_BRUSH("Icons/UMG/AnchorGizmo/bottomleft", Icon24x24, FLinearColor::Green));
 
-	Set("UMGEditor.AnchorGizmo.BottomRight", new IMAGE_BRUSH("Icons/UMG/AnchorGizmo/bottomright", FVector2D(24, 24)));
-	Set("UMGEditor.AnchorGizmo.BottomRight.Hovered", new IMAGE_BRUSH("Icons/UMG/AnchorGizmo/bottomright", FVector2D(24, 24), FLinearColor(0, 1, 0)));
+	Set("UMGEditor.AnchorGizmo.BottomRight", new IMAGE_BRUSH("Icons/UMG/AnchorGizmo/bottomright", Icon24x24));
+	Set("UMGEditor.AnchorGizmo.BottomRight.Hovered", new IMAGE_BRUSH("Icons/UMG/AnchorGizmo/bottomright", Icon24x24, FLinearColor::Green));
 
 
-	Set("UMGEditor.AnchoredWidget", new BOX_BRUSH("Common/Button", FVector2D(32, 32), 8.0f / 32.0f));
+	Set("UMGEditor.AnchoredWidget", new BOX_BRUSH("Common/Button", Icon32x32, 8.0f / 32.0f));
 	Set("UMGEditor.AnchoredWidgetAlignment", new IMAGE_BRUSH("Icons/icon_tab_DeviceManager_16x", Icon8x8));
 	
 
@@ -8149,11 +8149,11 @@ void FStarshipEditorStyle::FStyle::SetupUMGEditorStyles()
 		FTextBlockStyle(NormalText)
 		.SetFont(DEFAULT_FONT("Regular", 18))
 		.SetColorAndOpacity(FSlateColor::UseForeground())
-		.SetShadowOffset(FVector2D(1, 1))
+		.SetShadowOffset(FVector2f::UnitVector)
 		.SetShadowColorAndOpacity(FLinearColor::Black);
 
 
-	Set("UMGEditor.AddAnimationIcon", new IMAGE_BRUSH("Icons/PlusSymbol_12x", Icon12x12, FLinearColor(.05,.05,.05) ) );
+	Set("UMGEditor.AddAnimationIcon", new IMAGE_BRUSH("Icons/PlusSymbol_12x", Icon12x12, FLinearColor(.05f,.05f,.05f) ) );
 	Set("UMGEditor.NoAnimationFont", NoAnimationFont);
 
 	Set("UMGEditor.SwitchToDesigner", new IMAGE_BRUSH("UMG/Designer_40x", Icon20x20));
@@ -8253,7 +8253,7 @@ void FStarshipEditorStyle::FStyle::SetupLocalizationDashboardStyles()
 
 	Set( "LocalizationDashboard.CommandletLog.Text", FTextBlockStyle(NormalText)
 		.SetFont( DEFAULT_FONT( "Regular", 8 ) )
-		.SetShadowOffset( FVector2D::ZeroVector )
+		.SetShadowOffset( FVector2f::ZeroVector )
 		);
 }
 
