Tweak viewport toolbar padding back to 4px.

diff --git a/Engine/Source/Editor/EditorStyle/Private/StarshipStyle.cpp b/Engine/Source/Editor/EditorStyle/Private/StarshipStyle.cpp
--- a/Engine/Source/Editor/EditorStyle/Private/StarshipStyle.cpp
+++ b/Engine/Source/Editor/EditorStyle/Private/StarshipStyle.cpp
@@ -2557,7 +2557,7 @@ void FStarshipEditorStyle::FStyle::SetupViewportStyles()
 				.SetSeparatorBrush(SeparatorBrush)
 				.SetSeparatorThickness(1.0f)
 				.SetSeparatorPadding(FMargin(4.0f, 0.f))
-				.SetBackgroundPadding(FMargin(2.0f))
+				.SetBackgroundPadding(FMargin(4.0f))
 				.SetWrapButtonPadding(FMargin(4.0f, 0.f))
 				.SetWrapButtonIndex(-1)
 				.SetRaisedChildrenRightPadding(16.0f);
