[SlateEditorStyle] Removing the last of the UE4 stylesheet references and the file itself

diff --git a/Engine/Plugins/Editor/ModelingToolsEditorMode/Source/ModelingToolsEditorMode/Private/ModelingToolsEditorModeStyle.cpp b/Engine/Plugins/Editor/ModelingToolsEditorMode/Source/ModelingToolsEditorMode/Private/ModelingToolsEditorModeStyle.cpp
--- a/Engine/Plugins/Editor/ModelingToolsEditorMode/Source/ModelingToolsEditorMode/Private/ModelingToolsEditorModeStyle.cpp
+++ b/Engine/Plugins/Editor/ModelingToolsEditorMode/Source/ModelingToolsEditorMode/Private/ModelingToolsEditorModeStyle.cpp
@@ -83,15 +83,7 @@ void FModelingToolsEditorModeStyle::Initialize()
 	{
 		// Accept/Cancel/Complete active tool
 
-		if (FCoreStyle::IsStarshipStyle())
-		{
-			StyleSet->Set("LevelEditor.ModelingToolsMode", new IMAGE_BRUSH_SVG("Starship/geometry", FVector2D(20.0f, 20.0f)));
-		}
-		else
-		{
-			StyleSet->Set("LevelEditor.ModelingToolsMode", new IMAGE_PLUGIN_BRUSH("Icons/icon_ModelingToolsEditorMode", FVector2D(40.0f, 40.0f)));
-			StyleSet->Set("LevelEditor.ModelingToolsMode.Small", new IMAGE_PLUGIN_BRUSH("Icons/icon_ModelingToolsEditorMode", FVector2D(20.0f, 20.0f)));
-		}
+		StyleSet->Set("LevelEditor.ModelingToolsMode", new IMAGE_BRUSH_SVG("Starship/geometry", FVector2D(20.0f, 20.0f)));
 
 		StyleSet->Set("ModelingMode.DefaultSettings", new FSlateImageBrush(StyleSet->RootToCoreContentDir(TEXT("../Editor/Slate/Icons/GeneralTools/Settings_40x.png")), Icon20x20));
 		StyleSet->Set("ModelingMode.SubToolArrow", new FSlateImageBrush(StyleSet->RootToCoreContentDir(
diff --git a/Engine/Plugins/Experimental/ChaosEditor/Source/FractureEditor/Private/FractureEditorStyle.cpp b/Engine/Plugins/Experimental/ChaosEditor/Source/FractureEditor/Private/FractureEditorStyle.cpp
--- a/Engine/Plugins/Experimental/ChaosEditor/Source/FractureEditor/Private/FractureEditorStyle.cpp
+++ b/Engine/Plugins/Experimental/ChaosEditor/Source/FractureEditor/Private/FractureEditorStyle.cpp
@@ -102,14 +102,7 @@ FFractureEditorStyle::FFractureEditorStyle()
 	Set("FractureEditor.GenerateAsset", new IMAGE_BRUSH_SVG("FractureGenerateAsset", IconSize));
 	Set("FractureEditor.ResetAsset", new IMAGE_BRUSH_SVG("FractureReset", FVector2D(20.0f, 20.0f)));
 
-	if (FCoreStyle::IsStarshipStyle())
-	{
-		Set("LevelEditor.FractureMode", new IMAGE_BRUSH_SVG("fracture", FVector2D(20.0f, 20.0f)));
-	}
-	else
-	{
-		Set("LevelEditor.FractureMode", new IMAGE_BRUSH("FractureMode.png", FVector2D(40.0f, 40.0f)));
-	}
+	Set("LevelEditor.FractureMode", new IMAGE_BRUSH_SVG("fracture", FVector2D(20.0f, 20.0f)));
 
 	FSlateStyleRegistry::RegisterSlateStyle(*this);
 }
diff --git a/Engine/Source/Editor/EditorStyle/Private/EditorStyleModule.cpp b/Engine/Source/Editor/EditorStyle/Private/EditorStyleModule.cpp
--- a/Engine/Source/Editor/EditorStyle/Private/EditorStyleModule.cpp
+++ b/Engine/Source/Editor/EditorStyle/Private/EditorStyleModule.cpp
@@ -3,7 +3,6 @@
 #include "CoreMinimal.h"
 #include "Modules/ModuleManager.h"
 #include "Interfaces/IEditorStyleModule.h"
-#include "SlateEditorStyle.h"
 #include "StarshipStyle.h"
 #include "Styling/CoreStyle.h"
 #include "Styling/StyleColors.h"
@@ -21,43 +20,22 @@ public:
 
 	virtual void StartupModule( ) override
 	{
-		if (FCoreStyle::IsStarshipStyle())
-		{
 #if ALLOW_THEMES
-			USlateThemeManager::Get().ValidateActiveTheme();
+		USlateThemeManager::Get().ValidateActiveTheme();
 #endif
 
-			bUsingStarshipStyle = true;
-			FStarshipEditorStyle::Initialize();
+		FStarshipEditorStyle::Initialize();
 
-			// set the application style to be the editor style
-			FAppStyle::SetAppStyleSetName(FStarshipEditorStyle::GetStyleSetName());
-
-		}
-		else
-		{
-			FSlateEditorStyle::Initialize();
-
-			// set the application style to be the editor style
-			FAppStyle::SetAppStyleSetName(FAppStyle::GetAppStyleSetName());
-		}
+		// set the application style to be the editor style
+		FAppStyle::SetAppStyleSetName(FStarshipEditorStyle::GetStyleSetName());
 	}
 
 	virtual void ShutdownModule( ) override
 	{
-		if (bUsingStarshipStyle)
-		{
-			FStarshipEditorStyle::Shutdown();
-		}
-		else
-		{
-			FSlateEditorStyle::Shutdown();
-		}
+		FStarshipEditorStyle::Shutdown();
 	}
 
 	// End IModuleInterface interface
-private:
-	bool bUsingStarshipStyle = false;
 };
 
 
diff --git a/Engine/Source/Editor/EditorStyle/Private/SlateEditorStyle.cpp b/Engine/Source/Editor/EditorStyle/Private/SlateEditorStyle.cpp
deleted file mode 100644
diff --git a/Engine/Source/Editor/EditorStyle/Private/SlateEditorStyle.h b/Engine/Source/Editor/EditorStyle/Private/SlateEditorStyle.h
deleted file mode 100644
diff --git a/Engine/Source/Programs/SlateViewer/Private/SlateViewerApp.cpp b/Engine/Source/Programs/SlateViewer/Private/SlateViewerApp.cpp
--- a/Engine/Source/Programs/SlateViewer/Private/SlateViewerApp.cpp
+++ b/Engine/Source/Programs/SlateViewer/Private/SlateViewerApp.cpp
@@ -71,7 +71,7 @@ int RunSlateViewer( const TCHAR* CommandLine )
 		// Bring up perf test
 		SummonPerfTestSuite();
 	}
-	else if (FCoreStyle::IsStarshipStyle())
+	else
 	{
 		FGlobalTabmanager::Get()->SetApplicationTitle(LOCTEXT("AppTitleStarship", "Starship Slate Viewer"));
 		FAppStyle::SetAppStyleSetName(FStarshipCoreStyle::GetCoreStyle().GetStyleSetName());
diff --git a/Engine/Source/Runtime/Slate/Private/Framework/Application/SlateApplication.cpp b/Engine/Source/Runtime/Slate/Private/Framework/Application/SlateApplication.cpp
--- a/Engine/Source/Runtime/Slate/Private/Framework/Application/SlateApplication.cpp
+++ b/Engine/Source/Runtime/Slate/Private/Framework/Application/SlateApplication.cpp
@@ -7354,7 +7354,7 @@ TSharedRef<FSlateApplication> FSlateApplication::InitializeAsStandaloneApplicati
 
 void FSlateApplication::InitializeCoreStyle()
 {
-	if (FCoreStyle::IsStarshipStyle() && !FStarshipCoreStyle::IsInitialized())
+	if (!FStarshipCoreStyle::IsInitialized())
 	{
 #if ALLOW_THEMES
 		USlateThemeManager::Get().LoadThemes();
@@ -7362,11 +7362,6 @@ void FSlateApplication::InitializeCoreStyle()
 		FStarshipCoreStyle::ResetToDefault();
 		FAppStyle::SetAppStyleSet(FStarshipCoreStyle::GetCoreStyle());
 	}
-	else if(!FCoreStyle::IsStarshipStyle() && !FCoreStyle::IsInitialized())
-	{
-		FCoreStyle::ResetToDefault();
-		FAppStyle::SetAppStyleSet(FCoreStyle::GetCoreStyle());
-	}
 
 	FUMGCoreStyle::ResetToDefault();
 }
diff --git a/Engine/Source/Runtime/Slate/Private/Framework/MultiBox/MultiBox.cpp b/Engine/Source/Runtime/Slate/Private/Framework/MultiBox/MultiBox.cpp
--- a/Engine/Source/Runtime/Slate/Private/Framework/MultiBox/MultiBox.cpp
+++ b/Engine/Source/Runtime/Slate/Private/Framework/MultiBox/MultiBox.cpp
@@ -302,7 +302,7 @@ FMultiBox::FMultiBox(const EMultiBoxType InType, FMultiBoxCustomization InCustom
 	, LastSelectedCommandIndex( INDEX_NONE )
 {
 
-	if ((InType == EMultiBoxType::SlimHorizontalToolBar ||  InType == EMultiBoxType::SlimHorizontalUniformToolBar) && FCoreStyle::IsStarshipStyle())
+	if (InType == EMultiBoxType::SlimHorizontalToolBar ||  InType == EMultiBoxType::SlimHorizontalUniformToolBar)
 	{
 		StyleName = "SlimToolBar";
 	}
diff --git a/Engine/Source/Runtime/SlateCore/Private/Styling/CoreStyle.cpp b/Engine/Source/Runtime/SlateCore/Private/Styling/CoreStyle.cpp
--- a/Engine/Source/Runtime/SlateCore/Private/Styling/CoreStyle.cpp
+++ b/Engine/Source/Runtime/SlateCore/Private/Styling/CoreStyle.cpp
@@ -1373,11 +1373,6 @@ const TSharedPtr<FSlateDynamicImageBrush> FCoreStyle::GetDynamicImageBrush( FNam
 	return Instance->GetDynamicImageBrush(BrushTemplate, TextureResource, TextureName);
 }
 
-bool FCoreStyle::IsStarshipStyle()
-{
-	return !FParse::Param(FCommandLine::Get(), TEXT("oldeditorstyle"));
-}
-
 /* FSlateThrottleManager implementation
  *****************************************************************************/
 
diff --git a/Engine/Source/Runtime/SlateCore/Public/Styling/CoreStyle.h b/Engine/Source/Runtime/SlateCore/Public/Styling/CoreStyle.h
--- a/Engine/Source/Runtime/SlateCore/Public/Styling/CoreStyle.h
+++ b/Engine/Source/Runtime/SlateCore/Public/Styling/CoreStyle.h
@@ -70,7 +70,8 @@ public:
 	static const int32 RegularTextSize = 9;
 	static const int32 SmallTextSize = 8;
 
-	static SLATECORE_API bool IsStarshipStyle();
+	UE_DEPRECATED(5.6, "IsStarshipStyle is deprecated, this function will always return true. Please remove any special case handling for the legacy slate style and calls to this function.")
+	static constexpr SLATECORE_API bool IsStarshipStyle() { return true; }
 
 	static bool IsInitialized() { return Instance.IsValid(); }
 
