Explicit mode icon

#rnx

diff --git a/Engine/Source/Editor/EditorStyle/Private/StarshipStyle.cpp b/Engine/Source/Editor/EditorStyle/Private/StarshipStyle.cpp
--- a/Engine/Source/Editor/EditorStyle/Private/StarshipStyle.cpp
+++ b/Engine/Source/Editor/EditorStyle/Private/StarshipStyle.cpp
@@ -4643,12 +4643,11 @@ void FStarshipEditorStyle::FStyle::SetupLevelEditorStyle()
 		Set("EditorViewport.ToggleSurfaceSnapping", new IMAGE_BRUSH_SVG("Starship/EditorViewport/surface-snap", Icon16x16));
 		Set("EditorViewport.ToggleInGameExposure", new IMAGE_BRUSH_SVG("Starship/Common/GameSettings_16", Icon16x16));
 
-		Set( "EditorViewport.RelativeCoordinateSystem_World", new IMAGE_BRUSH_SVG( "Starship/EditorViewport/globe", Icon16x16 ) );
-		Set("EditorViewport.RelativeCoordinateSystem_Local",
-			new IMAGE_BRUSH_SVG("Starship/Common/transform-local", Icon16x16));
-		Set("EditorViewport.RelativeCoordinateSystem_Parent",
-			new CORE_IMAGE_BRUSH_SVG("Starship/Common/ParentHierarchy", Icon16x16));
-
+		Set("EditorViewport.RelativeCoordinateSystem_World", new IMAGE_BRUSH_SVG( "Starship/EditorViewport/globe", Icon16x16 ));
+		Set("EditorViewport.RelativeCoordinateSystem_Local", new IMAGE_BRUSH_SVG("Starship/Common/transform-local", Icon16x16));
+		Set("EditorViewport.RelativeCoordinateSystem_Parent", new CORE_IMAGE_BRUSH_SVG("Starship/Common/ParentHierarchy", Icon16x16));
+		Set("EditorViewport.RelativeCoordinateSystem_Explicit", new IMAGE_BRUSH_SVG("Starship/Common/transform-explicit", Icon16x16));
+		
 		Set("EditorViewport.RestoreCoordinateSpaceOnSwitch",
 			new IMAGE_BRUSH_SVG("Starship/EditorViewport/RestoreCoordinateSpace_16", Icon16x16));
 		Set("EditorViewport.LocalTransformsInEachLocalSpace",
diff --git a/Engine/Source/Editor/UnrealEd/Private/ViewportToolbar/UnrealEdViewportToolbar.cpp b/Engine/Source/Editor/UnrealEd/Private/ViewportToolbar/UnrealEdViewportToolbar.cpp
--- a/Engine/Source/Editor/UnrealEd/Private/ViewportToolbar/UnrealEdViewportToolbar.cpp
+++ b/Engine/Source/Editor/UnrealEd/Private/ViewportToolbar/UnrealEdViewportToolbar.cpp
@@ -671,18 +671,17 @@ FSlateIcon GetIconFromCoordSystem(ECoordSystem InCoordSystem)
 	}
 	else if (InCoordSystem == COORD_Parent)
 	{
-		static const FName ParentIcon("Icons.ConstraintManager.ParentHierarchy");
+		static const FName ParentIcon("EditorViewport.RelativeCoordinateSystem_Parent");
 		return FSlateIcon(FAppStyle::GetAppStyleSetName(), ParentIcon);
 	}
 	else if (InCoordSystem == COORD_Explicit)
 	{
-		// TODO: find a better icon for this
-		static const FName GimbalIcon("EditorViewport.RotateMode");
-		return FSlateIcon(FAppStyle::GetAppStyleSetName(), GimbalIcon);
+		static const FName ExplicitIcon("EditorViewport.RelativeCoordinateSystem_Explicit");
+		return FSlateIcon(FAppStyle::GetAppStyleSetName(), ExplicitIcon);
 	}
 	else
 	{
-		static FName LocalIcon("Icons.Transform");
+		static FName LocalIcon("EditorViewport.RelativeCoordinateSystem_Local");
 		return FSlateIcon(FAppStyle::GetAppStyleSetName(), LocalIcon);
 	}
 }
