Add support for toolbar button groups in UToolMenus.

diff --git a/Engine/Source/Developer/ToolMenus/Private/ToolMenus.cpp b/Engine/Source/Developer/ToolMenus/Private/ToolMenus.cpp
--- a/Engine/Source/Developer/ToolMenus/Private/ToolMenus.cpp
+++ b/Engine/Source/Developer/ToolMenus/Private/ToolMenus.cpp
@@ -1590,6 +1590,7 @@ void UToolMenus::PopulateToolBarBuilderWithTopLevelChildren(
 	// Seed the previous section with the first blocks's section so we don't start with adding a separator because
 	// sections seem to have changed.
 	FName PreviousSectionName = BlocksToAdd[0].Get<2>();
+	FName PreviousBlockGroupName = NAME_None;
 	for (int i = 0; i < BlocksToAdd.Num(); ++i)
 	{
 		UToolMenu* const SubMenu = BlocksToAdd[i].Get<0>();
@@ -1649,6 +1650,20 @@ void UToolMenus::PopulateToolBarBuilderWithTopLevelChildren(
 		// Make sure we actually add the entry if the reason we added a separator above was that the section names changed.
 		if (Entry->Type != EMultiBlockType::Separator)
 		{
+			const FName EntryBlockGroupName = Entry->ToolBarData.BlockGroupName;
+			if (EntryBlockGroupName != PreviousBlockGroupName)
+			{
+				if (!PreviousBlockGroupName.IsNone())
+				{
+					ToolBarBuilder.EndBlockGroup();
+				}
+				if (!EntryBlockGroupName.IsNone())
+				{
+					ToolBarBuilder.BeginBlockGroup();
+				}
+			}
+			PreviousBlockGroupName = EntryBlockGroupName;
+		
 			constexpr bool bRaiseToTopLevel = true;
 			PopulateToolBarBuilderWithEntry(ToolBarBuilder, SubMenu, *Entry, bRaiseToTopLevel);
 			bHasRaisedEntrySinceLastSeparator = true;
@@ -1664,6 +1679,11 @@ void UToolMenus::PopulateToolBarBuilderWithTopLevelChildren(
 
 		PreviousSectionName = SectionName;
 	}
+	
+	if (!PreviousBlockGroupName.IsNone())
+	{
+		ToolBarBuilder.EndBlockGroup();
+	}
 
 	if (bAddSpaceAfterRaisedChildren)
 	{
@@ -2100,6 +2120,7 @@ void UToolMenus::PopulateToolBarBuilder(FToolBarBuilder& ToolBarBuilder, UToolMe
 			const bool bSectionShouldHaveSeparator = !bFirstSectionInAlignmentGroup;
 			ToolBarBuilder.BeginSection(Section.Name, bSectionShouldHaveSeparator);
 
+			FName PreviousBlockGroupName = NAME_None;
 			for (int BlockIndex = 0; BlockIndex < Section.Blocks.Num(); ++BlockIndex)
 			{
 				const bool bIsLastBlock = BlockIndex >= Section.Blocks.Num() - 1;
@@ -2107,8 +2128,28 @@ void UToolMenus::PopulateToolBarBuilder(FToolBarBuilder& ToolBarBuilder, UToolMe
 
 				const bool bRaiseToTopLevel = false;
 				const bool bIsLastBlockOfLastSection = bIsLastBlock && bIsLastSection;
+				
+				const FName EntryBlockGroupName = Block.ToolBarData.BlockGroupName;
+				if (EntryBlockGroupName != PreviousBlockGroupName)
+				{
+					if (!PreviousBlockGroupName.IsNone())
+					{
+						ToolBarBuilder.EndBlockGroup();
+					}
+					if (!EntryBlockGroupName.IsNone())
+					{
+						ToolBarBuilder.BeginBlockGroup();
+					}
+				}
+				PreviousBlockGroupName = EntryBlockGroupName;
+				
 				PopulateToolBarBuilderWithEntry(ToolBarBuilder, MenuData, Block, bRaiseToTopLevel, bIsLastBlockOfLastSection);
 			}
+			
+			if (!PreviousBlockGroupName.IsNone())
+			{
+				ToolBarBuilder.EndBlockGroup();
+			}
 
 			ToolBarBuilder.EndSection();
 
diff --git a/Engine/Source/Developer/ToolMenus/Public/ToolMenuEntry.h b/Engine/Source/Developer/ToolMenus/Public/ToolMenuEntry.h
--- a/Engine/Source/Developer/ToolMenus/Public/ToolMenuEntry.h
+++ b/Engine/Source/Developer/ToolMenus/Public/ToolMenuEntry.h
@@ -62,6 +62,9 @@ public:
 
 	/** Optional style override used when this entry appears in a toolbar. */
 	FName StyleNameOverride = NAME_None;
+	
+	/** When set, entries with the same name will be grouped together. */
+	FName BlockGroupName = NAME_None;
 
 	/** Delegate that generates a widget for this combo button's menu content. Called when the menu is summoned. */
 	FNewToolMenuChoice ComboButtonContextMenuGenerator;
diff --git a/Engine/Source/Editor/EditorStyle/Private/StarshipStyle.cpp b/Engine/Source/Editor/EditorStyle/Private/StarshipStyle.cpp
--- a/Engine/Source/Editor/EditorStyle/Private/StarshipStyle.cpp
+++ b/Engine/Source/Editor/EditorStyle/Private/StarshipStyle.cpp
@@ -2496,19 +2496,30 @@ void FStarshipEditorStyle::FStyle::SetupViewportStyles()
 {
 	// New viewport toolbar.
 	{
+		const FSlateRoundedBoxBrush TransparentBrush(FLinearColor(0.f, 0.f, 0.f, 0.f), 0.f, FLinearColor(0.f, 0.f, 0.f, 0.f), 0.f);
+	
 		FLinearColor ToolbarBackgroundColor = FStyleColors::Dropdown.GetSpecifiedColor();
 		ToolbarBackgroundColor.A = .90f;
 		const FSlateRoundedBoxBrush BackgroundBrush(ToolbarBackgroundColor, 5.f, ToolbarBackgroundColor, 1.0f);
+		const FSlateRoundedBoxBrush BackgroundBrushStart(ToolbarBackgroundColor, FVector4(5.f, 0.f, 0.f, 5.f), ToolbarBackgroundColor, 1.0f);
+		const FSlateRoundedBoxBrush BackgroundBrushMiddle(ToolbarBackgroundColor, FVector4(0.f, 0.f, 0.f, 0.f), ToolbarBackgroundColor, 1.0f);
+		const FSlateRoundedBoxBrush BackgroundBrushEnd(ToolbarBackgroundColor, FVector4(0.f, 5.f, 5.f, 0.f), ToolbarBackgroundColor, 1.0f);
 
 		FLinearColor ToolbarHoveredColor = FStyleColors::Dropdown.GetSpecifiedColor();
 		ToolbarHoveredColor.R = FMath::Min(1.0f, ToolbarHoveredColor.R * 1.5f);
 		ToolbarHoveredColor.G = FMath::Min(1.0f, ToolbarHoveredColor.G * 1.5f);
 		ToolbarHoveredColor.B = FMath::Min(1.0f, ToolbarHoveredColor.B * 1.5f);
 		const FSlateRoundedBoxBrush HoveredButtonBrush(ToolbarHoveredColor, 5.f, ToolbarHoveredColor, 1.0f);
+		const FSlateRoundedBoxBrush HoveredButtonBrushStart(ToolbarHoveredColor, FVector4(5.f, 0.f, 0.f, 5.f), ToolbarHoveredColor, 1.0f);
+		const FSlateRoundedBoxBrush HoveredButtonBrushMiddle(ToolbarHoveredColor, FVector4(0.f, 0.f, 0.f, 0.f), ToolbarHoveredColor, 1.0f);
+		const FSlateRoundedBoxBrush HoveredButtonBrushEnd(ToolbarHoveredColor, FVector4(0.f, 5.f, 5.f, 0.f), ToolbarHoveredColor, 1.0f);
 
 		FLinearColor ToolbarPressedColor = FStyleColors::Recessed.GetSpecifiedColor();
 		ToolbarPressedColor.A = .50f;
 		const FSlateRoundedBoxBrush PressedButtonBrush(ToolbarPressedColor, 5.f, ToolbarPressedColor, 1.0f);
+		const FSlateRoundedBoxBrush PressedButtonBrushStart(ToolbarPressedColor, FVector4(5.f, 0.f, 0.f, 5.f), ToolbarPressedColor, 1.0f);
+		const FSlateRoundedBoxBrush PressedButtonBrushMiddle(ToolbarPressedColor, FVector4(0.f, 0.f, 0.f, 0.f), ToolbarPressedColor, 1.0f);
+		const FSlateRoundedBoxBrush PressedButtonBrushEnd(ToolbarPressedColor, FVector4(0.f, 5.f, 5.f, 0.f), ToolbarPressedColor, 1.0f);
 
 		FButtonStyle ButtonStyle;
 		{
@@ -2522,23 +2533,59 @@ void FStarshipEditorStyle::FStyle::SetupViewportStyles()
 				.SetNormalPadding(FMargin(6.0f, 4.0f))
 				.SetPressedPadding(FMargin(6.0f, 4.0f));
 		}
+		
+		FButtonStyle ButtonStyleStart = FButtonStyle(ButtonStyle)
+			.SetNormal(BackgroundBrushStart)
+			.SetHovered(HoveredButtonBrushStart)
+			.SetPressed(PressedButtonBrushStart);
+		
+		FButtonStyle ButtonStyleMiddle = FButtonStyle(ButtonStyle)
+			.SetNormal(BackgroundBrushMiddle)
+			.SetHovered(HoveredButtonBrushMiddle)
+			.SetPressed(PressedButtonBrushMiddle);
+		
+		FButtonStyle ButtonStyleEnd = FButtonStyle(ButtonStyle)
+			.SetNormal(BackgroundBrushEnd)
+			.SetHovered(HoveredButtonBrushEnd)
+			.SetPressed(PressedButtonBrushEnd);
+		
 
 		const FToolBarStyle SlimToolBarStyle =
 			FStarshipCoreStyle::GetCoreStyle().GetWidgetStyle<FToolBarStyle>("SlimToolBar");
 
-		FCheckBoxStyle ToggleButtonStyle = SlimToolBarStyle.ToggleButton;
-		{
-			const FSlateRoundedBoxBrush TransparentBrush(FLinearColor(0.f, 0.f, 0.f, 0.f), 0.f, FLinearColor(0.f, 0.f, 0.f, 0.f), 0.f);
-
-			ToggleButtonStyle.SetCheckedImage(TransparentBrush)
-				.SetCheckedHoveredImage(HoveredButtonBrush)
-				.SetCheckedPressedImage(PressedButtonBrush)
-				.SetUncheckedPressedImage(PressedButtonBrush)
-				.SetCheckedForegroundColor(FStyleColors::AccentBlue)
-				.SetCheckedHoveredForegroundColor(FStyleColors::AccentBlue)
-				.SetCheckedPressedForegroundColor(FStyleColors::AccentBlue)
-				.SetPadding(FMargin(4.0f));
-		}
+		FCheckBoxStyle ToggleButtonStyle = FCheckBoxStyle(SlimToolBarStyle.ToggleButton)
+			.SetCheckedImage(TransparentBrush)
+			.SetCheckedHoveredImage(HoveredButtonBrush)
+			.SetCheckedPressedImage(PressedButtonBrush)
+			.SetUncheckedPressedImage(PressedButtonBrush)
+			.SetCheckedForegroundColor(FStyleColors::AccentBlue)
+			.SetCheckedHoveredForegroundColor(FStyleColors::AccentBlue)
+			.SetCheckedPressedForegroundColor(FStyleColors::AccentBlue)
+			.SetPadding(FMargin(4.0f));
+		
+		FCheckBoxStyle ToggleButtonStyleStart = FCheckBoxStyle(ToggleButtonStyle)
+			.SetCheckedImage(BackgroundBrushStart)
+			.SetCheckedHoveredImage(HoveredButtonBrushStart)
+			.SetCheckedPressedImage(PressedButtonBrushStart)
+			.SetUncheckedImage(BackgroundBrushStart)
+			.SetUncheckedHoveredImage(HoveredButtonBrushStart)
+			.SetUncheckedPressedImage(PressedButtonBrushStart);
+		
+		FCheckBoxStyle ToggleButtonStyleMiddle = FCheckBoxStyle(ToggleButtonStyle)
+			.SetCheckedImage(BackgroundBrushMiddle)
+			.SetCheckedHoveredImage(HoveredButtonBrushMiddle)
+			.SetCheckedPressedImage(PressedButtonBrushMiddle)
+			.SetUncheckedImage(BackgroundBrushMiddle)
+			.SetUncheckedHoveredImage(HoveredButtonBrushMiddle)
+			.SetUncheckedPressedImage(PressedButtonBrushMiddle);
+		
+		FCheckBoxStyle ToggleButtonStyleEnd = FCheckBoxStyle(ToggleButtonStyle)
+			.SetCheckedImage(BackgroundBrushEnd)
+			.SetCheckedHoveredImage(HoveredButtonBrushEnd)
+			.SetCheckedPressedImage(PressedButtonBrushEnd)
+			.SetUncheckedImage(BackgroundBrushEnd)
+			.SetUncheckedHoveredImage(HoveredButtonBrushEnd)
+			.SetUncheckedPressedImage(PressedButtonBrushEnd);
 
 		FComboButtonStyle ComboButtonStyle = SlimToolBarStyle.ComboButtonStyle;
 		{
@@ -2559,7 +2606,7 @@ void FStarshipEditorStyle::FStyle::SetupViewportStyles()
 			ViewportToolbarStyle.SetIconSize(Icon16x16)
 				.SetButtonStyle(ButtonStyle)
 				.SetToggleButtonStyle(ToggleButtonStyle)
-				.SetButtonPadding(0.f)
+				.SetButtonPadding(FMargin(4.f, 0.f))
 				.SetComboButtonStyle(ComboButtonStyle)
 				.SetComboButtonPadding(FMargin(4.0f, 0.f))
 				.SetSeparatorBrush(SeparatorBrush)
@@ -2571,36 +2618,62 @@ void FStarshipEditorStyle::FStyle::SetupViewportStyles()
 				.SetRaisedChildrenRightPadding(16.0f);
 
 			Set("ViewportToolbar", ViewportToolbarStyle);
+			
+			Set("ViewportToolbar.Button.Start", ButtonStyleStart);
+			Set("ViewportToolbar.Button.Middle", ButtonStyleMiddle);
+			Set("ViewportToolbar.Button.End", ButtonStyleEnd);
+			
+			Set("ViewportToolbar.ToggleButton.Start", ToggleButtonStyleStart);
+			Set("ViewportToolbar.ToggleButton.Middle", ToggleButtonStyleMiddle);
+			Set("ViewportToolbar.ToggleButton.End", ToggleButtonStyleEnd);
 		}
 
 		// Special styling for top-level raised buttons.
 		FToolBarStyle ViewportToolbarRaisedStyle = ViewportToolbarStyle;
 		{
-			FButtonStyle ButtonStyleRaised = ButtonStyle;
+			const auto& ConvertButtonToRaised = [&TransparentBrush](const FButtonStyle& ButtonStyle)
 			{
-				FLinearColor BackgroundColorRaised = FStyleColors::Dropdown.GetSpecifiedColor();
-				BackgroundColorRaised.A = .0f;
-				const FSlateRoundedBoxBrush BackgroundBrushRaised(BackgroundColorRaised, 5.f, BackgroundColorRaised, 1.0);
-
-				ButtonStyleRaised
-					.SetNormal(BackgroundBrushRaised)
+				return FButtonStyle(ButtonStyle)
+					.SetNormal(TransparentBrush)
 					.SetNormalForeground(FStyleColors::Foreground)
 					.SetNormalPadding(4.0f)
 					.SetPressedPadding(4.0f);
-			}
+			};
+		
+			FButtonStyle ButtonStyleRaised = ConvertButtonToRaised(ButtonStyle);
+			FButtonStyle ButtonStyleRaisedStart = ConvertButtonToRaised(ButtonStyleStart);
+			FButtonStyle ButtonStyleRaisedMiddle = ConvertButtonToRaised(ButtonStyleMiddle);
+			FButtonStyle ButtonStyleRaisedEnd = ConvertButtonToRaised(ButtonStyleEnd);
 			
-			FCheckBoxStyle ToggleButtonStyleRaised = ToggleButtonStyle;
-			{
-				ToggleButtonStyleRaised.SetUncheckedHoveredImage(HoveredButtonBrush)
-					.SetCheckedPressedImage(BackgroundBrush)
-					.SetUncheckedPressedImage(BackgroundBrush);
-			}
-
-			FComboButtonStyle ComboStyleRaised = ViewportToolbarStyle.ComboButtonStyle;
-			{
-				ComboButtonStyle.SetDownArrowPadding(FMargin(2.0f, 0, 0, 0));
-			}
-
+			FCheckBoxStyle ToggleButtonStyleRaised = FCheckBoxStyle(ToggleButtonStyle)
+				.SetUncheckedHoveredImage(HoveredButtonBrush)
+				.SetCheckedPressedImage(BackgroundBrush)
+				.SetUncheckedPressedImage(BackgroundBrush);
+			
+			FCheckBoxStyle ToggleButtonStyleRaisedStart = FCheckBoxStyle(ToggleButtonStyleStart)
+				.SetUncheckedImage(TransparentBrush)
+				.SetCheckedImage(TransparentBrush)
+				.SetUncheckedHoveredImage(HoveredButtonBrushStart)
+				.SetCheckedPressedImage(BackgroundBrushStart)
+				.SetUncheckedPressedImage(BackgroundBrushStart);
+				
+			FCheckBoxStyle ToggleButtonStyleRaisedMiddle = FCheckBoxStyle(ToggleButtonStyleMiddle)
+				.SetUncheckedImage(TransparentBrush)
+				.SetCheckedImage(TransparentBrush)
+				.SetUncheckedHoveredImage(HoveredButtonBrushMiddle)
+				.SetCheckedPressedImage(BackgroundBrushMiddle)
+				.SetUncheckedPressedImage(BackgroundBrushMiddle);
+				
+			FCheckBoxStyle ToggleButtonStyleRaisedEnd = FCheckBoxStyle(ToggleButtonStyleEnd)
+				.SetUncheckedImage(TransparentBrush)
+				.SetCheckedImage(TransparentBrush)
+				.SetUncheckedHoveredImage(HoveredButtonBrushEnd)
+				.SetCheckedPressedImage(BackgroundBrushEnd)
+				.SetUncheckedPressedImage(BackgroundBrushEnd);
+
+			FComboButtonStyle ComboStyleRaised = FComboButtonStyle(ViewportToolbarStyle.ComboButtonStyle)
+				.SetDownArrowPadding(FMargin(2.0f, 0, 0, 0));
+				
 			ViewportToolbarRaisedStyle
 				.SetButtonStyle(ButtonStyleRaised)
 				.SetToggleButtonStyle(ToggleButtonStyleRaised)
@@ -2610,6 +2683,14 @@ void FStarshipEditorStyle::FStyle::SetupViewportStyles()
 				.SetComboButtonPadding(FMargin(3.0f, 0));
 			
 			Set("ViewportToolbar.Raised", ViewportToolbarRaisedStyle);
+			
+			Set("ViewportToolbar.Raised.Button.Start", ButtonStyleRaisedStart);
+			Set("ViewportToolbar.Raised.Button.Middle", ButtonStyleRaisedMiddle);
+			Set("ViewportToolbar.Raised.Button.End", ButtonStyleRaisedEnd);
+			
+			Set("ViewportToolbar.Raised.ToggleButton.Start", ToggleButtonStyleRaisedStart);
+			Set("ViewportToolbar.Raised.ToggleButton.Middle", ToggleButtonStyleRaisedMiddle);
+			Set("ViewportToolbar.Raised.ToggleButton.End", ToggleButtonStyleRaisedEnd);
 		}
 
 		// Create a special viewport toolbar style for the level editor. To accomodate the persistent top-right min/max
diff --git a/Engine/Source/Runtime/Slate/Private/Framework/MultiBox/SToolBarButtonBlock.cpp b/Engine/Source/Runtime/Slate/Private/Framework/MultiBox/SToolBarButtonBlock.cpp
--- a/Engine/Source/Runtime/Slate/Private/Framework/MultiBox/SToolBarButtonBlock.cpp
+++ b/Engine/Source/Runtime/Slate/Private/Framework/MultiBox/SToolBarButtonBlock.cpp
@@ -485,8 +485,9 @@ void SToolBarButtonBlock::BuildMultiBlockWidget(const ISlateStyle* StyleSet, con
 	}
 	else
 	{
-		// Space between buttons. It does not make the buttons larger
-		ChildSlot.Padding(ToolBarStyle.ButtonPadding);
+		// Space between buttons. It does not make the buttons larger.
+		// Button groups eliminate spacing between buttons.
+		ChildSlot.Padding(EMultiBlockLocation::ToHorizontalMargin(MultiBox->GetType(), BlockLocation, ToolBarStyle.ButtonPadding));
 	}
 
 	// Bind our widget's visible state to whether or not the button should be visible
diff --git a/Engine/Source/Runtime/Slate/Public/Framework/MultiBox/MultiBoxDefs.h b/Engine/Source/Runtime/Slate/Public/Framework/MultiBox/MultiBoxDefs.h
--- a/Engine/Source/Runtime/Slate/Public/Framework/MultiBox/MultiBoxDefs.h
+++ b/Engine/Source/Runtime/Slate/Public/Framework/MultiBox/MultiBoxDefs.h
@@ -137,6 +137,46 @@ namespace EMultiBlockLocation
 		}
 		return StyleName;
 	}
+
+	/**
+	 * Trims a provided margin based on the block location. Grouped multiblocks want to be
+	 * directly next to each other so that they can be presented as a single visual item.
+	 * This function allows for defining a single "padding between items" value and have the
+	 * appropriate sides of that padding zeroed out. 
+	 * @param InBoxType The type of box the group is within. Used to determine horizontal/vertical flow.
+	 * @param InLocation The location of the block within the group. 
+	 * @param Margin The margin around a group or individual block.
+	 * @return The margin for the block at the given location.
+	 */
+	static FMargin ToHorizontalMargin(EMultiBoxType InBoxType, Type InLocation, FMargin Margin)
+	{
+		if (InBoxType == EMultiBoxType::VerticalToolBar || InBoxType == EMultiBoxType::Menu)
+		{
+			// Vertical groups cut off the tops & bottoms
+			switch (InLocation)
+			{
+			case Start:
+				return FMargin(Margin.Left, Margin.Top, Margin.Right, 0.f);
+			case Middle:
+				return FMargin(Margin.Left, 0.f, Margin.Right, 0.f);
+			case End:
+				return FMargin(Margin.Left, 0.f, Margin.Right, Margin.Bottom);
+			}
+			return Margin;
+		}
+		
+		// Horizontal groups cut off the left & right
+		switch (InLocation)
+		{
+		case Start:
+			return FMargin(Margin.Left, Margin.Top, 0.f, Margin.Bottom);
+		case Middle:
+			return FMargin(0.f, Margin.Top, 0.f, Margin.Bottom);
+		case End:
+			return FMargin(0.f, Margin.Top, Margin.Right, Margin.Bottom);
+		}
+		return Margin;
+	}
 }
 
 /** Contains various Style parameters and overrides. Not all are applicable to a given entry */
