Viewport toolbar: Removed the extra padding around the coordinate space combobutton

diff --git a/Engine/Source/Editor/EditorStyle/Private/StarshipStyle.cpp b/Engine/Source/Editor/EditorStyle/Private/StarshipStyle.cpp
--- a/Engine/Source/Editor/EditorStyle/Private/StarshipStyle.cpp
+++ b/Engine/Source/Editor/EditorStyle/Private/StarshipStyle.cpp
@@ -2673,14 +2673,24 @@ void FStarshipEditorStyle::FStyle::SetupViewportStyles()
 
 			FComboButtonStyle ComboStyleRaised = FComboButtonStyle(ViewportToolbarStyle.ComboButtonStyle)
 				.SetDownArrowPadding(FMargin(2.0f, 0, 0, 0));
-				
+
+			FComboButtonStyle SettingsComboStyleRaised = ViewportToolbarStyle.SettingsComboButton;
+			{
+				SettingsComboStyleRaised.ButtonStyle
+					.SetNormalPadding(FMargin(2.f, 0.f))
+					.SetPressedPadding(FMargin(2.f, 0.f))
+					.SetHovered(HoveredButtonBrush)
+					.SetHoveredForeground(FStyleColors::ForegroundHover);
+			}
+
 			ViewportToolbarRaisedStyle
 				.SetButtonStyle(ButtonStyleRaised)
 				.SetToggleButtonStyle(ToggleButtonStyleRaised)
 				.SetComboButtonStyle(ComboStyleRaised)
 				.SetBlockHovered(BackgroundBrush)
 				.SetButtonPadding(FMargin(3.0f, 0))
-				.SetComboButtonPadding(FMargin(3.0f, 0));
+				.SetComboButtonPadding(FMargin(3.0f, 0))
+				.SetSettingsComboButtonStyle(SettingsComboStyleRaised);
 			
 			Set("ViewportToolbar.Raised", ViewportToolbarRaisedStyle);
 			
diff --git a/Engine/Source/Runtime/Slate/Private/Framework/MultiBox/SToolBarComboButtonBlock.cpp b/Engine/Source/Runtime/Slate/Private/Framework/MultiBox/SToolBarComboButtonBlock.cpp
--- a/Engine/Source/Runtime/Slate/Private/Framework/MultiBox/SToolBarComboButtonBlock.cpp
+++ b/Engine/Source/Runtime/Slate/Private/Framework/MultiBox/SToolBarComboButtonBlock.cpp
@@ -275,8 +275,6 @@ void SToolBarComboButtonBlock::BuildMultiBlockWidget(const ISlateStyle* StyleSet
 	{
 		if (HasAction())
 		{
-			ComboStyle = &ToolBarStyle.SettingsComboButton;
-
 			// clang-format off
 			SAssignNew(LeftHandSideWidget, SButton)
 			.ButtonStyle(ButtonStyle)
@@ -289,6 +287,9 @@ void SToolBarComboButtonBlock::BuildMultiBlockWidget(const ISlateStyle* StyleSet
 				IconWidget
 			];
 			// clang-format on
+
+			ComboStyle = &ToolBarStyle.SettingsComboButton;
+			ButtonStyle = &ToolBarStyle.SettingsComboButton.ButtonStyle;
 		}
 
 		// clang-format off
