[Sequencer] Navigation Tool:
Fix missing slate resource
Rename GetExtendedColumnNames() to GetExtendedColumnViewNames()

diff --git a/Engine/Source/Editor/EditorStyle/Private/StarshipStyle.cpp b/Engine/Source/Editor/EditorStyle/Private/StarshipStyle.cpp
--- a/Engine/Source/Editor/EditorStyle/Private/StarshipStyle.cpp
+++ b/Engine/Source/Editor/EditorStyle/Private/StarshipStyle.cpp
@@ -2525,9 +2525,13 @@ void FStarshipEditorStyle::FStyle::SetupSequencerStyles()
 		EndTimeSpinStyle.SetForegroundColor(FStyleColors::AccentRed);
 		Set("Sequencer.NavigationTool.OutTime.SpinBox", EndTimeSpinStyle);
 
-		FEditableTextBoxStyle NonEditableTextBoxStyle = GetWidgetStyle<FEditableTextBoxStyle>(TEXT("Graph.EditableTextBox"));
-		NonEditableTextBoxStyle.SetForegroundColor(FSlateColor::UseSubduedForeground());
-		NonEditableTextBoxStyle.BackgroundImageReadOnly = FSlateRoundedBoxBrush(FStyleColors::Background, 4.0f, FSlateColor::UseSubduedForeground(), 1.0f);
+		FEditableTextBoxStyle NonEditableTextBoxStyle = FEditableTextBoxStyle()
+			.SetPadding(FMargin(4.f))
+			.SetForegroundColor(FSlateColor::UseSubduedForeground())
+			.SetBackgroundImageNormal(BOX_BRUSH("Graph/CommonWidgets/TextBox", FMargin(4.0f/16.0f)))
+			.SetBackgroundImageHovered(BOX_BRUSH("Graph/CommonWidgets/TextBox_Hovered", FMargin(4.0f/16.0f)))
+			.SetBackgroundImageFocused(BOX_BRUSH("Graph/CommonWidgets/TextBox_Hovered", FMargin(4.0f/16.0f)))
+			.SetBackgroundImageReadOnly(FSlateRoundedBoxBrush(FStyleColors::Background, 4.0f, FStyleColors::InputOutline, 1.0f));
 		Set("Sequencer.NavigationTool.NonEditableTextBox", NonEditableTextBoxStyle);
 	}
 }
diff --git a/Engine/Source/Editor/Sequencer/Private/Tools/NavigationTool/NavigationTool.cpp b/Engine/Source/Editor/Sequencer/Private/Tools/NavigationTool/NavigationTool.cpp
--- a/Engine/Source/Editor/Sequencer/Private/Tools/NavigationTool/NavigationTool.cpp
+++ b/Engine/Source/Editor/Sequencer/Private/Tools/NavigationTool/NavigationTool.cpp
@@ -202,7 +202,7 @@ void FNavigationTool::OnProvidersChanged(const TSharedRef<INavigationTool>& InTo
 		UNavigationToolSettings* const ToolSettings = GetMutableDefault<UNavigationToolSettings>();
 		TSet<FNavigationToolColumnView>& CustomColumnViews = ToolSettings->GetCustomColumnViews();
 
-		for (const FName ColumnViewName : InProvider->GetExtendedColumnNames())
+		for (const FName ColumnViewName : InProvider->GetExtendedColumnViewNames())
 		{
 			CustomColumnViews.Remove(ColumnViewName);
 		}
diff --git a/Engine/Source/Editor/Sequencer/Private/Tools/NavigationTool/NavigationToolView.cpp b/Engine/Source/Editor/Sequencer/Private/Tools/NavigationTool/NavigationToolView.cpp
--- a/Engine/Source/Editor/Sequencer/Private/Tools/NavigationTool/NavigationToolView.cpp
+++ b/Engine/Source/Editor/Sequencer/Private/Tools/NavigationTool/NavigationToolView.cpp
@@ -138,12 +138,6 @@ void FNavigationToolView::CreateColumns(const TSharedRef<FNavigationToolProvider
 
 void FNavigationToolView::CreateDefaultColumnViews(const TSharedRef<FNavigationToolProvider>& InProvider)
 {
-	const TSharedPtr<FNavigationTool> Tool = WeakTool.Pin();
-	if (!Tool.IsValid())
-	{
-		return;
-	}
-
 	InProvider->OnExtendColumnViews(ToolSettings->GetCustomColumnViews());
 	ToolSettings->SaveConfig();
 }
diff --git a/Engine/Source/Editor/Sequencer/Private/Tools/NavigationTool/Providers/NavigationToolProvider.cpp b/Engine/Source/Editor/Sequencer/Private/Tools/NavigationTool/Providers/NavigationToolProvider.cpp
--- a/Engine/Source/Editor/Sequencer/Private/Tools/NavigationTool/Providers/NavigationToolProvider.cpp
+++ b/Engine/Source/Editor/Sequencer/Private/Tools/NavigationTool/Providers/NavigationToolProvider.cpp
@@ -10,10 +10,10 @@ void FNavigationToolProvider::OnExtendColumnViews(TSet<FNavigationToolColumnView
 {
 	INavigationToolProvider::OnExtendColumnViews(OutColumnViews);
 
-	ExtendedColumnNames.Reset(OutColumnViews.Num());
+	ExtendedColumnViewNames.Reset(OutColumnViews.Num());
 	for (const FNavigationToolColumnView& ColumnView : OutColumnViews)
 	{
-		ExtendedColumnNames.Add(ColumnView.ViewName);
+		ExtendedColumnViewNames.Add(ColumnView.ViewName);
 	}
 }
 
@@ -104,9 +104,9 @@ void FNavigationToolProvider::EnsureToolViewCount(const INavigationTool& InTool,
 	}
 }
 
-const TArray<FName>& FNavigationToolProvider::GetExtendedColumnNames() const
+const TArray<FName>& FNavigationToolProvider::GetExtendedColumnViewNames() const
 {
-	return ExtendedColumnNames;
+	return ExtendedColumnViewNames;
 }
 
 TOptional<EItemDropZone> FNavigationToolProvider::OnToolItemCanAcceptDrop(const FDragDropEvent& InDragDropEvent
diff --git a/Engine/Source/Editor/Sequencer/Private/Tools/NavigationTool/Widgets/Columns/SNavigationToolLength.cpp b/Engine/Source/Editor/Sequencer/Private/Tools/NavigationTool/Widgets/Columns/SNavigationToolLength.cpp
--- a/Engine/Source/Editor/Sequencer/Private/Tools/NavigationTool/Widgets/Columns/SNavigationToolLength.cpp
+++ b/Engine/Source/Editor/Sequencer/Private/Tools/NavigationTool/Widgets/Columns/SNavigationToolLength.cpp
@@ -5,7 +5,6 @@
 #include "Sections/MovieSceneSubSection.h"
 #include "Tools/NavigationTool/Items/NavigationToolItemUtils.h"
 #include "Tools/NavigationTool/Items/NavigationToolSequence.h"
-#include "Tools/NavigationTool/NavigationTool.h"
 #include "Tools/NavigationTool/NavigationToolView.h"
 #include "Tools/NavigationTool/Widgets/SNavigationToolTreeRow.h"
 #include "Widgets/Input/SEditableTextBox.h"
diff --git a/Engine/Source/Editor/Sequencer/Public/Tools/NavigationTool/Providers/NavigationToolProvider.h b/Engine/Source/Editor/Sequencer/Public/Tools/NavigationTool/Providers/NavigationToolProvider.h
--- a/Engine/Source/Editor/Sequencer/Public/Tools/NavigationTool/Providers/NavigationToolProvider.h
+++ b/Engine/Source/Editor/Sequencer/Public/Tools/NavigationTool/Providers/NavigationToolProvider.h
@@ -35,10 +35,10 @@ public:
 	/** Ensure saved editor only data contains an entry for a specific view, creating it if necessary */
 	void EnsureToolViewCount(const INavigationTool& InTool, const int32 InToolViewId);
 
-	const TArray<FName>& GetExtendedColumnNames() const;
+	const TArray<FName>& GetExtendedColumnViewNames() const;
 
 protected:
-	TArray<FName> ExtendedColumnNames;
+	TArray<FName> ExtendedColumnViewNames;
 };
 
 }
