[Reimport Icon] Adding a new reimport icon and applying it to icons that were using the import icon (when they should be using a reimport one)

diff --git a/Engine/Plugins/Animation/RigLogic/Source/RigLogicEditor/Private/RigLogicEditor.cpp b/Engine/Plugins/Animation/RigLogic/Source/RigLogicEditor/Private/RigLogicEditor.cpp
--- a/Engine/Plugins/Animation/RigLogic/Source/RigLogicEditor/Private/RigLogicEditor.cpp
+++ b/Engine/Plugins/Animation/RigLogic/Source/RigLogicEditor/Private/RigLogicEditor.cpp
@@ -86,7 +86,7 @@ void FRigLogicEditor::GetDNAMenu(FMenuBuilder& MenuBuilder, const TArray<FAssetD
 	MenuBuilder.AddMenuEntry(
 		LOCTEXT("Reimport DNA", "Reimport existing DNA File"),
 		LOCTEXT("ReimportDNA_Tooltip ", "Reimport DNA"),
-		FSlateIcon(FAppStyle::GetAppStyleSetName(), "Icons.Import"),
+		FSlateIcon(FAppStyle::GetAppStyleSetName(), "Icons.Reimport"),
 		FUIAction(FExecuteAction::CreateStatic(&FRigLogicEditor::ExecuteDNAReimport, static_cast<UObject*>(Mesh)))
 	);
 }
diff --git a/Engine/Plugins/Editor/ContentBrowser/ContentBrowserAssetDataSource/Source/ContentBrowserAssetDataSource/Private/AssetFileContextMenu.cpp b/Engine/Plugins/Editor/ContentBrowser/ContentBrowserAssetDataSource/Source/ContentBrowserAssetDataSource/Private/AssetFileContextMenu.cpp
--- a/Engine/Plugins/Editor/ContentBrowser/ContentBrowserAssetDataSource/Source/ContentBrowserAssetDataSource/Private/AssetFileContextMenu.cpp
+++ b/Engine/Plugins/Editor/ContentBrowser/ContentBrowserAssetDataSource/Source/ContentBrowserAssetDataSource/Private/AssetFileContextMenu.cpp
@@ -156,7 +156,7 @@ bool FAssetFileContextMenu::AddImportedAssetMenuOptions(UToolMenu* Menu)
 								NAME_None,
 								ReimportLabel,
 								ReimportLabelTooltip,
-								FSlateIcon(FAppStyle::GetAppStyleSetName(), "Icons.Import"),
+								FSlateIcon(FAppStyle::GetAppStyleSetName(), "Icons.Reimport"),
 								FUIAction(
 									FExecuteAction::CreateSP(this, &FAssetFileContextMenu::ExecuteReimportWithNewFile, SourceFileIndex),
 									FCanExecuteAction()
@@ -169,7 +169,7 @@ bool FAssetFileContextMenu::AddImportedAssetMenuOptions(UToolMenu* Menu)
 								NAME_None,
 								ReimportLabel,
 								ReimportLabelTooltip,
-								FSlateIcon(FAppStyle::GetAppStyleSetName(), "Icons.Import"),
+								FSlateIcon(FAppStyle::GetAppStyleSetName(), "Icons.Reimport"),
 								FUIAction(
 									FExecuteAction::CreateSP(this, &FAssetFileContextMenu::ExecuteReimport, SourceFileIndex),
 									FCanExecuteAction()
@@ -195,7 +195,7 @@ bool FAssetFileContextMenu::AddImportedAssetMenuOptions(UToolMenu* Menu)
 					FText(),
 					FNewToolMenuDelegate::CreateLambda(CreateSubMenu, false),
 					false,
-					FSlateIcon(FAppStyle::GetAppStyleSetName(), "Icons.Import"));
+					FSlateIcon(FAppStyle::GetAppStyleSetName(), "Icons.Reimport"));
 				//With new file
 				Section.AddSubMenu(
 					"ReimportWithNewFile",
@@ -203,7 +203,7 @@ bool FAssetFileContextMenu::AddImportedAssetMenuOptions(UToolMenu* Menu)
 					FText(),
 					FNewToolMenuDelegate::CreateLambda(CreateSubMenu, true),
 					false,
-					FSlateIcon(FAppStyle::GetAppStyleSetName(), "Icons.Import"));
+					FSlateIcon(FAppStyle::GetAppStyleSetName(), "Icons.Reimport"));
 			}
 			else
 			{
@@ -211,7 +211,7 @@ bool FAssetFileContextMenu::AddImportedAssetMenuOptions(UToolMenu* Menu)
 					"Reimport",
 					LOCTEXT("Reimport", "Reimport"),
 					LOCTEXT("ReimportTooltip", "Reimport the selected asset(s) from the source file on disk."),
-					FSlateIcon(FAppStyle::GetAppStyleSetName(), "Icons.Import"),
+					FSlateIcon(FAppStyle::GetAppStyleSetName(), "Icons.Reimport"),
 					FUIAction(
 						FExecuteAction::CreateSP(this, &FAssetFileContextMenu::ExecuteReimport, (int32)INDEX_NONE),
 						FCanExecuteAction::CreateSP(this, &FAssetFileContextMenu::CanExecuteReimportAssetActions, ResolvedFilePaths)
@@ -224,7 +224,7 @@ bool FAssetFileContextMenu::AddImportedAssetMenuOptions(UToolMenu* Menu)
 						"ReimportWithNewFile",
 						LOCTEXT("ReimportWithNewFile", "Reimport With New File"),
 						LOCTEXT("ReimportWithNewFileTooltip", "Reimport the selected asset from a new source file on disk."),
-						FSlateIcon(FAppStyle::GetAppStyleSetName(), "Icons.Import"),
+						FSlateIcon(FAppStyle::GetAppStyleSetName(), "Icons.Reimport"),
 						FUIAction(
 							FExecuteAction::CreateSP(this, &FAssetFileContextMenu::ExecuteReimportWithNewFile, (int32)INDEX_NONE),
 							FCanExecuteAction::CreateSP(this, &FAssetFileContextMenu::CanExecuteReimportAssetActions, ResolvedFilePaths)
diff --git a/Engine/Plugins/Editor/WaveformEditor/Source/WaveformEditor/Private/WaveformEditor.cpp b/Engine/Plugins/Editor/WaveformEditor/Source/WaveformEditor/Private/WaveformEditor.cpp
--- a/Engine/Plugins/Editor/WaveformEditor/Source/WaveformEditor/Private/WaveformEditor.cpp
+++ b/Engine/Plugins/Editor/WaveformEditor/Source/WaveformEditor/Private/WaveformEditor.cpp
@@ -530,7 +530,7 @@ bool FWaveformEditor::RegisterToolbar()
 					Commands.ReimportAsset,
 					FText(),
 					TAttribute< FText >::CreateRaw(LockedObserver.Get(), &FWaveformEditor::GetReimportButtonToolTip),
-					FSlateIcon(FAppStyle::GetAppStyleSetName(), "Icons.Import"));
+					FSlateIcon(FAppStyle::GetAppStyleSetName(), "Icons.Reimport"));
 
 				InSection.AddEntry(ReimportEntry);
 			}
@@ -548,7 +548,7 @@ bool FWaveformEditor::RegisterToolbar()
 					FNewToolMenuChoice(FOnGetContent::CreateSP(LockedObserver.Get(), &FWaveformEditor::GenerateImportOptionsMenu)),
 					LOCTEXT("ReimportOptions_Label", "Reimport Options"),
 					LOCTEXT("ReimportOptions_ToolTip", "Reimport options for this USoundWave"),
-					FSlateIcon(FAppStyle::GetAppStyleSetName(), "Icons.Import"),
+					FSlateIcon(FAppStyle::GetAppStyleSetName(), "Icons.Reimport"),
 					true
 				);
 
diff --git a/Engine/Plugins/Editor/WaveformEditor/Source/WaveformEditor/Private/WaveformEditorInstantiator.cpp b/Engine/Plugins/Editor/WaveformEditor/Source/WaveformEditor/Private/WaveformEditorInstantiator.cpp
--- a/Engine/Plugins/Editor/WaveformEditor/Source/WaveformEditor/Private/WaveformEditorInstantiator.cpp
+++ b/Engine/Plugins/Editor/WaveformEditor/Source/WaveformEditor/Private/WaveformEditorInstantiator.cpp
@@ -55,7 +55,7 @@ void FWaveformEditorInstantiator::ExtendContentBrowserSelectionMenu()
 				{
 					const TAttribute<FText> Label = LOCTEXT("SoundWave_ReimportAndOverwriteTransformations", "Reimport And Overwrite Transformations");
 					const TAttribute<FText> ToolTip = LOCTEXT("SoundWave_ReimportAndOverwriteTransformations_Tooltip", "Reimport And Overwrite Transformations");
-					const FSlateIcon Icon = FSlateIcon(FAppStyle::GetAppStyleSetName(), "Icons.Import");
+					const FSlateIcon Icon = FSlateIcon(FAppStyle::GetAppStyleSetName(), "Icons.Reimport");
 					const FToolMenuExecuteAction UIAction = FToolMenuExecuteAction::CreateSP(this, &FWaveformEditorInstantiator::ExecuteReImportAndOverwriteTransformations);
 
 					InSection.AddMenuEntry("SoundWave_ReimportAndOverwriteTransformations", Label, ToolTip, Icon, UIAction);
diff --git a/Engine/Plugins/Interchange/Editor/Source/InterchangeEditor/Private/AssetDefinition_InterchangeSceneImportAsset.cpp b/Engine/Plugins/Interchange/Editor/Source/InterchangeEditor/Private/AssetDefinition_InterchangeSceneImportAsset.cpp
--- a/Engine/Plugins/Interchange/Editor/Source/InterchangeEditor/Private/AssetDefinition_InterchangeSceneImportAsset.cpp
+++ b/Engine/Plugins/Interchange/Editor/Source/InterchangeEditor/Private/AssetDefinition_InterchangeSceneImportAsset.cpp
@@ -206,7 +206,7 @@ namespace MenuExtension_InterchangeSceneImportAsset
 						{
 							const TAttribute<FText> Label = LOCTEXT("InterchangeSceneImportAsset_Reimport", "Reimport Scene");
 							const TAttribute<FText> ToolTip = LOCTEXT("InterchangeSceneImportAsset_ReimportTooltip", "Reimport the scene associated with each selected InterchangeSceneImportAsset.");
-							const FSlateIcon Icon = FSlateIcon(FAppStyle::GetAppStyleSetName(), "Icons.Import");
+							const FSlateIcon Icon = FSlateIcon(FAppStyle::GetAppStyleSetName(), "Icons.Reimport");
 
 							FToolUIAction UIAction;
 							UIAction.ExecuteAction = FToolMenuExecuteAction::CreateStatic(&ExecuteReimport);
@@ -217,7 +217,7 @@ namespace MenuExtension_InterchangeSceneImportAsset
 						{
 							const TAttribute<FText> Label = LOCTEXT("InterchangeSceneImportAsset_ReimportWithFile", "Reimport Scene With File");
 							const TAttribute<FText> ToolTip = LOCTEXT("InterchangeSceneImportAsset_ReimportWithFile_Tooltip", "Reimport the scene associated with the selected InterchangeSceneImportAsset using a new file.");
-							const FSlateIcon Icon = FSlateIcon(FAppStyle::GetAppStyleSetName(), "Icons.Import");
+							const FSlateIcon Icon = FSlateIcon(FAppStyle::GetAppStyleSetName(), "Icons.Reimport");
 
 							FToolUIAction UIAction;
 							UIAction.ExecuteAction = FToolMenuExecuteAction::CreateStatic(&ExecuteReimportWithFile);
diff --git a/Engine/Source/Editor/CurveTableEditor/Private/CurveTableEditor.cpp b/Engine/Source/Editor/CurveTableEditor/Private/CurveTableEditor.cpp
--- a/Engine/Source/Editor/CurveTableEditor/Private/CurveTableEditor.cpp
+++ b/Engine/Source/Editor/CurveTableEditor/Private/CurveTableEditor.cpp
@@ -517,7 +517,7 @@ void FCurveTableEditor::ExtendToolbar()
 				NAME_None,
 				FText::GetEmpty(),
 				LOCTEXT("Reimport_Tooltip", "Reimport the Curve Table from the source file.  All changes will be lost.  This action cannot be undone."),
-				FSlateIcon(FAppStyle::Get().GetStyleSetName(), "Icons.Toolbar.Import")
+				FSlateIcon(FAppStyle::Get().GetStyleSetName(), "Icons.Toolbar.Reimport")
 			);
 
 			bool HasRichCurves = GetCurveTable()->HasRichCurves();
diff --git a/Engine/Source/Editor/DataTableEditor/Private/DataTableEditor.cpp b/Engine/Source/Editor/DataTableEditor/Private/DataTableEditor.cpp
--- a/Engine/Source/Editor/DataTableEditor/Private/DataTableEditor.cpp
+++ b/Engine/Source/Editor/DataTableEditor/Private/DataTableEditor.cpp
@@ -573,7 +573,7 @@ void FDataTableEditor::FillToolbar(FToolBarBuilder& ToolbarBuilder)
 			NAME_None,
 			LOCTEXT("ReimportText", "Reimport"),
 			LOCTEXT("ReimportTooltip", "Reimport this DataTable"),
-			FSlateIcon(FAppStyle::GetAppStyleSetName(), "Icons.Import"));
+			FSlateIcon(FAppStyle::GetAppStyleSetName(), "Icons.Reimport"));
 
 		ToolbarBuilder.AddSeparator();
 
diff --git a/Engine/Source/Editor/EditorStyle/Private/StarshipStyle.cpp b/Engine/Source/Editor/EditorStyle/Private/StarshipStyle.cpp
--- a/Engine/Source/Editor/EditorStyle/Private/StarshipStyle.cpp
+++ b/Engine/Source/Editor/EditorStyle/Private/StarshipStyle.cpp
@@ -474,6 +474,7 @@ void FStarshipEditorStyle::FStyle::SetupGeneralStyles()
 		Set("Icons.Toolbar.Settings", new CORE_IMAGE_BRUSH_SVG( "Starship/Common/Settings", Icon20x20));
 		Set("Icons.Toolbar.Details", new IMAGE_BRUSH_SVG("Starship/Common/Details", Icon16x16));
 		Set("Icons.Toolbar.Import", new CORE_IMAGE_BRUSH_SVG("Starship/Common/import_20", Icon20x20));
+		Set("Icons.Toolbar.Reimport", new CORE_IMAGE_BRUSH_SVG("Starship/Common/reimport", Icon20x20));
 		Set("Icons.Toolbar.Export", new CORE_IMAGE_BRUSH_SVG("Starship/Common/export_20", Icon20x20));
 
 		Set("Icons.Alert.Solid", new IMAGE_BRUSH_SVG("Starship/Common/AlertTriangleSolid", Icon16x16, FStyleColors::Foreground));
@@ -1230,8 +1231,8 @@ void FStarshipEditorStyle::FStyle::SetupGeneralStyles()
 		Set( "AssetEditor.SaveAsset", new IMAGE_BRUSH_SVG("Starship/Common/SaveCurrent", Icon16x16));
 		Set( "AssetEditor.SaveAssetAs", new IMAGE_BRUSH_SVG("Starship/Common/SaveCurrentAs", Icon16x16));
 
-		Set( "AssetEditor.ReimportAsset", new IMAGE_BRUSH( "Icons/icon_TextureEd_Reimport_40x", Icon40x40 ) );
-		Set( "AssetEditor.ReimportAsset.Small", new IMAGE_BRUSH( "Icons/icon_TextureEd_Reimport_40x", Icon20x20 ) );
+		Set( "AssetEditor.ReimportAsset", new CORE_IMAGE_BRUSH_SVG( "Starship/Common/reimport", Icon40x40 ) );
+		Set( "AssetEditor.ReimportAsset.Small", new CORE_IMAGE_BRUSH_SVG( "Starship/Common/reimport", Icon20x20 ) );
 
 		Set( "AssetEditor.ReadOnlyBorder", new FSlateRoundedBoxBrush(FStyleColors::Foreground, 10.0f));
 		Set("AssetEditor.ReadOnlyOpenable", new IMAGE_BRUSH_SVG("Starship/AssetEditors/LockEye", Icon16x16));
@@ -5503,10 +5504,10 @@ void FStarshipEditorStyle::FStyle::SetupPersonaStyle()
 		Set("Persona.AnimNotifyWindow", new IMAGE_BRUSH_SVG("Starship/Persona/AnimationNotifies", Icon20x20));
 		Set("Persona.RetargetManager", new IMAGE_BRUSH_SVG("Starship/Persona/RetargetManager", Icon20x20));
 		Set("Persona.ImportMesh", new CORE_IMAGE_BRUSH_SVG("Starship/Common/import", Icon20x20));
-		Set("Persona.ReimportMesh", new CORE_IMAGE_BRUSH_SVG("Starship/Common/import", Icon20x20));
+		Set("Persona.ReimportMesh", new CORE_IMAGE_BRUSH_SVG("Starship/Common/reimport", Icon20x20));
 		Set("Persona.ImportLODs", new CORE_IMAGE_BRUSH_SVG("Starship/Common/import", Icon20x20));
 		Set("Persona.ImportAnimation", new CORE_IMAGE_BRUSH_SVG("Starship/Common/import", Icon20x20));
-		Set("Persona.ReimportAnimation", new CORE_IMAGE_BRUSH_SVG("Starship/Common/import", Icon20x20));
+		Set("Persona.ReimportAnimation", new CORE_IMAGE_BRUSH_SVG("Starship/Common/reimport", Icon20x20));
 		Set("Persona.ApplyCompression", new IMAGE_BRUSH_SVG("Starship/Common/Compress", Icon20x20));
 		Set("Persona.ExportToFBX", new CORE_IMAGE_BRUSH_SVG("Starship/Common/export_20", Icon20x20));
 		Set("Persona.CreateAsset", new IMAGE_BRUSH_SVG("Starship/Persona/PersonaCreateAsset", Icon20x20));
@@ -5525,7 +5526,7 @@ void FStarshipEditorStyle::FStyle::SetupPersonaStyle()
 
 		// persona extras
 		Set("Persona.ConvertAnimationGraph", new IMAGE_BRUSH("Old/Graph/ConvertIcon", Icon40x40));
-		Set("Persona.ReimportAsset", new CORE_IMAGE_BRUSH_SVG("Starship/Common/import", Icon20x20));
+		Set("Persona.ReimportAsset", new CORE_IMAGE_BRUSH_SVG("Starship/Common/reimport", Icon20x20));
 		Set("Persona.ConvertToStaticMesh", new IMAGE_BRUSH_SVG("Starship/Common/MakeStaticMesh", Icon20x20));
 		Set("Persona.BakeMaterials", new IMAGE_BRUSH("Icons/icon_tab_Layers_40x", Icon40x40));
 
@@ -5615,8 +5616,8 @@ void FStarshipEditorStyle::FStyle::SetupPersonaStyle()
 	{
 		Set("AnimationEditor.ApplyCompression", new IMAGE_BRUSH_SVG("Starship/Common/Compress", Icon20x20));
 		Set("AnimationEditor.ExportToFBX", new CORE_IMAGE_BRUSH_SVG("Starship/Common/export_20", Icon20x20));
-		Set("AnimationEditor.ReimportAnimation", new CORE_IMAGE_BRUSH_SVG("Starship/Common/import", Icon20x20));
-		Set("AnimationEditor.ReimportAnimationWithDialog", new CORE_IMAGE_BRUSH_SVG("Starship/Common/import", Icon20x20));
+		Set("AnimationEditor.ReimportAnimation", new CORE_IMAGE_BRUSH_SVG("Starship/Common/reimport", Icon20x20));
+		Set("AnimationEditor.ReimportAnimationWithDialog", new CORE_IMAGE_BRUSH_SVG("Starship/Common/reimport", Icon20x20));
 		Set("AnimationEditor.CreateAsset", new IMAGE_BRUSH_SVG("Starship/Persona/PersonaCreateAsset", Icon20x20));
 		Set("AnimationEditor.SetKey", new CORE_IMAGE_BRUSH_SVG("Starship/Common/plus", Icon20x20));
 		Set("AnimationEditor.ApplyAnimation", new IMAGE_BRUSH_SVG("Starship/Common/Apply", Icon20x20));
@@ -5641,8 +5642,8 @@ void FStarshipEditorStyle::FStyle::SetupPersonaStyle()
 
 	// Skeletal mesh editor
 	{
-		Set("SkeletalMeshEditor.ReimportMesh", new CORE_IMAGE_BRUSH_SVG("Starship/Common/import", Icon20x20));
-		Set("SkeletalMeshEditor.ReimportWithDialog", new CORE_IMAGE_BRUSH_SVG("Starship/Common/import", Icon20x20));
+		Set("SkeletalMeshEditor.ReimportMesh", new CORE_IMAGE_BRUSH_SVG("Starship/Common/reimport", Icon20x20));
+		Set("SkeletalMeshEditor.ReimportWithDialog", new CORE_IMAGE_BRUSH_SVG("Starship/Common/reimport", Icon20x20));
 		Set("SkeletalMeshEditor.ImportLODs", new CORE_IMAGE_BRUSH_SVG("Starship/Common/import", Icon20x20));
 
 		Set("SkeletalMeshEditor.MeshSectionSelection", new IMAGE_BRUSH_SVG("Starship/Persona/SectionSelection", Icon20x20));
@@ -6118,7 +6119,7 @@ void FStarshipEditorStyle::FStyle::SetupPersonaStyle()
 	// Asset context menu
 	{
 		Set("Persona.AssetActions.CreateAnimAsset", new IMAGE_BRUSH_SVG("Starship/Persona/PersonaCreateAsset", Icon20x20));
-		Set("Persona.AssetActions.ReimportAnim", new CORE_IMAGE_BRUSH_SVG("Starship/Common/import", Icon20x20));
+		Set("Persona.AssetActions.ReimportAnim", new CORE_IMAGE_BRUSH_SVG("Starship/Common/reimport", Icon20x20));
 		Set("Persona.AssetActions.Retarget", new IMAGE_BRUSH(TEXT("Icons/icon_Persona_Retarget_16x"), Icon16x16));
 		Set("Persona.AssetActions.RetargetSkeleton", new IMAGE_BRUSH(TEXT("Icons/icon_Animation_Retarget_Skeleton_16x"), Icon16x16));
 		Set("Persona.AssetActions.FindSkeleton", new IMAGE_BRUSH(TEXT("Icons/icon_Genericfinder_16x"), Icon16x16));
@@ -6917,7 +6918,7 @@ void FStarshipEditorStyle::FStyle::SetupContentBrowserStyle()
 		Set( "ContentBrowser.AssetActions.OpenSourceLocation", new IMAGE_BRUSH( "Icons/icon_Asset_Open_Source_Location_16x", Icon16x16) );
 		Set( "ContentBrowser.AssetActions.OpenInExternalEditor", new IMAGE_BRUSH( "Icons/icon_Asset_Open_In_External_Editor_16x", Icon16x16) );
 		Set( "ContentBrowser.AssetActions.PublicAssetToggle", new IMAGE_BRUSH("Icons/hiererchy_16x", Icon16x16));
-		Set( "ContentBrowser.AssetActions.ReimportAsset", new IMAGE_BRUSH( "Icons/icon_TextureEd_Reimport_40x", Icon16x16 ) );
+		Set( "ContentBrowser.AssetActions.ReimportAsset", new CORE_IMAGE_BRUSH_SVG( "Starship/Common/reimport", Icon16x16 ) );
 		Set( "ContentBrowser.AssetActions.GoToCodeForAsset", new IMAGE_BRUSH( "GameProjectDialog/feature_code_32x", Icon16x16 ) );
 		Set( "ContentBrowser.AssetActions.FindAssetInWorld", new IMAGE_BRUSH( "/Icons/icon_Genericfinder_16x", Icon16x16 ) );
 		Set( "ContentBrowser.AssetActions.CreateThumbnail", new IMAGE_BRUSH( "Icons/icon_Asset_Create_Thumbnail_16x", Icon16x16) );
@@ -7523,7 +7524,7 @@ void FStarshipEditorStyle::FStyle::SetupToolkitStyles()
 		Set("TextureEditor.Saturation.Small", new IMAGE_BRUSH( "Icons/icon_TextureEd_Saturation_40x", Icon20x20));
 
 		Set("TextureEditor.CompressNow", new IMAGE_BRUSH_SVG( "Starship/Common/Compress", Icon20x20));
-		Set("TextureEditor.Reimport", new CORE_IMAGE_BRUSH_SVG("Starship/Common/import", Icon20x20));
+		Set("TextureEditor.Reimport", new CORE_IMAGE_BRUSH_SVG("Starship/Common/reimport", Icon20x20));
 
 		FButtonStyle MipmapButtonStyle = 
 			FButtonStyle(FStarshipCoreStyle::GetCoreStyle().GetWidgetStyle<FButtonStyle>("Button"))
diff --git a/Engine/Source/Runtime/SlateCore/Private/Styling/StarshipCoreStyle.cpp b/Engine/Source/Runtime/SlateCore/Private/Styling/StarshipCoreStyle.cpp
--- a/Engine/Source/Runtime/SlateCore/Private/Styling/StarshipCoreStyle.cpp
+++ b/Engine/Source/Runtime/SlateCore/Private/Styling/StarshipCoreStyle.cpp
@@ -348,6 +348,7 @@ TSharedRef<ISlateStyle> FStarshipCoreStyle::Create()
 		Style->Set("Icons.Favorites.Small", new IMAGE_BRUSH_SVG("Starship/Common/Favorite", Icon20x20));
 
 		Style->Set("Icons.Import", new IMAGE_BRUSH_SVG("Starship/Common/import", Icon16x16));
+		Style->Set("Icons.Reimport", new IMAGE_BRUSH_SVG("Starship/Common/reimport", Icon16x16));
 		Style->Set("Icons.Filter", new IMAGE_BRUSH_SVG("Starship/Common/filter", Icon16x16));
 		Style->Set("Icons.AutoFilter", new IMAGE_BRUSH_SVG("Starship/Common/FilterAuto", Icon16x16));
 
