Add Constrained Aspect Ratio in Preview to the new toolbar
Set temporary Icon to Constrained Aspect Ratio Button

diff --git a/Engine/Source/Editor/EditorStyle/Private/StarshipStyle.cpp b/Engine/Source/Editor/EditorStyle/Private/StarshipStyle.cpp
--- a/Engine/Source/Editor/EditorStyle/Private/StarshipStyle.cpp
+++ b/Engine/Source/Editor/EditorStyle/Private/StarshipStyle.cpp
@@ -4860,6 +4860,7 @@ void FStarshipEditorStyle::FStyle::SetupLevelEditorStyle()
 		Set("LevelViewport.ToggleImmersive", new IMAGE_BRUSH_SVG("Starship/Common/EnableFullscreen", Icon16x16));
 		Set("LevelViewport.HighResScreenshot", new IMAGE_BRUSH_SVG("Starship/Common/HighResolutionScreenshot", Icon16x16));
 		Set("LevelViewport.ToggleCinematicPreview", new IMAGE_BRUSH_SVG("Starship/Common/Cinematics", Icon16x16));
+		Set("LevelViewport.ToggleAllowConstrainedAspectRatioInPreview", new IMAGE_BRUSH_SVG("Starship/Common/EnableFullscreen", Icon16x16));
 		Set("LevelViewport.ToggleCameraShakePreview", new IMAGE_BRUSH_SVG("Starship/Common/CameraShake", Icon16x16));
 		Set("LevelViewport.AdvancedSettings", new CORE_IMAGE_BRUSH_SVG("Starship/Common/settings", Icon16x16));
 
diff --git a/Engine/Source/Editor/LevelEditor/Private/ViewportToolbar/LevelEditorSubmenus.cpp b/Engine/Source/Editor/LevelEditor/Private/ViewportToolbar/LevelEditorSubmenus.cpp
--- a/Engine/Source/Editor/LevelEditor/Private/ViewportToolbar/LevelEditorSubmenus.cpp
+++ b/Engine/Source/Editor/LevelEditor/Private/ViewportToolbar/LevelEditorSubmenus.cpp
@@ -2103,6 +2103,9 @@ FToolMenuEntry CreatePerformanceAndScalabilitySubmenu()
 					{
 						UnnamedSection.AddEntry(UE::UnrealEd::CreateRemoveRealtimeOverrideEntry(Context->Viewport));
 					}
+
+					FToolMenuEntry& MenuEntry = UnnamedSection.AddMenuEntry(FLevelViewportCommands::Get().ToggleAllowConstrainedAspectRatioInPreview);
+					MenuEntry.SetShowInToolbarTopLevel(true);
 				}
 
 				{
