Add new Icon for Constrained Aspect Ratio Preview
Set the new Icon

diff --git a/Engine/Source/Editor/EditorStyle/Private/StarshipStyle.cpp b/Engine/Source/Editor/EditorStyle/Private/StarshipStyle.cpp
--- a/Engine/Source/Editor/EditorStyle/Private/StarshipStyle.cpp
+++ b/Engine/Source/Editor/EditorStyle/Private/StarshipStyle.cpp
@@ -4860,7 +4860,7 @@ void FStarshipEditorStyle::FStyle::SetupLevelEditorStyle()
 		Set("LevelViewport.ToggleImmersive", new IMAGE_BRUSH_SVG("Starship/Common/EnableFullscreen", Icon16x16));
 		Set("LevelViewport.HighResScreenshot", new IMAGE_BRUSH_SVG("Starship/Common/HighResolutionScreenshot", Icon16x16));
 		Set("LevelViewport.ToggleCinematicPreview", new IMAGE_BRUSH_SVG("Starship/Common/Cinematics", Icon16x16));
-		Set("LevelViewport.ToggleAllowConstrainedAspectRatioInPreview", new IMAGE_BRUSH_SVG("Starship/Common/EnableFullscreen", Icon16x16));
+		Set("LevelViewport.ToggleAllowConstrainedAspectRatioInPreview", new IMAGE_BRUSH_SVG("Starship/Common/ConstrainedAspectRatio_16", Icon16x16));
 		Set("LevelViewport.ToggleCameraShakePreview", new IMAGE_BRUSH_SVG("Starship/Common/CameraShake", Icon16x16));
 		Set("LevelViewport.AdvancedSettings", new CORE_IMAGE_BRUSH_SVG("Starship/Common/settings", Icon16x16));
 
