Update icons used for material editor preview.
Update Static Mesh icon.

diff --git a/Engine/Source/Editor/EditorStyle/Private/StarshipStyle.cpp b/Engine/Source/Editor/EditorStyle/Private/StarshipStyle.cpp
--- a/Engine/Source/Editor/EditorStyle/Private/StarshipStyle.cpp
+++ b/Engine/Source/Editor/EditorStyle/Private/StarshipStyle.cpp
@@ -7390,11 +7390,11 @@ void FStarshipEditorStyle::FStyle::SetupToolkitStyles()
 		Set( "MaterialEditor.LiveUpdate", new IMAGE_BRUSH_SVG("Starship/MaterialEditor/LiveUpdate", Icon20x20));
 		Set( "MaterialEditor.Hierarchy", new IMAGE_BRUSH_SVG("Starship/MaterialEditor/Hierarchy", Icon20x20));
 
-		Set( "MaterialEditor.SetCylinderPreview", new IMAGE_BRUSH("Icons/AssetIcons/Cylinder_16x", Icon16x16));
-		Set( "MaterialEditor.SetSpherePreview", new IMAGE_BRUSH("Icons/AssetIcons/Sphere_16x", Icon16x16));
-		Set( "MaterialEditor.SetPlanePreview", new IMAGE_BRUSH( "Icons/AssetIcons/Plane_16x", Icon16x16) );
-		Set( "MaterialEditor.SetCubePreview", new IMAGE_BRUSH( "Icons/AssetIcons/Cube_16x", Icon16x16) );
-		Set( "MaterialEditor.SetPreviewMeshFromSelection", new IMAGE_BRUSH_SVG( "Starship/AssetIcons/StaticMesh_64", Icon16x16 ) );
+		Set( "MaterialEditor.SetCylinderPreview", new IMAGE_BRUSH_SVG("Icons/AssetIcons/ModelingCylinder_16", Icon16x16));
+		Set( "MaterialEditor.SetSpherePreview", new IMAGE_BRUSH_SVG("Icons/AssetIcons/ModelingSphereShaded_16", Icon16x16));
+		Set( "MaterialEditor.SetPlanePreview", new IMAGE_BRUSH_SVG( "Icons/AssetIcons/ModelingPlane_16", Icon16x16) );
+		Set( "MaterialEditor.SetCubePreview", new IMAGE_BRUSH_SVG( "Icons/AssetIcons/ModelingBox_16", Icon16x16) );
+		Set( "MaterialEditor.SetPreviewMeshFromSelection", new IMAGE_BRUSH_SVG( "Starship/AssetIcons/StaticMesh_16", Icon16x16 ) );
 		Set( "MaterialEditor.TogglePreviewGrid", new IMAGE_BRUSH_SVG( "Starship/Common/Grid", Icon16x16 ) );
 		Set( "MaterialEditor.ToggleMaterialStats", new IMAGE_BRUSH_SVG( "Starship/Common/MaterialAnalyzer", Icon20x20 ) );
 		Set( "MaterialEditor.ToggleMaterialStats.Tab", new IMAGE_BRUSH_SVG("Starship/Common/MaterialAnalyzer", Icon16x16));
