Add icon for resetting show flags to defaults.

diff --git a/Engine/Source/Editor/EditorStyle/Private/StarshipStyle.cpp b/Engine/Source/Editor/EditorStyle/Private/StarshipStyle.cpp
--- a/Engine/Source/Editor/EditorStyle/Private/StarshipStyle.cpp
+++ b/Engine/Source/Editor/EditorStyle/Private/StarshipStyle.cpp
@@ -5232,6 +5232,8 @@ void FStarshipEditorStyle::FStyle::SetupLevelEditorStyle()
 		Set("LevelViewport.ViewportConfig_FourPanes2x2", new IMAGE_BRUSH_SVG("Starship/ViewportLayout/Layouts4Panes", IconLayoutSize));
 		Set("LevelViewport.ViewportConfig_FourPanes2x2.Small", new IMAGE_BRUSH_SVG("Starship/ViewportLayout/Layouts4Panes", IconLayoutSizeSmall));
 
+		Set("LevelViewport.UseDefaultShowFlags", new IMAGE_BRUSH_SVG("Starship/Common/ResetToDefault", Icon16x16));
+		
 		Set( "LevelViewport.EjectActorPilot", new IMAGE_BRUSH_SVG( "Starship/Common/StopPiloting_16", Icon16x16 ) );
 		Set( "LevelViewport.EjectActorPilot.Small", new IMAGE_BRUSH_SVG( "Starship/Common/StopPiloting_16", Icon16x16 ) );
 		Set( "LevelViewport.PilotSelectedActor", new IMAGE_BRUSH_SVG( "Starship/EditorViewport/pilot", Icon16x16 ) );
