Curve Editor: Update lattice tool icon

diff --git a/Engine/Source/Editor/EditorStyle/Private/StarshipStyle.cpp b/Engine/Source/Editor/EditorStyle/Private/StarshipStyle.cpp
--- a/Engine/Source/Editor/EditorStyle/Private/StarshipStyle.cpp
+++ b/Engine/Source/Editor/EditorStyle/Private/StarshipStyle.cpp
@@ -7778,7 +7778,7 @@ void FStarshipEditorStyle::FStyle::SetupToolkitStyles()
 		Set("CurveEditorTools.ActivateTransformTool", new IMAGE_BRUSH_SVG("Starship/GenericCurveEditor/ToolsTransform", Icon20x20));
 		Set("CurveEditorTools.ActivateRetimeTool", new IMAGE_BRUSH_SVG("Starship/GenericCurveEditor/ToolsRetime", Icon20x20));
 		Set("CurveEditorTools.ActivateMultiScaleTool", new IMAGE_BRUSH_SVG("Starship/GenericCurveEditor/ToolsMultiScale", Icon20x20));
-		Set("CurveEditorTools.ActivateLatticeTool", new IMAGE_BRUSH_SVG("Starship/GenericCurveEditor/ToolsMultiScale", Icon20x20));
+		Set("CurveEditorTools.ActivateLatticeTool", new IMAGE_BRUSH_SVG("Starship/GenericCurveEditor/ToolsLattice", Icon20x20));
 	}
 
 	// General Curve Icons
