Fix suble issues with toolbar styles.
- Fix spacing on options shevrons & viewport sizing icon.
- Fix spacing on raised transform space menu items & gap.
- Fix button colors not being in-spec.

diff --git a/Engine/Source/Developer/ToolMenus/Private/ToolMenus.cpp b/Engine/Source/Developer/ToolMenus/Private/ToolMenus.cpp
--- a/Engine/Source/Developer/ToolMenus/Private/ToolMenus.cpp
+++ b/Engine/Source/Developer/ToolMenus/Private/ToolMenus.cpp
@@ -1820,8 +1820,11 @@ void UToolMenus::PopulateToolBarBuilderWithTopLevelChildren(FToolBarBuilder& Too
 		FMenuEntryResizeParams ResizeParams;
 		// Never show spacers in overflow menus.
 		ResizeParams.VisibleInOverflow = false;
+		// Default clipping priority is 0. Set to a negative number so that the spaces drop before any content.
+		ResizeParams.ClippingPriority = -100;
 
-		const FToolBarStyle ToolbarStyle = ParentMenu->GetStyleSet()->GetWidgetStyle<FToolBarStyle>(ParentMenu->StyleName);
+		const FName StyleName = InBlock.StyleNameOverride.IsNone() ? ParentMenu->StyleName : InBlock.StyleNameOverride;
+		const FToolBarStyle ToolbarStyle = ParentMenu->GetStyleSet()->GetWidgetStyle<FToolBarStyle>(StyleName);
 		ToolBarBuilder.AddWidget(
 			SNew(SSpacer).Size(FVector2D(ToolbarStyle.RaisedChildrenRightPadding, 0)),
 			StyleParams,
diff --git a/Engine/Source/Editor/EditorStyle/Private/StarshipStyle.cpp b/Engine/Source/Editor/EditorStyle/Private/StarshipStyle.cpp
--- a/Engine/Source/Editor/EditorStyle/Private/StarshipStyle.cpp
+++ b/Engine/Source/Editor/EditorStyle/Private/StarshipStyle.cpp
@@ -2545,26 +2545,21 @@ void FStarshipEditorStyle::FStyle::SetupViewportStyles()
 {
 	// New viewport toolbar.
 	{
-		const FSlateRoundedBoxBrush TransparentBrush(FLinearColor(0.f, 0.f, 0.f, 0.f), 0.f, FLinearColor(0.f, 0.f, 0.f, 0.f), 0.f);
+		const FSlateRoundedBoxBrush TransparentBrush(FStyleColors::Transparent, 0.f, FStyleColors::Transparent, 0.f);
 	
-		FLinearColor ToolbarBackgroundColor = FStyleColors::Dropdown.GetSpecifiedColor();
-		ToolbarBackgroundColor.A = .90f;
+		FSlateColor ToolbarBackgroundColor = FStyleColors::Dropdown;
 		const FSlateRoundedBoxBrush BackgroundBrush(ToolbarBackgroundColor, 5.f, ToolbarBackgroundColor, 1.0f);
 		const FSlateRoundedBoxBrush BackgroundBrushStart(ToolbarBackgroundColor, FVector4(5.f, 0.f, 0.f, 5.f), ToolbarBackgroundColor, 1.0f);
 		const FSlateRoundedBoxBrush BackgroundBrushMiddle(ToolbarBackgroundColor, FVector4(0.f, 0.f, 0.f, 0.f), ToolbarBackgroundColor, 1.0f);
 		const FSlateRoundedBoxBrush BackgroundBrushEnd(ToolbarBackgroundColor, FVector4(0.f, 5.f, 5.f, 0.f), ToolbarBackgroundColor, 1.0f);
 
-		FLinearColor ToolbarHoveredColor = FStyleColors::Dropdown.GetSpecifiedColor();
-		ToolbarHoveredColor.R = FMath::Min(1.0f, ToolbarHoveredColor.R * 1.5f);
-		ToolbarHoveredColor.G = FMath::Min(1.0f, ToolbarHoveredColor.G * 1.5f);
-		ToolbarHoveredColor.B = FMath::Min(1.0f, ToolbarHoveredColor.B * 1.5f);
+		FSlateColor ToolbarHoveredColor = FStyleColors::Hover;
 		const FSlateRoundedBoxBrush HoveredButtonBrush(ToolbarHoveredColor, 5.f, ToolbarHoveredColor, 1.0f);
 		const FSlateRoundedBoxBrush HoveredButtonBrushStart(ToolbarHoveredColor, FVector4(5.f, 0.f, 0.f, 5.f), ToolbarHoveredColor, 1.0f);
 		const FSlateRoundedBoxBrush HoveredButtonBrushMiddle(ToolbarHoveredColor, FVector4(0.f, 0.f, 0.f, 0.f), ToolbarHoveredColor, 1.0f);
 		const FSlateRoundedBoxBrush HoveredButtonBrushEnd(ToolbarHoveredColor, FVector4(0.f, 5.f, 5.f, 0.f), ToolbarHoveredColor, 1.0f);
 
-		FLinearColor ToolbarPressedColor = FStyleColors::Recessed.GetSpecifiedColor();
-		ToolbarPressedColor.A = .50f;
+		FSlateColor ToolbarPressedColor = FStyleColors::Header;
 		const FSlateRoundedBoxBrush PressedButtonBrush(ToolbarPressedColor, 5.f, ToolbarPressedColor, 1.0f);
 		const FSlateRoundedBoxBrush PressedButtonBrushStart(ToolbarPressedColor, FVector4(5.f, 0.f, 0.f, 5.f), ToolbarPressedColor, 1.0f);
 		const FSlateRoundedBoxBrush PressedButtonBrushMiddle(ToolbarPressedColor, FVector4(0.f, 0.f, 0.f, 0.f), ToolbarPressedColor, 1.0f);
@@ -2578,7 +2573,7 @@ void FStarshipEditorStyle::FStyle::SetupViewportStyles()
 				.SetNormalForeground(FStyleColors::ForegroundHeader)
 				.SetHoveredForeground(FStyleColors::ForegroundHover)
 				.SetPressedForeground(FStyleColors::ForegroundHover)
-				.SetDisabledForeground(FStyleColors::Foreground)
+				.SetDisabledForeground(FStyleColors::Hover2)
 				.SetNormalPadding(FMargin(6.0f, 4.0f))
 				.SetPressedPadding(FMargin(6.0f, 4.0f));
 		}
@@ -2743,8 +2738,9 @@ void FStarshipEditorStyle::FStyle::SetupViewportStyles()
 				.SetToggleButtonStyle(ToggleButtonStyleRaised)
 				.SetComboButtonStyle(ComboStyleRaised)
 				.SetBlockHovered(BackgroundBrush)
-				.SetButtonPadding(FMargin(3.0f, 0))
-				.SetComboButtonPadding(FMargin(3.0f, 0))
+				// The offset gap allows raised buttons to be 6px away from themselves & their top-level parents
+				.SetButtonPadding(FMargin(2.f, 0.f, 4.f, 0.f))
+				.SetComboButtonPadding(FMargin(2.f, 0.f, 4.f, 0.f))
 				.SetSettingsComboButtonStyle(SettingsComboStyleRaised);
 			
 			Set("ViewportToolbar.Raised", ViewportToolbarRaisedStyle);
@@ -2766,8 +2762,17 @@ void FStarshipEditorStyle::FStyle::SetupViewportStyles()
 
 		// Special style for raised transform buttons
 		{
-			FToolBarStyle RaisedTransformTools = FToolBarStyle(GetWidgetStyle<FToolBarStyle>("ViewportToolbar.Raised"));
-			RaisedTransformTools.SetButtonPadding(0);
+			FToolBarStyle TransformToolsParent = FToolBarStyle(GetWidgetStyle<FToolBarStyle>("ViewportToolbar"))
+				// The raised children all have no padding, so the raised gap padding needs to be extended
+				// to make the sizing consistent.
+				// The design target is 24px from the last raised item to the next top-level item.
+				.SetRaisedChildrenRightPadding(20.f);
+		
+			FToolBarStyle RaisedTransformTools = FToolBarStyle(GetWidgetStyle<FToolBarStyle>("ViewportToolbar.Raised"))
+				.SetButtonPadding(0.f)
+				.SetComboButtonPadding(0.f);
+				
+			Set("ViewportToolbar.TransformTools", TransformToolsParent);
 			Set("ViewportToolbar.TransformTools.Raised", RaisedTransformTools);
 		}
 
@@ -2825,48 +2830,35 @@ void FStarshipEditorStyle::FStyle::SetupViewportStyles()
 		// Special style for the "Viewport Sizing" submenu, with rounded corners on the left.
 		{
 			const FButtonStyle ViewportSizingButtonStyle =
-				FButtonStyle(ButtonStyle)
-					.SetNormal(FSlateRoundedBoxBrush(ToolbarBackgroundColor, FVector4(4.0f, 0.0f, 0.0f, 4.0f)))
-					.SetHovered(FSlateRoundedBoxBrush(ToolbarHoveredColor, FVector4(4.0f, 0.0f, 0.0f, 4.0f)))
-					.SetPressed(FSlateRoundedBoxBrush(ToolbarPressedColor, FVector4(4.0f, 0.0f, 0.0f, 4.0f)))
-					.SetDisabled(FSlateRoundedBoxBrush(FStyleColors::Dropdown, FVector4(4.0f, 0.0f, 0.0f, 4.0f)))
-					.SetNormalPadding(FMargin(0, 0, 2.f, 0))
-					.SetPressedPadding(FMargin(0, 0, 2.f, 0));
+				FButtonStyle(ButtonStyleStart)
+					.SetNormalPadding(FMargin(0.f))
+					.SetPressedPadding(FMargin(0.f));
 
 			const FComboButtonStyle ViewportSizingComboButtonStyle =
 				FComboButtonStyle(ComboButtonStyle)
-					.SetDownArrowPadding(FMargin(5.f, 0.f, 2.f, 0.f))
+					.SetDownArrowPadding(FMargin(5.f, 0.f, 5.f, 0.f))
 					.SetDownArrowImage(CORE_IMAGE_BRUSH_SVG("Starship/Common/ellipsis-vertical-narrow", FVector2f(6.f, 20.f)))
 					.SetButtonStyle(ViewportSizingButtonStyle);
 
 			const FToolBarStyle ViewportSizingStyle = FToolBarStyle(ViewportToolbarStyle)
-														  .SetButtonStyle(ViewportSizingButtonStyle)
-														  .SetComboButtonStyle(ViewportSizingComboButtonStyle)
-														  .SetComboButtonPadding(FMargin(4.0f, 0, 1.0f, 0))
-														  .SetSeparatorPadding(FMargin(0.f, 0.f, 0.f, 0.f));
+				.SetButtonStyle(ViewportSizingButtonStyle)
+				.SetComboButtonStyle(ViewportSizingComboButtonStyle)
+				.SetComboButtonPadding(FMargin(4.0f, 0, 1.0f, 0))
+				.SetSeparatorPadding(FMargin(0.f));
 
 			Set("ViewportToolbarViewportSizingSubmenu", ViewportSizingStyle);
 
 			// A raised button keeps the gray background here, with rounded corners on the right.
 			{
-				FButtonStyle ButtonStyleRaised = ButtonStyle;
-				{
-					FLinearColor BackgroundColorRaised = FStyleColors::Dropdown.GetSpecifiedColor();
-					BackgroundColorRaised.A = .90f;
-					const FSlateRoundedBoxBrush BackgroundBrushRaised(BackgroundColorRaised, 5.f, BackgroundColorRaised, 1.0f);
-
-					ButtonStyleRaised.SetNormal(BackgroundBrushRaised)
-						.SetNormalForeground(FStyleColors::Foreground)
-						.SetNormal(FSlateRoundedBoxBrush(ToolbarBackgroundColor, FVector4(0.f, 4.0f, 4.0f, 0.f)))
-						.SetHovered(FSlateRoundedBoxBrush(ToolbarHoveredColor, FVector4(0.f, 4.0f, 4.0f, 0.f)))
-						.SetPressed(FSlateRoundedBoxBrush(ToolbarPressedColor, FVector4(0.f, 4.0f, 4.0f, 0.f)))
-						.SetDisabled(FSlateRoundedBoxBrush(FStyleColors::Dropdown, FVector4(0.f, 4.0f, 4.0f, 0.f)));
-				}
-
+				FButtonStyle ButtonStyleRaised = FButtonStyle(ButtonStyleEnd)
+					// The icon cannot be perfectly centered, so it needs a 1px shove to the right
+					// to match the design spec.
+					.SetNormalPadding(FMargin(7.0f, 4.0f, 6.0f, 4.0f))
+					.SetPressedPadding(FMargin(7.0f, 4.0f, 6.0f, 4.0f));
 				const FToolBarStyle ViewportSizingRaisedStyle = FToolBarStyle(ViewportToolbarStyle)
-																	.SetButtonStyle(ButtonStyleRaised)
-																	.SetButtonPadding(0.f)
-																	.SetComboButtonPadding(FMargin(2.0f, 0.f));
+					.SetButtonStyle(ButtonStyleRaised)
+					.SetButtonPadding(FMargin(0.f, 0.f, 4.f, 0.f))
+					.SetComboButtonPadding(FMargin(2.0f, 0.f));
 
 				Set("ViewportToolbarViewportSizingSubmenu.Raised", ViewportSizingRaisedStyle);
 			}
diff --git a/Engine/Source/Editor/UnrealEd/Private/ViewportToolbar/UnrealEdViewportToolbar.cpp b/Engine/Source/Editor/UnrealEd/Private/ViewportToolbar/UnrealEdViewportToolbar.cpp
--- a/Engine/Source/Editor/UnrealEd/Private/ViewportToolbar/UnrealEdViewportToolbar.cpp
+++ b/Engine/Source/Editor/UnrealEd/Private/ViewportToolbar/UnrealEdViewportToolbar.cpp
@@ -733,41 +733,44 @@ FToolMenuEntry CreateTransformsSubmenu()
 			[](UToolMenu* Submenu) -> void
 			{
 				{
+					FToolMenuEntryToolBarData ToolBarData;
+					ToolBarData.StyleNameOverride = "ViewportToolbar.TransformTools";
+				
 					FToolMenuSection& TransformToolsSection =
 						Submenu->FindOrAddSection("TransformTools", LOCTEXT("TransformToolsLabel", "Transform Tools"));
 
 					FToolMenuEntry SelectMode = FToolMenuEntry::InitMenuEntry(FEditorViewportCommands::Get().SelectMode);
 					SelectMode.SetShowInToolbarTopLevel(true);
-					SelectMode.ToolBarData.StyleNameOverride = "ViewportToolbar.TransformTools";
+					SelectMode.ToolBarData = ToolBarData;
 					
 					TransformToolsSection.AddEntry(SelectMode);
 
 					FToolMenuEntry TranslateMode =
 						FToolMenuEntry::InitMenuEntry(FEditorViewportCommands::Get().TranslateMode);
 					TranslateMode.SetShowInToolbarTopLevel(true);
-					TranslateMode.ToolBarData.StyleNameOverride = "ViewportToolbar.TransformTools";
+					TranslateMode.ToolBarData = ToolBarData;
 					TransformToolsSection.AddEntry(TranslateMode);
 
 					FToolMenuEntry RotateMode = FToolMenuEntry::InitMenuEntry(FEditorViewportCommands::Get().RotateMode);
 					RotateMode.SetShowInToolbarTopLevel(true);
-					RotateMode.ToolBarData.StyleNameOverride = "ViewportToolbar.TransformTools";
+					RotateMode.ToolBarData = ToolBarData;
 					TransformToolsSection.AddEntry(RotateMode);
 
 					FToolMenuEntry ScaleMode = FToolMenuEntry::InitMenuEntry(FEditorViewportCommands::Get().ScaleMode);
 					ScaleMode.SetShowInToolbarTopLevel(true);
-					ScaleMode.ToolBarData.StyleNameOverride = "ViewportToolbar.TransformTools";
+					ScaleMode.ToolBarData = ToolBarData;
 					TransformToolsSection.AddEntry(ScaleMode);
 
 					FToolMenuEntry TranslateRotateMode =
 						FToolMenuEntry::InitMenuEntry(FEditorViewportCommands::Get().TranslateRotateMode);
 					TranslateRotateMode.SetShowInToolbarTopLevel(true);
-					TranslateRotateMode.ToolBarData.StyleNameOverride = "ViewportToolbar.TransformTools";
+					TranslateRotateMode.ToolBarData = ToolBarData;
 					TransformToolsSection.AddEntry(TranslateRotateMode);
 
 					FToolMenuEntry TranslateRotate2DMode =
 						FToolMenuEntry::InitMenuEntry(FEditorViewportCommands::Get().TranslateRotate2DMode);
 					TranslateRotate2DMode.SetShowInToolbarTopLevel(true);
-					TranslateRotate2DMode.ToolBarData.StyleNameOverride = "ViewportToolbar.TransformTools";
+					TranslateRotate2DMode.ToolBarData = ToolBarData;
 					TransformToolsSection.AddEntry(TranslateRotate2DMode);
 
 					// Build a submenu for selecting the coordinate system to use.
@@ -922,6 +925,7 @@ FToolMenuEntry CreateTransformsSubmenu()
 												});
 										}
 
+										CoordinateSystemSubmenu.ToolBarData = ToolBarData;
 										CoordinateSystemSubmenu.ToolBarData.LabelOverride = FText();
 										CoordinateSystemSubmenu.ToolBarData.ActionOverride = CycleCoordSystemAction;
 										CoordinateSystemSubmenu.SetShowInToolbarTopLevel(true);
@@ -1036,6 +1040,7 @@ FToolMenuEntry CreateTransformsSubmenu()
 	Entry.Icon = FSlateIcon(FAppStyle::GetAppStyleSetName(), "LevelEditor.SelectMode");
 	Entry.ToolBarData.LabelOverride = FText();
 	Entry.ToolBarData.ResizeParams.ClippingPriority = 1000;
+	Entry.StyleNameOverride = "ViewportToolbar.TransformTools";
 
 	return Entry;
 }
diff --git a/Engine/Source/Runtime/SlateCore/Private/Styling/StarshipCoreStyle.cpp b/Engine/Source/Runtime/SlateCore/Private/Styling/StarshipCoreStyle.cpp
--- a/Engine/Source/Runtime/SlateCore/Private/Styling/StarshipCoreStyle.cpp
+++ b/Engine/Source/Runtime/SlateCore/Private/Styling/StarshipCoreStyle.cpp
@@ -2339,7 +2339,8 @@ void FStarshipCoreStyle::SetupMultiboxStyles(TSharedRef<FStyle>& Style)
 		FComboButtonStyle SlimToolBarSettingsComboButton = FComboButtonStyle(SlimToolBarComboButton)
 			.SetContentPadding(0)
 			.SetButtonStyle(SlimToolBarButtonSettingsRegion)
-			.SetDownArrowImage(IMAGE_BRUSH_SVG("Starship/Common/ellipsis-vertical-narrow", FVector2f(6.f, 20.f)));
+			.SetDownArrowImage(IMAGE_BRUSH_SVG("Starship/Common/ellipsis-vertical-narrow", FVector2f(6.f, 20.f)))
+			.SetDownArrowPadding(FMargin(1.f, 0.f));
 
 		SlimToolbarStyle.SetSettingsComboButtonStyle(SlimToolBarSettingsComboButton);
 
