[Sequencer] Navigation Tool use foreground color instead of white

diff --git a/Engine/Source/Editor/EditorStyle/Private/StarshipStyle.cpp b/Engine/Source/Editor/EditorStyle/Private/StarshipStyle.cpp
--- a/Engine/Source/Editor/EditorStyle/Private/StarshipStyle.cpp
+++ b/Engine/Source/Editor/EditorStyle/Private/StarshipStyle.cpp
@@ -2519,7 +2519,8 @@ void FStarshipEditorStyle::FStyle::SetupSequencerStyles()
 		const FTableRowStyle& AlternatingTableRowStyle = FAppStyle::Get().GetWidgetStyle<FTableRowStyle>(TEXT("TableView.AlternatingRow"));
 		Set("Sequencer.NavigationTool.TableViewRow", FTableRowStyle(AlternatingTableRowStyle));
 
-		const FSpinBoxStyle FrameTimeSpinStyle = GetWidgetStyle<FSpinBoxStyle>(TEXT("SpinBox"));
+		FSpinBoxStyle FrameTimeSpinStyle = GetWidgetStyle<FSpinBoxStyle>(TEXT("SpinBox"));
+		FrameTimeSpinStyle.ForegroundColor = FStyleColors::Foreground;
 		Set("Sequencer.NavigationTool.SpinBox", FrameTimeSpinStyle);
 
 		FSpinBoxStyle StartTimeSpinStyle = FrameTimeSpinStyle;
diff --git a/Engine/Source/Editor/Sequencer/Private/Tools/NavigationTool/Widgets/Columns/SNavigationToolEndFrameOffset.cpp b/Engine/Source/Editor/Sequencer/Private/Tools/NavigationTool/Widgets/Columns/SNavigationToolEndFrameOffset.cpp
--- a/Engine/Source/Editor/Sequencer/Private/Tools/NavigationTool/Widgets/Columns/SNavigationToolEndFrameOffset.cpp
+++ b/Engine/Source/Editor/Sequencer/Private/Tools/NavigationTool/Widgets/Columns/SNavigationToolEndFrameOffset.cpp
@@ -47,6 +47,7 @@ void SNavigationToolEndFrameOffset::Construct(const FArguments& InArgs
 		[
 			SNew(SSpinBox<int32>)
 			.Justification(ETextJustify::Center)
+			.Style(&FAppStyle::Get(), TEXT("Sequencer.NavigationTool.SpinBox"))
 			.Value(this, &SNavigationToolEndFrameOffset::GetValue)
 			.OnValueChanged(this, &SNavigationToolEndFrameOffset::OnValueChanged)
 			.OnValueCommitted(this, &SNavigationToolEndFrameOffset::OnValueCommitted)
diff --git a/Engine/Source/Editor/Sequencer/Private/Tools/NavigationTool/Widgets/Columns/SNavigationToolStartFrameOffset.cpp b/Engine/Source/Editor/Sequencer/Private/Tools/NavigationTool/Widgets/Columns/SNavigationToolStartFrameOffset.cpp
--- a/Engine/Source/Editor/Sequencer/Private/Tools/NavigationTool/Widgets/Columns/SNavigationToolStartFrameOffset.cpp
+++ b/Engine/Source/Editor/Sequencer/Private/Tools/NavigationTool/Widgets/Columns/SNavigationToolStartFrameOffset.cpp
@@ -38,6 +38,7 @@ void SNavigationToolStartFrameOffset::Construct(const FArguments& InArgs
 		[
 			SNew(SSpinBox<int32>)
 			.Justification(ETextJustify::Center)
+			.Style(&FAppStyle::Get(), TEXT("Sequencer.NavigationTool.SpinBox"))
 			.Value(this, &SNavigationToolStartFrameOffset::GetValue)
 			.OnValueChanged(this, &SNavigationToolStartFrameOffset::OnValueChanged)
 			.OnValueCommitted(this, &SNavigationToolStartFrameOffset::OnValueCommitted)
