[ContentBrowser] Thumbnail Edit Mode tweaks
Change EditMode icons to new ones

diff --git a/Engine/Source/Editor/EditorStyle/Private/StarshipStyle.cpp b/Engine/Source/Editor/EditorStyle/Private/StarshipStyle.cpp
--- a/Engine/Source/Editor/EditorStyle/Private/StarshipStyle.cpp
+++ b/Engine/Source/Editor/EditorStyle/Private/StarshipStyle.cpp
@@ -6895,10 +6895,10 @@ void FStarshipEditorStyle::FStyle::SetupContentBrowserStyle()
 
 			
 		Set( "ContentBrowser.PrimitiveCustom", new IMAGE_BRUSH( "ContentBrowser/ThumbnailCustom", Icon32x32 ) );
-		Set( "ContentBrowser.PrimitiveSphere", new IMAGE_BRUSH( "ContentBrowser/ThumbnailSphere", Icon32x32 ) );
-		Set( "ContentBrowser.PrimitiveCube", new IMAGE_BRUSH( "ContentBrowser/ThumbnailCube", Icon32x32 ) );
-		Set( "ContentBrowser.PrimitivePlane", new IMAGE_BRUSH( "ContentBrowser/ThumbnailPlane", Icon32x32 ) );
-		Set( "ContentBrowser.PrimitiveCylinder", new IMAGE_BRUSH( "ContentBrowser/ThumbnailCylinder", Icon32x32 ) );
+		Set( "ContentBrowser.PrimitiveSphere", new IMAGE_BRUSH_SVG("Icons/AssetIcons/ModelingSphereShaded_16", Icon16x16));
+		Set( "ContentBrowser.PrimitiveCube", new IMAGE_BRUSH_SVG( "Icons/AssetIcons/ModelingBox_16", Icon16x16) );
+		Set( "ContentBrowser.PrimitivePlane", new IMAGE_BRUSH_SVG( "Icons/AssetIcons/ModelingPlane_16", Icon16x16) );
+		Set( "ContentBrowser.PrimitiveCylinder", new IMAGE_BRUSH_SVG("Icons/AssetIcons/ModelingCylinder_16", Icon16x16));
 		Set( "ContentBrowser.ResetPrimitiveToDefault", new CORE_IMAGE_BRUSH_SVG("Starship/Common/Undo", Icon20x20) );
 
 		Set( "ContentBrowser.TopBar.Font", FTextBlockStyle( NormalText )
