Fixes for Zen status dialog when used under the starship style.

diff --git a/Engine/Source/Editor/EditorStyle/Private/StarshipStyle.cpp b/Engine/Source/Editor/EditorStyle/Private/StarshipStyle.cpp
--- a/Engine/Source/Editor/EditorStyle/Private/StarshipStyle.cpp
+++ b/Engine/Source/Editor/EditorStyle/Private/StarshipStyle.cpp
@@ -1098,7 +1098,8 @@ void FStarshipEditorStyle::FStyle::SetupGeneralStyles()
 	{
 		Set("Zen.Icons.LaunchDashboard", new CORE_IMAGE_BRUSH_SVG("Starship/Common/SessionFrontend", Icon16x16));
 		Set("Zen.Icons.Server", new CORE_IMAGE_BRUSH_SVG("Starship/Common/SessionFrontend", Icon16x16));
-		Set("Zen.Icons.FolderExplore", new IMAGE_BRUSH_SVG("Starship/ContentBrowser/show-in-explorer", Icon16x16));
+		Set("Zen.Icons.WebBrowser", new IMAGE_BRUSH_SVG("Starship/Common/WebBrowser", Icon12x12));
+		Set("Zen.Icons.FolderExplore", new IMAGE_BRUSH_SVG("Starship/ContentBrowser/show-in-explorer", Icon12x12));
 		Set("Zen.Icons.Clipboard", new IMAGE_BRUSH_SVG("Starship/Common/Clipboard", Icon16x16));
 	}
 
