Sequencer: Move Force Whole Frames to a button on the toolbar

diff --git a/Engine/Source/Editor/EditorStyle/Private/StarshipStyle.cpp b/Engine/Source/Editor/EditorStyle/Private/StarshipStyle.cpp
--- a/Engine/Source/Editor/EditorStyle/Private/StarshipStyle.cpp
+++ b/Engine/Source/Editor/EditorStyle/Private/StarshipStyle.cpp
@@ -2168,12 +2168,12 @@ void FStarshipEditorStyle::FStyle::SetupSequencerStyles()
 	{
 		FToolBarStyle SequencerToolbar = FStarshipCoreStyle::GetCoreStyle().GetWidgetStyle<FToolBarStyle>("AssetEditorToolbar");
 
-		SequencerToolbar.SetButtonPadding(       FMargin(0.0f, 0.0f));
+		SequencerToolbar.SetButtonPadding(       FMargin(2.0f, 0.0f));
 		SequencerToolbar.SetCheckBoxPadding(     FMargin(0.0f, 0.0f));
 		SequencerToolbar.SetComboButtonPadding(  FMargin(0.0f, 0.0f));
 		SequencerToolbar.SetIndentedBlockPadding(FMargin(0.0f, 0.0f));
 		SequencerToolbar.SetBlockPadding(        FMargin(0.0f, 0.0f));
-		SequencerToolbar.SetSeparatorPadding(    FMargin(6.0f, 0.0f));
+		SequencerToolbar.SetSeparatorPadding(    FMargin(2.0f, 0.0f));
 
 		Set("SequencerToolBar", SequencerToolbar);
 
@@ -2354,6 +2354,7 @@ void FStarshipEditorStyle::FStyle::SetupSequencerStyles()
 		Set("Sequencer.SetKeyGroup", new IMAGE_BRUSH_SVG("Starship/Sequencer/KeyGroup", Icon20x20));
 		Set("Sequencer.SetKeyChanged", new IMAGE_BRUSH_SVG("Starship/Sequencer/KeyChanged", Icon20x20));
 		Set("Sequencer.ToggleIsSnapEnabled", new IMAGE_BRUSH_SVG("Starship/Sequencer/Snap", Icon20x20));
+		Set("Sequencer.ToggleForceWholeFrames", new IMAGE_BRUSH_SVG("Starship/Sequencer/ForceWholeFrames", Icon20x20));
 		Set("Sequencer.ToggleLimitViewportSelection",
 			new IMAGE_BRUSH_SVG("Starship/Sequencer/SelectOnlyInSequence_16", Icon20x20));
 
diff --git a/Engine/Source/Editor/Sequencer/Private/SSequencer.cpp b/Engine/Source/Editor/Sequencer/Private/SSequencer.cpp
--- a/Engine/Source/Editor/Sequencer/Private/SSequencer.cpp
+++ b/Engine/Source/Editor/Sequencer/Private/SSequencer.cpp
@@ -1643,6 +1643,11 @@ void SSequencer::PopulateToolBar(UToolMenu* InMenu)
 		SnapOptionsEntry.StyleNameOverride = SequencerToolbarStyleName;
 		SnappingSection.AddEntry(SnapOptionsEntry);
 
+		FToolMenuEntry ToggleWholeFramesEntry =
+			FToolMenuEntry::InitToolBarButton(FSequencerCommands::Get().ToggleForceWholeFrames);
+		ToggleWholeFramesEntry.StyleNameOverride = SequencerToolbarStyleName;
+		SnappingSection.AddEntry(ToggleWholeFramesEntry);
+
 		SnappingSection.AddSeparator("PlayRate");
 
 		FToolMenuEntry PlayRateEntry = FToolMenuEntry::InitWidget(
@@ -2510,8 +2515,7 @@ TSharedRef<SWidget> SSequencer::MakeRenderMovieMenu()
 }
 
 TSharedRef<SWidget> SSequencer::MakeSnapMenu()
-{
-	
+{	
 	static const FName MenuName("Sequencer.SnapOptions");
 	if (!UToolMenus::Get()->IsMenuRegistered(MenuName))
 	{
@@ -2537,11 +2541,6 @@ TSharedRef<SWidget> SSequencer::MakeSnapMenu()
 			FToolMenuSection& Section = Toolbar->AddSection( "CurveSnapping", LOCTEXT( "SnappingMenuCurveHeader", "Curve Snapping" ) );
 			Section.AddMenuEntry( FSequencerCommands::Get().ToggleSnapCurveValueToInterval );
 		}
-		
-		{
-			FToolMenuSection& Section = Toolbar->AddSection( "IntervalSnapping", LOCTEXT("SnappingMenuIntervalHeader", "Force Whole Frames") );
-			Section.AddMenuEntry(FSequencerCommands::Get().ToggleForceWholeFrames);
-		}
 	}
 
 	USequencerToolMenuContext* ContextObject = NewObject<USequencerToolMenuContext>();
diff --git a/Engine/Source/Editor/Sequencer/Private/SequencerCommands.cpp b/Engine/Source/Editor/Sequencer/Private/SequencerCommands.cpp
--- a/Engine/Source/Editor/Sequencer/Private/SequencerCommands.cpp
+++ b/Engine/Source/Editor/Sequencer/Private/SequencerCommands.cpp
@@ -145,7 +145,7 @@ void FSequencerCommands::RegisterCommands()
 
 	UI_COMMAND( ToggleSnapCurveValueToInterval, "Snap Curve Key Values", "Snap curve keys to the value snapping interval", EUserInterfaceActionType::ToggleButton, FInputChord() );
 
-	UI_COMMAND( ToggleForceWholeFrames, "Force Whole Frames", "Snap keys/sections/playhead to the interval which forces whole frames", EUserInterfaceActionType::ToggleButton, FInputChord() );
+	UI_COMMAND( ToggleForceWholeFrames, "Force Whole Frames", "Enables or disables keys/sections/playhead on whole frames", EUserInterfaceActionType::ToggleButton, FInputChord() );
 
 	UI_COMMAND( FindInContentBrowser, "Find in Content Browser", "Find the viewed sequence asset in the content browser", EUserInterfaceActionType::Button, FInputChord() );
 	UI_COMMAND( BrowseToObject, "Browse to Object", "Browse to the selected section's object", EUserInterfaceActionType::Button, FInputChord(EModifierKey::Control, EKeys::B) );
