Update pilot actor icon. Add select piloted actor icon.

#rnx

diff --git a/Engine/Source/Editor/EditorStyle/Private/StarshipStyle.cpp b/Engine/Source/Editor/EditorStyle/Private/StarshipStyle.cpp
--- a/Engine/Source/Editor/EditorStyle/Private/StarshipStyle.cpp
+++ b/Engine/Source/Editor/EditorStyle/Private/StarshipStyle.cpp
@@ -5223,6 +5223,7 @@ void FStarshipEditorStyle::FStyle::SetupLevelEditorStyle()
 		Set( "LevelViewport.EjectActorPilot", new IMAGE_BRUSH_SVG( "Starship/Common/StopPiloting_16", Icon16x16 ) );
 		Set( "LevelViewport.EjectActorPilot.Small", new IMAGE_BRUSH_SVG( "Starship/Common/StopPiloting_16", Icon16x16 ) );
 		Set( "LevelViewport.PilotSelectedActor", new IMAGE_BRUSH_SVG( "Starship/EditorViewport/pilot", Icon16x16 ) );
+		Set( "LevelViewport.SelectPilotedActor", new IMAGE_BRUSH_SVG( "Starship/EditorViewport/pilot-select", Icon16x16 ));
 
 		Set( "LevelViewport.ToggleActorPilotCameraView",       new IMAGE_BRUSH_SVG(  "Starship/Common/ExactCameraView_16", Icon16x16 ) );
 		Set( "LevelViewport.ToggleActorPilotCameraView.Small", new IMAGE_BRUSH_SVG( "Starship/Common/ExactCameraView_16", Icon16x16 ) );
