[Viewport Toolbar] Add missing Transform menu icons

diff --git a/Engine/Source/Editor/EditorStyle/Private/StarshipStyle.cpp b/Engine/Source/Editor/EditorStyle/Private/StarshipStyle.cpp
--- a/Engine/Source/Editor/EditorStyle/Private/StarshipStyle.cpp
+++ b/Engine/Source/Editor/EditorStyle/Private/StarshipStyle.cpp
@@ -4725,11 +4725,15 @@ void FStarshipEditorStyle::FStyle::SetupLevelEditorStyle()
 		Set("LevelEditor.BrowseViewportControls",  new IMAGE_BRUSH_SVG("Starship/Common/ViewportControls", Icon16x16));
 		Set("LevelEditor.PasteHere",				new IMAGE_BRUSH_SVG("Starship/Actors/paste-here", Icon16x16));
 
+		Set("LevelEditor.AllowArcballRotation",		    new IMAGE_BRUSH_SVG("Starship/Common/RotationArcball_16", Icon16x16));
+		Set("LevelEditor.AllowScreenspaceRotation",		    new IMAGE_BRUSH_SVG("Starship/Common/RotationScreenspace_16", Icon16x16));
+		Set("LevelEditor.EnableViewportHoverFeedback",		    new IMAGE_BRUSH_SVG("Starship/Common/PreselectionHighlight_16", Icon16x16));
 		Set("LevelEditor.AllowGroupSelection",			new IMAGE_BRUSH_SVG("Starship/Common/GroupActors", Icon16x16));
 		Set("LevelEditor.AllowTranslucentSelection",		new IMAGE_BRUSH_SVG("Starship/Common/Transparency", Icon16x16));
 		Set("LevelEditor.EnableActorSnap",				new IMAGE_BRUSH_SVG("Starship/Common/SnapActor_16", Icon16x16));
 		Set("LevelEditor.EnableVertexSnap",				new IMAGE_BRUSH_SVG("Starship/Common/SnapVertex_16", Icon16x16));
 		Set("LevelEditor.InvertSelection",				new IMAGE_BRUSH_SVG("Starship/Common/SelectInvert_16", Icon16x16));
+		Set("LevelEditor.PreserveNonUniformScale",		    new IMAGE_BRUSH_SVG("Starship/Common/ScaleNonUniform_16", Icon16x16));
 		Set("LevelEditor.PreviewPlatform", new IMAGE_BRUSH_SVG("Starship/Common/PreviewPlatform_16", Icon16x16));
 
 		Set("LevelEditor.SelectAllDescendants",			new IMAGE_BRUSH_SVG("Starship/Common/AllDescendants_16", Icon16x16));
@@ -4741,6 +4745,8 @@ void FStarshipEditorStyle::FStyle::SetupLevelEditorStyle()
 		Set("LevelEditor.StrictBoxSelect",				new IMAGE_BRUSH_SVG("Starship/Common/StrictMarqueeSelection_16", Icon16x16));
 		Set("LevelEditor.ToggleSocketSnapping", 			new IMAGE_BRUSH_SVG("Starship/Common/SnapSocket_16", Icon16x16));
 		Set("LevelEditor.TransparentBoxSelect", 			new IMAGE_BRUSH_SVG("Starship/Common/MarqueeSelectOccluded_16", Icon16x16));
+		Set("LevelEditor.ShowSelectionSubcomponents", 			new IMAGE_BRUSH_SVG("Starship/Common/ShowSubcomponents_16", Icon16x16));
+		Set("LevelEditor.UseExperimentalGizmos", 			new IMAGE_BRUSH_SVG("Starship/Common/GizmoExperimental_16", Icon16x16));
 
 		Set("LevelEditor.MaterialQuality", new IMAGE_BRUSH_SVG("Starship/Common/MaterialQuality_16", Icon16x16));
 			
diff --git a/Engine/Source/Editor/UnrealEd/Private/ViewportToolbar/UnrealEdViewportToolbar.cpp b/Engine/Source/Editor/UnrealEd/Private/ViewportToolbar/UnrealEdViewportToolbar.cpp
--- a/Engine/Source/Editor/UnrealEd/Private/ViewportToolbar/UnrealEdViewportToolbar.cpp
+++ b/Engine/Source/Editor/UnrealEd/Private/ViewportToolbar/UnrealEdViewportToolbar.cpp
@@ -968,7 +968,7 @@ FToolMenuEntry CreateTransformsSubmenu()
 								"TRSGizmoToggle",
 								LOCTEXT("TRSGizmoToggleLabel", "Use Experimental Gizmos"),
 								LOCTEXT("TRSGizmoToggleTooltip", "Whether or not to use the new Transform Gizmos"),
-								FSlateIcon(),
+								FSlateIcon(FAppStyle::GetAppStyleSetName(), "LevelEditor.UseExperimentalGizmos"),
 								Action,
 								EUserInterfaceActionType::ToggleButton
 							);
@@ -1044,7 +1044,7 @@ FToolMenuEntry CreateTransformsSubmenu()
 								"ScaleGridPreserveNonUniformScale_ToolTip",
 								"When this option is checked, scaling objects that have a non-uniform scale will preserve the ratios between each axis, snapping the axis with the largest value."
 							),
-							FSlateIcon(),
+							FSlateIcon(FAppStyle::GetAppStyleSetName(), "LevelEditor.PreserveNonUniformScale"),
 							PreserveNonUniformScaleAction,
 							EUserInterfaceActionType::ToggleButton
 						);
