Landscape Editor - Updating icons for Paint, Sculpt, Visibility, AutoFill Layers,  Create Layers From Materials, and Landscape Spline convert

diff --git a/Engine/Source/Editor/EditorStyle/Private/StarshipStyle.cpp b/Engine/Source/Editor/EditorStyle/Private/StarshipStyle.cpp
--- a/Engine/Source/Editor/EditorStyle/Private/StarshipStyle.cpp
+++ b/Engine/Source/Editor/EditorStyle/Private/StarshipStyle.cpp
@@ -2084,11 +2084,18 @@ void FStarshipEditorStyle::FStyle::SetupLevelGeneralStyles()
 			.SetDisabled(IMAGE_BRUSH_SVG("Starship/Splines/Spline_SelectLast", Icon20x20, SelectionColor_Inactive))
 		);
 
-		Set("SplineComponentDetails.Convert", FButtonStyle(Button)
-			.SetNormal(IMAGE_BRUSH_SVG("Starship/Splines/Spline_ToggleSelectMode", Icon20x20))
-			.SetHovered(IMAGE_BRUSH_SVG("Starship/Splines/Spline_ToggleSelectMode", Icon20x20, SelectionColor))
-			.SetPressed(IMAGE_BRUSH_SVG("Starship/Splines/Spline_ToggleSelectMode", Icon20x20, SelectionColor_Pressed))
-			.SetDisabled(IMAGE_BRUSH_SVG("Starship/Splines/Spline_ToggleSelectMode", Icon20x20, SelectionColor_Inactive))
+		Set("SplineComponentDetails.ConvertToSegments", FButtonStyle(Button)
+			.SetNormal(IMAGE_BRUSH_SVG("Starship/Splines/Spline_Segment", Icon20x20, FSlateColor::UseForeground()))
+			.SetHovered(IMAGE_BRUSH_SVG("Starship/Splines/Spline_Segment", Icon20x20, SelectionColor))
+			.SetPressed(IMAGE_BRUSH_SVG("Starship/Splines/Spline_Segment", Icon20x20, SelectionColor_Pressed))
+			.SetDisabled(IMAGE_BRUSH_SVG("Starship/Splines/Spline_Segment", Icon20x20, SelectionColor_Inactive))
+		);
+
+		Set("SplineComponentDetails.ConvertToPoints", FButtonStyle(Button)
+			.SetNormal(IMAGE_BRUSH_SVG("Starship/Splines/Spline_ControlPoint", Icon20x20, FSlateColor::UseForeground()))
+			.SetHovered(IMAGE_BRUSH_SVG("Starship/Splines/Spline_ControlPoint", Icon20x20, SelectionColor))
+			.SetPressed(IMAGE_BRUSH_SVG("Starship/Splines/Spline_ControlPoint", Icon20x20, SelectionColor_Pressed))
+			.SetDisabled(IMAGE_BRUSH_SVG("Starship/Splines/Spline_ControlPoint", Icon20x20, SelectionColor_Inactive))
 		);
 	}
 }
@@ -7125,8 +7132,8 @@ void FStarshipEditorStyle::FStyle::SetupLandscapeEditorStyle()
 			Set("LandscapeEditor.NoiseTool.Small",        		new IMAGE_BRUSH("Icons/Landscape/Landscape_Noise_x40",             Icon20x20));
 			Set("LandscapeEditor.RetopologizeTool", 			new IMAGE_BRUSH("Icons/Landscape/Landscape_Retopologize_x40",      Icon20x20));
 			Set("LandscapeEditor.RetopologizeTool.Small", 		new IMAGE_BRUSH("Icons/Landscape/Landscape_Retopologize_x40",      Icon20x20));
-			Set("LandscapeEditor.VisibilityTool",   			new IMAGE_BRUSH("Icons/Landscape/Landscape_VisibilityTool_x40",    Icon20x20));
-			Set("LandscapeEditor.VisibilityTool.Small",   		new IMAGE_BRUSH("Icons/Landscape/Landscape_VisibilityTool_x40",    Icon20x20));
+			Set("LandscapeEditor.VisibilityTool",   			new IMAGE_BRUSH_SVG("Icons/Landscape/Landscape_VisibilityTool_x16",Icon20x20));
+			Set("LandscapeEditor.VisibilityTool.Small",   		new IMAGE_BRUSH_SVG("Icons/Landscape/Landscape_VisibilityTool_x16",Icon20x20));
 			Set("LandscapeEditor.BlueprintBrushTool", 			new IMAGE_BRUSH("Icons/Landscape/Landscape_BlueprintTool_x40",     Icon20x20));
 			Set("LandscapeEditor.BlueprintBrushTool.Small", 	new IMAGE_BRUSH("Icons/Landscape/Landscape_BlueprintTool_x40",     Icon20x20));
 
@@ -7203,9 +7210,10 @@ void FStarshipEditorStyle::FStyle::SetupLandscapeEditorStyle()
 		Set("LandscapeEditor.TargetList.RowSelected",          new BOX_BRUSH("Common/RoundedSelection_16x", 4.0f / 16.0f, FStyleColors::Select));
 		Set("LandscapeEditor.TargetList.RowSelectedHovered",   new BOX_BRUSH("Common/RoundedSelection_16x", 4.0f / 16.0f, FStyleColors::PrimaryHover));
 
-		Set("LandscapeEditor.Target_Heightmap",  new IMAGE_BRUSH("Icons/icon_Landscape_Target_Heightmap_48x",  Icon48x48));
-		Set("LandscapeEditor.Target_Visibility", new IMAGE_BRUSH("Icons/icon_Landscape_Target_Visibility_48x", Icon48x48));
-		Set("LandscapeEditor.Target_Invalid",    new IMAGE_BRUSH("Icons/icon_Landscape_Target_Invalid_48x",    Icon48x48));
+		Set("LandscapeEditor.Target_Heightmap",  new IMAGE_BRUSH_SVG("Icons/icon_Landscape_Target_Heightmap_64x",  Icon48x48));
+		Set("LandscapeEditor.Target_Weightmap",	 new IMAGE_BRUSH_SVG("Icons/icon_Landscape_Target_Weightmap_64x",  Icon48x48));
+		Set("LandscapeEditor.Target_Visibility", new IMAGE_BRUSH_SVG("Icons/icon_Landscape_Target_Visibility_64x", Icon48x48));
+		Set("LandscapeEditor.Target_Unknown",    new IMAGE_BRUSH_SVG("Icons/icon_Landscape_Target_Unknown_64x",    Icon48x48));
 
 		Set("LandscapeEditor.Target_Create",     new IMAGE_BRUSH("Icons/icon_Landscape_Target_Create_12x", Icon12x12));
 		Set("LandscapeEditor.Target_MakePublic", new IMAGE_BRUSH("Icons/assign_right_12x",                 Icon12x12));
diff --git a/Engine/Source/Editor/LandscapeEditor/Private/LandscapeEditorDetailCustomization_ImportLayers.cpp b/Engine/Source/Editor/LandscapeEditor/Private/LandscapeEditorDetailCustomization_ImportLayers.cpp
--- a/Engine/Source/Editor/LandscapeEditor/Private/LandscapeEditorDetailCustomization_ImportLayers.cpp
+++ b/Engine/Source/Editor/LandscapeEditor/Private/LandscapeEditorDetailCustomization_ImportLayers.cpp
@@ -141,7 +141,7 @@ void FLandscapeEditorStructCustomization_FLandscapeImportLayer::CustomizeChildre
 		.VAlign(VAlign_Center)
 		.Padding(FMargin(2))
 		[
-			SNew(SLandscapeAssetThumbnail, ThumbnailMIC, StructCustomizationUtils.GetThumbnailPool().ToSharedRef())
+			SNew(SLandscapeAssetThumbnail, ThumbnailMIC, StructCustomizationUtils.GetThumbnailPool().ToSharedRef(), FName("LandscapeEditor.Target_Weightmap"))
 			.ThumbnailSize(FIntPoint(48, 48))
 		]
 		+ SHorizontalBox::Slot()
diff --git a/Engine/Source/Editor/LandscapeEditor/Private/LandscapeEditorDetailCustomization_TargetLayers.cpp b/Engine/Source/Editor/LandscapeEditor/Private/LandscapeEditorDetailCustomization_TargetLayers.cpp
--- a/Engine/Source/Editor/LandscapeEditor/Private/LandscapeEditorDetailCustomization_TargetLayers.cpp
+++ b/Engine/Source/Editor/LandscapeEditor/Private/LandscapeEditorDetailCustomization_TargetLayers.cpp
@@ -1118,8 +1118,8 @@ TSharedPtr<SWidget> FLandscapeEditorCustomNodeBuilder_TargetLayers::GenerateRow(
 					.VAlign(VAlign_Center)
 					.Padding(FMargin(2))
 					[
-						SNew(SImage)
-							.Image(FAppStyle::GetBrush(Target->TargetType == ELandscapeToolTargetType::Heightmap ? TEXT("LandscapeEditor.Target_Heightmap") : TEXT("LandscapeEditor.Target_Visibility")))
+						SNew(SLandscapeAssetThumbnail, /*Asset = */ nullptr, ThumbnailPool, Target->TargetType == ELandscapeToolTargetType::Heightmap ? FName("LandscapeEditor.Target_Heightmap") : FName("LandscapeEditor.Target_Visibility"))
+						.ThumbnailSize(FIntPoint(48, 48))
 					]
 					+ SHorizontalBox::Slot()
 					.VAlign(VAlign_Center)
@@ -1186,28 +1186,20 @@ TSharedPtr<SWidget> FLandscapeEditorCustomNodeBuilder_TargetLayers::GenerateRow(
 					.VAlign(VAlign_Center)
 					.Padding(FMargin(2))
 					[
-						(Target->bIsLayerReferencedByMaterial)
-							? (TSharedRef<SWidget>)(
-								SNew(SLandscapeAssetThumbnail, Target->ThumbnailMIC.Get(), ThumbnailPool)
-								.Visibility_Static(&FLandscapeEditorCustomNodeBuilder_TargetLayers::GetDebugModeLayerUsageVisibility_Invert, Target)
-								.ThumbnailSize(FIntPoint(48, 48))
-								// Open landscape layer info asset on double-click on the thumbnail : 
-								.OnAccessAsset_Lambda([Target](UObject* InObject)
-								{ 
-									// Note : the object being returned here is the landscape MIC so it's not what we use for opening the landscape layer info asset : 
-									if ((Target->TargetType == ELandscapeToolTargetType::Weightmap) && (Target->LayerInfoObj != nullptr))
-									{
-										UAssetEditorSubsystem* AssetEditorSubsystem = GEditor->GetEditorSubsystem<UAssetEditorSubsystem>();
-										return AssetEditorSubsystem->OpenEditorForAsset(Target->LayerInfoObj.Get());
-									}
-									return false;
-								}))
-							: (TSharedRef<SWidget>)(
-								SNew(SImage)
-								.Visibility_Static(&FLandscapeEditorCustomNodeBuilder_TargetLayers::GetDebugModeLayerUsageVisibility_Invert, Target)
-								.Image(FAppStyle::GetBrush(TEXT("LandscapeEditor.Target_Invalid")))
-								.ToolTipText(LOCTEXT("LandscapeTargetLayer_NotReferencedByMaterial", "This target layer is not referenced by any assigned landscape material"))
-								)
+						SNew(SLandscapeAssetThumbnail, Target->bIsLayerReferencedByMaterial ? Target->ThumbnailMIC.Get() : nullptr, ThumbnailPool, Target->bIsLayerReferencedByMaterial ? FName("LandscapeEditor.Target_Weightmap") : FName("LandscapeEditor.Target_Unknown"))
+						.Visibility_Static(&FLandscapeEditorCustomNodeBuilder_TargetLayers::GetDebugModeLayerUsageVisibility_Invert, Target)
+						.ThumbnailSize(FIntPoint(48, 48))
+						// Open landscape layer info asset on double-click on the thumbnail : 
+						.OnAccessAsset_Lambda([Target](UObject* InObject)
+						{ 
+							// Note : the object being returned here is the landscape MIC so it's not what we use for opening the landscape layer info asset : 
+							if ((Target->TargetType == ELandscapeToolTargetType::Weightmap) && (Target->LayerInfoObj != nullptr))
+							{
+								UAssetEditorSubsystem* AssetEditorSubsystem = GEditor->GetEditorSubsystem<UAssetEditorSubsystem>();
+								return AssetEditorSubsystem->OpenEditorForAsset(Target->LayerInfoObj.Get());
+							}
+							return false;
+						})
 					]
 					+ SHorizontalBox::Slot()
 					.VAlign(VAlign_Center)
diff --git a/Engine/Source/Editor/LandscapeEditor/Private/LandscapeSplineDetails.cpp b/Engine/Source/Editor/LandscapeEditor/Private/LandscapeSplineDetails.cpp
--- a/Engine/Source/Editor/LandscapeEditor/Private/LandscapeSplineDetails.cpp
+++ b/Engine/Source/Editor/LandscapeEditor/Private/LandscapeSplineDetails.cpp
@@ -170,7 +170,7 @@ void FLandscapeSplineDetails::CustomizeDetails(IDetailLayoutBuilder& DetailBuild
 		.AutoWidth()
 		[
 			SNew(SButton)
-				.ButtonStyle(FAppStyle::Get(), "SplineComponentDetails.Convert")
+				.ButtonStyle(FAppStyle::Get(), IsSegmentSelectModeEnabled() ? "SplineComponentDetails.ConvertToPoints" : "SplineComponentDetails.ConvertToSegments")
 				.ContentPadding(ButtonPadding)
 				.VAlign(VAlign_Center)
 				.HAlign(HAlign_Center)
diff --git a/Engine/Source/Editor/LandscapeEditor/Private/SLandscapeEditor.cpp b/Engine/Source/Editor/LandscapeEditor/Private/SLandscapeEditor.cpp
--- a/Engine/Source/Editor/LandscapeEditor/Private/SLandscapeEditor.cpp
+++ b/Engine/Source/Editor/LandscapeEditor/Private/SLandscapeEditor.cpp
@@ -36,7 +36,7 @@ namespace UE::Landscape::Private {
 	);
 }
 
-void SLandscapeAssetThumbnail::Construct(const FArguments& InArgs, UObject* Asset, TSharedRef<FAssetThumbnailPool> ThumbnailPool)
+void SLandscapeAssetThumbnail::Construct(const FArguments& InArgs, UObject* Asset, TSharedRef<FAssetThumbnailPool> ThumbnailPool, const FName& ClassThumbnailBrushOverride)
 {
 	FIntPoint ThumbnailSize = InArgs._ThumbnailSize;
 
@@ -44,11 +44,17 @@ void SLandscapeAssetThumbnail::Construct(const FArguments& InArgs, UObject* Asse
 	OnAccessAsset = InArgs._OnAccessAsset;
 
 	FAssetThumbnailConfig AssetThumbnailConfig;
+	AssetThumbnailConfig.ShowAssetChip = false;
+
 	// If the asset is null, then we purposefully don't want this layer to have a thumbnail. Display a generic icon in that case :
 	if (Asset == nullptr)
 	{
 		AssetThumbnailConfig.bForceGenericThumbnail = true;
-		AssetThumbnailConfig.ClassThumbnailBrushOverride = FName("LevelEditor.LandscapeMode");
+		AssetThumbnailConfig.ClassThumbnailBrushOverride = ClassThumbnailBrushOverride;
+		AssetThumbnailConfig.bAllowHintText = false;
+		AssetThumbnailConfig.AlwaysExpandTooltip = false;
+		AssetThumbnailConfig.AllowAssetStatusThumbnailOverlay = false;
+		AssetThumbnailConfig.ThumbnailLabel = EThumbnailLabel::Type::NoLabel;
 	}
 	ChildSlot
 	[
diff --git a/Engine/Source/Editor/LandscapeEditor/Private/SLandscapeEditor.h b/Engine/Source/Editor/LandscapeEditor/Private/SLandscapeEditor.h
--- a/Engine/Source/Editor/LandscapeEditor/Private/SLandscapeEditor.h
+++ b/Engine/Source/Editor/LandscapeEditor/Private/SLandscapeEditor.h
@@ -31,7 +31,7 @@ public:
 		SLATE_EVENT(FAccessAsset, OnAccessAsset)
 	SLATE_END_ARGS()
 
-	void Construct(const FArguments& InArgs, UObject* Asset, TSharedRef<FAssetThumbnailPool> ThumbnailPool);
+	void Construct(const FArguments& InArgs, UObject* Asset, TSharedRef<FAssetThumbnailPool> ThumbnailPool, const FName& ClassThumbnailBrushOverride);
 	~SLandscapeAssetThumbnail();
 
 	void SetAsset(UObject* Asset);
