Sequencer: Updated proxy icon overlay to new style. Tweaked wording on menus to match other entries.

diff --git a/Engine/Plugins/MovieScene/LevelSequenceEditor/Source/LevelSequenceEditor/Private/LevelSequenceEditorSubsystem.cpp b/Engine/Plugins/MovieScene/LevelSequenceEditor/Source/LevelSequenceEditor/Private/LevelSequenceEditorSubsystem.cpp
--- a/Engine/Plugins/MovieScene/LevelSequenceEditor/Source/LevelSequenceEditor/Private/LevelSequenceEditorSubsystem.cpp
+++ b/Engine/Plugins/MovieScene/LevelSequenceEditor/Source/LevelSequenceEditor/Private/LevelSequenceEditorSubsystem.cpp
@@ -3373,11 +3373,13 @@ void ULevelSequenceEditorSubsystem::AddChangeBindingTypeMenu(FMenuBuilder& MenuB
 		});
 
 		MenuBuilder.AddSubMenu(
-			LOCTEXT("ConvertToProxy", "Convert to Proxy"),
+			LOCTEXT("ConvertToProxy", "Proxy Binding"),
 			LOCTEXT("ConvertToProxyTooltip", "Convert selected binding(s) to a proxy binding that simply references a binding in another sequence"),
 			FNewMenuDelegate::CreateLambda([ProxyPicker](FMenuBuilder& SubMenuBuilder) {
 				ProxyPicker->GetPickerMenu(SubMenuBuilder);
-			})
+			}),
+			false,
+			FSlateIcon(FAppStyle::GetAppStyleSetName(), "Sequencer.ProxyIconOverlay")
 		);
 	}
 
diff --git a/Engine/Source/Editor/EditorStyle/Private/StarshipStyle.cpp b/Engine/Source/Editor/EditorStyle/Private/StarshipStyle.cpp
--- a/Engine/Source/Editor/EditorStyle/Private/StarshipStyle.cpp
+++ b/Engine/Source/Editor/EditorStyle/Private/StarshipStyle.cpp
@@ -2389,7 +2389,7 @@ void FStarshipEditorStyle::FStyle::SetupSequencerStyles()
 		Set("Sequencer.SpawnableIconOverlay", new IMAGE_BRUSH(TEXT("Sequencer/SpawnableIconOverlay"), FVector2f(13.f, 13.f)));
 		Set("Sequencer.ReplaceableIconOverlay", new IMAGE_BRUSH(TEXT("Sequencer/ReplaceableIconOverlay"), FVector2f(13.f, 13.f)));
 		Set("Sequencer.MultipleIconOverlay", new IMAGE_BRUSH(TEXT("Sequencer/MultipleIconOverlay"), FVector2f(13.f, 13.f)));
-		Set("Sequencer.ProxyIconOverlay", new IMAGE_BRUSH(TEXT("Sequencer/ProxyIconOverlay"), FVector2f(13.f, 13.f)));
+		Set("Sequencer.ProxyIconOverlay", new IMAGE_BRUSH_SVG(TEXT("Sequencer/ProxyIconOverlay_13"), FVector2f(13.f, 13.f)));
 		Set("Sequencer.DynamicBindingIconOverlay", new IMAGE_BRUSH(TEXT("Sequencer/DynamicBindingIconOverlay"), Icon16x16));
 		Set("Sequencer.SpawnableDynamicBindingIconOverlay", new IMAGE_BRUSH(TEXT("Sequencer/SpawnableDynamicBindingIconOverlay"), Icon16x16));
 		Set("Sequencer.LockSequence", new IMAGE_BRUSH("Sequencer/Main_Icons/Icon_Sequencer_Locked_16x", Icon16x16));
