Landscape Editor - Updating landscape visibility icon file path to use updated 5.6 icon

diff --git a/Engine/Source/Editor/EditorStyle/Private/StarshipStyle.cpp b/Engine/Source/Editor/EditorStyle/Private/StarshipStyle.cpp
--- a/Engine/Source/Editor/EditorStyle/Private/StarshipStyle.cpp
+++ b/Engine/Source/Editor/EditorStyle/Private/StarshipStyle.cpp
@@ -7128,8 +7128,8 @@ void FStarshipEditorStyle::FStyle::SetupLandscapeEditorStyle()
 			Set("LandscapeEditor.NoiseTool.Small",        		new IMAGE_BRUSH("Icons/Landscape/Landscape_Noise_x40",             Icon20x20));
 			Set("LandscapeEditor.RetopologizeTool", 			new IMAGE_BRUSH("Icons/Landscape/Landscape_Retopologize_x40",      Icon20x20));
 			Set("LandscapeEditor.RetopologizeTool.Small", 		new IMAGE_BRUSH("Icons/Landscape/Landscape_Retopologize_x40",      Icon20x20));
-			Set("LandscapeEditor.VisibilityTool",   			new IMAGE_BRUSH_SVG("Icons/Landscape/Landscape_VisibilityTool_x16",Icon20x20));
-			Set("LandscapeEditor.VisibilityTool.Small",   		new IMAGE_BRUSH_SVG("Icons/Landscape/Landscape_VisibilityTool_x16",Icon20x20));
+			Set("LandscapeEditor.VisibilityTool",   			new IMAGE_BRUSH_SVG("Icons/icon_Landscape_Target_Visibility_16x",  Icon20x20));
+			Set("LandscapeEditor.VisibilityTool.Small",   		new IMAGE_BRUSH_SVG("Icons/icon_Landscape_Target_Visibility_16x",  Icon20x20));
 			Set("LandscapeEditor.BlueprintBrushTool", 			new IMAGE_BRUSH("Icons/Landscape/Landscape_BlueprintTool_x40",     Icon20x20));
 			Set("LandscapeEditor.BlueprintBrushTool.Small", 	new IMAGE_BRUSH("Icons/Landscape/Landscape_BlueprintTool_x40",     Icon20x20));
 
