[Translation Picker] Adding an icon for the Translation Picker menu entry

diff --git a/Engine/Source/Editor/EditorStyle/Private/StarshipStyle.cpp b/Engine/Source/Editor/EditorStyle/Private/StarshipStyle.cpp
--- a/Engine/Source/Editor/EditorStyle/Private/StarshipStyle.cpp
+++ b/Engine/Source/Editor/EditorStyle/Private/StarshipStyle.cpp
@@ -973,6 +973,11 @@ void FStarshipEditorStyle::FStyle::SetupGeneralStyles()
 		Set("AutomationTools.TestAutomation", new IMAGE_BRUSH_SVG("Starship/Common/TestAutomation", Icon16x16));
 	}
 
+	// Translation Picker Menu
+	{
+		Set("TranslationPicker.MenuIcon", new CORE_IMAGE_BRUSH_SVG("Starship/Common/TranslationPicker", Icon16x16));
+	}
+
 	// Session Browser tab
 	{
 		Set("SessionBrowser.Terminate.Font", FTextBlockStyle(NormalText)
diff --git a/Engine/Source/Editor/MainFrame/Private/Menus/MainMenu.cpp b/Engine/Source/Editor/MainFrame/Private/Menus/MainMenu.cpp
--- a/Engine/Source/Editor/MainFrame/Private/Menus/MainMenu.cpp
+++ b/Engine/Source/Editor/MainFrame/Private/Menus/MainMenu.cpp
@@ -638,7 +638,7 @@ void FMainMenu::RegisterToolsMenu()
 					"TranslationPicker",
 					LOCTEXT("TranslationPickerMenuItem", "Translation Picker"),
 					LOCTEXT("TranslationPickerMenuItemToolTip", "Launch the Translation Picker to Modify Editor Translations"),
-					FSlateIcon(),
+					FSlateIcon(FAppStyle::GetAppStyleSetName(), "TranslationPicker.MenuIcon"),
 					FUIAction(FExecuteAction::CreateLambda(
 						[]()
 						{
