Updated headers using UnrealCodeFixup to make sure dllstorage is on methods/staticvars instead of types. Part 5/n

diff --git a/Engine/Source/Editor/EditorStyle/Public/EditorFontGlyphs.h b/Engine/Source/Editor/EditorStyle/Public/EditorFontGlyphs.h
--- a/Engine/Source/Editor/EditorStyle/Public/EditorFontGlyphs.h
+++ b/Engine/Source/Editor/EditorStyle/Public/EditorFontGlyphs.h
@@ -4,293 +4,297 @@
 
 #include "CoreMinimal.h"
 
+#define UE_API EDITORSTYLE_API
+
 class FText;
 
-class EDITORSTYLE_API FEditorFontGlyphs
+class FEditorFontGlyphs
 {
 public:
 	// (laquo, quote, previous, back) http://fortawesome.github.io/Font-Awesome/icon/angle-double-left/
-	static FText Angle_Double_Left;
+	static UE_API FText Angle_Double_Left;
 
 	// (raquo, quote, next, forward) http://fortawesome.github.io/Font-Awesome/icon/angle-double-right/
-	static FText Angle_Double_Right;
+	static UE_API FText Angle_Double_Right;
 
 	// (previous, back) http://fortawesome.github.io/Font-Awesome/icon/angle-left/
-	static FText Angle_Left;
+	static UE_API FText Angle_Left;
 
 	// (next, forward) http://fortawesome.github.io/Font-Awesome/icon/angle-right/
-	static FText Angle_Right;
+	static UE_API FText Angle_Right;
 
 	// http://fortawesome.github.io/Font-Awesome/icon/angle-up/
-	static FText Angle_Up;
+	static UE_API FText Angle_Up;
 
 	// (box, storage) http://fortawesome.github.io/Font-Awesome/icon/archive/
-	static FText Archive;
+	static UE_API FText Archive;
 
 	// (next, forward) http://fortawesome.github.io/Font-Awesome/icon/arrow-circle-o-right/
-	static FText Arrow_Circle_O_Right;
+	static UE_API FText Arrow_Circle_O_Right;
 
 	// (download) http://fortawesome.github.io/Font-Awesome/icon/arrow-down/
-	static FText Arrow_Down;
+	static UE_API FText Arrow_Down;
 
 	// (previous, back) http://fortawesome.github.io/Font-Awesome/icon/arrow-left/
-	static FText Arrow_Left;
+	static UE_API FText Arrow_Left;
 
 	// (next, forward) http://fortawesome.github.io/Font-Awesome/icon/arrow-right/
-	static FText Arrow_Right;
+	static UE_API FText Arrow_Right;
 
 	// (move, reorder, resize) http://fortawesome.github.io/Font-Awesome/icon/arrows/
-	static FText Arrows;
+	static UE_API FText Arrows;
 
 	// (delete, remove, trash, hide, block, stop, abort, cancel) http://fortawesome.github.io/Font-Awesome/icon/ban/
-	static FText Ban;
+	static UE_API FText Ban;
 
 	// (menu, drag, reorder, settings, list, ul, ol, checklist, todo, list, hamburger) http://fortawesome.github.io/Font-Awesome/icon/bars/
-	static FText Bars;
+	static UE_API FText Bars;
 
 	// (read, documentation) http://fortawesome.github.io/Font-Awesome/icon/book/
-	static FText Book;
+	static UE_API FText Book;
 
 	// (report, insect) http://fortawesome.github.io/Font-Awesome/icon/bug/
-	static FText Bug;
+	static UE_API FText Bug;
 
 	// (vehicle) http://fortawesome.github.io/Font-Awesome/icon/car/
-	static FText Car;
+	static UE_API FText Car;
 
 	// (more, dropdown, menu, triangle down) http://fortawesome.github.io/Font-Awesome/icon/caret-down/
-	static FText Caret_Down;
+	static UE_API FText Caret_Down;
 
 	// (next, forward, triangle right) http://fortawesome.github.io/Font-Awesome/icon/caret-right/
-	static FText Caret_Right;
+	static UE_API FText Caret_Right;
 
 	// (more, dropdown, menu) http://fortawesome.github.io/Font-Awesome/icon/caret-square-o-down/
-	static FText Caret_Square_O_Down;
+	static UE_API FText Caret_Square_O_Down;
 
 	// (previous, back) http://fortawesome.github.io/Font-Awesome/icon/caret-square-o-left/
-	static FText Caret_Square_O_Left;
+	static UE_API FText Caret_Square_O_Left;
 
 	// (next, forward) http://fortawesome.github.io/Font-Awesome/icon/caret-square-o-right/
-	static FText Caret_Square_O_Right;
+	static UE_API FText Caret_Square_O_Right;
 
 	// http://fortawesome.github.io/Font-Awesome/icon/caret-square-o-up/
-	static FText Caret_Square_O_Up;
+	static UE_API FText Caret_Square_O_Up;
 
 	// (triangle up) http://fortawesome.github.io/Font-Awesome/icon/caret-up/
-	static FText Caret_Up;
+	static UE_API FText Caret_Up;
 
 	// (checkmark, done, todo, agree, accept, confirm, tick) http://fortawesome.github.io/Font-Awesome/icon/check/
-	static FText Check;
+	static UE_API FText Check;
 
 	// (todo, done, agree, accept, confirm) http://fortawesome.github.io/Font-Awesome/icon/check-circle/
-	static FText Check_Circle;
+	static UE_API FText Check_Circle;
 
 	// (dot, notification) http://fortawesome.github.io/Font-Awesome/icon/circle/
-	static FText Circle;
+	static UE_API FText Circle;
 
 	// http://fortawesome.github.io/Font-Awesome/icon/circle-o/
-	static FText Circle_O;
+	static UE_API FText Circle_O;
 
 	// (watch, timer, late, timestamp) http://fortawesome.github.io/Font-Awesome/icon/clock-o/
-	static FText Clock_O;
+	static UE_API FText Clock_O;
 
 	// (settings) http://fortawesome.github.io/Font-Awesome/icon/cog/
-	static FText Cog;
+	static UE_API FText Cog;
 
 	// (settings) http://fortawesome.github.io/Font-Awesome/icon/cogs/
-	static FText Cogs;
+	static UE_API FText Cogs;
 
 	// (picker) http://fortawesome.github.io/Font-Awesome/icon/crosshairs/
-	static FText Crosshairs;
+	static UE_API FText Crosshairs;
 
 	// http://fortawesome.github.io/Font-Awesome/icon/database/
-	static FText Database;
+	static UE_API FText Database;
 
 	// (import) http://fortawesome.github.io/Font-Awesome/icon/download/
-	static FText Download;
+	static UE_API FText Download;
 
 	// http://fortawesome.github.io/Font-Awesome/icon/eraser/
-	static FText Eraser;
+	static UE_API FText Eraser;
 
 	// (transfer, arrows) http://fortawesome.github.io/Font-Awesome/icon/exchange/
-	static FText Exchange;
+	static UE_API FText Exchange;
 
 	// (warning, error, problem, notification, notify, alert) http://fortawesome.github.io/Font-Awesome/icon/exclamation/
-	static FText Exclamation;
+	static UE_API FText Exclamation;
 
 	// (warning, error, problem, notification, alert) http://fortawesome.github.io/Font-Awesome/icon/exclamation-triangle/
-	static FText Exclamation_Triangle;
+	static UE_API FText Exclamation_Triangle;
 
 	// (enlarge, bigger, resize) http://fortawesome.github.io/Font-Awesome/icon/expand/
-	static FText Expand;
+	static UE_API FText Expand;
 
 	// (open, new) http://fortawesome.github.io/Font-Awesome/icon/external-link/
-	static FText External_Link;
+	static UE_API FText External_Link;
 
 	// (show, visible, views) http://fortawesome.github.io/Font-Awesome/icon/eye/
-	static FText Eye;
+	static UE_API FText Eye;
 
 	// (toggle, show, hide, visible, visiblity, views) http://fortawesome.github.io/Font-Awesome/icon/eye-slash/
-	static FText Eye_Slash;
+	static UE_API FText Eye_Slash;
 
 	// (next, end, last) http://fortawesome.github.io/Font-Awesome/icon/fast-forward/
-	static FText Fast_Forward;
+	static UE_API FText Fast_Forward;
 
 	// (new, page, pdf, document) http://fortawesome.github.io/Font-Awesome/icon/file/
-	static FText File;
+	static UE_API FText File;
 
 	// (movie) http://fortawesome.github.io/Font-Awesome/icon/film/
-	static FText Film;
+	static UE_API FText Film;
 
 	// (funnel, options) http://fortawesome.github.io/Font-Awesome/icon/filter/
-	static FText Filter;
+	static UE_API FText Filter;
 
 	// http://fortawesome.github.io/Font-Awesome/icon/floppy-o/
-	static FText Floppy_O;
+	static UE_API FText Floppy_O;
 
 	// http://fortawesome.github.io/Font-Awesome/icon/folder/
-	static FText Folder;
+	static UE_API FText Folder;
 
 	// http://fortawesome.github.io/Font-Awesome/icon/folder-open/
-	static FText Folder_Open;
+	static UE_API FText Folder_Open;
 
 	// http://fortawesome.github.io/Font-Awesome/icon/hourglass/
-	static FText Hourglass;
+	static UE_API FText Hourglass;
 
 	// http://fortawesome.github.io/Font-Awesome/icon/hourglass-o/
-	static FText Hourglass_O;
+	static UE_API FText Hourglass_O;
 
 	// (help, information, more, details) http://fortawesome.github.io/Font-Awesome/icon/info/
-	static FText Info;
+	static UE_API FText Info;
 
 	// (help, information, more, details) http://fortawesome.github.io/Font-Awesome/icon/info-circle/
-	static FText Info_Circle;
+	static UE_API FText Info_Circle;
 
 	// http://fortawesome.github.io/Font-Awesome/icon/level-down/
-	static FText Level_Down;
+	static UE_API FText Level_Down;
 
 	// http://fortawesome.github.io/Font-Awesome/icon/level-up/
-	static FText Level_Up;
+	static UE_API FText Level_Up;
 
 	// (idea, inspiration) http://fortawesome.github.io/Font-Awesome/icon/lightbulb-o/
-	static FText Lightbulb_O;
+	static UE_API FText Lightbulb_O;
 
 	// (graph, analytics) http://fortawesome.github.io/Font-Awesome/icon/line-chart/
-	static FText Line_Chart;
+	static UE_API FText Line_Chart;
 
 	// (chain) http://fortawesome.github.io/Font-Awesome/icon/link/
-	static FText Link;
+	static UE_API FText Link;
 
 	// (protect, admin) http://fortawesome.github.io/Font-Awesome/icon/lock/
-	static FText Lock;
+	static UE_API FText Lock;
 
 	// http://fortawesome.github.io/Font-Awesome/icon/long-arrow-down/
-	static FText Long_Arrow_Down;
+	static UE_API FText Long_Arrow_Down;
 
 	// http://fortawesome.github.io/Font-Awesome/icon/long-arrow-right/
-	static FText Long_Arrow_Right;
+	static UE_API FText Long_Arrow_Right;
 
 	// http://fortawesome.github.io/Font-Awesome/icon/map/
-	static FText Map;
+	static UE_API FText Map;
 
 	// (map, pin, location, coordinates, localize, address, travel, where, place) http://fortawesome.github.io/Font-Awesome/icon/map-marker/
-	static FText Map_Marker;
+	static UE_API FText Map_Marker;
 
 	// (hide, minify, delete, remove, trash, hide, collapse) http://fortawesome.github.io/Font-Awesome/icon/minus/
-	static FText Minus;
+	static UE_API FText Minus;
 
 	// (delete, remove, trash, hide) http://fortawesome.github.io/Font-Awesome/icon/minus-circle/
-	static FText Minus_Circle;
+	static UE_API FText Minus_Circle;
 
 	// http://fortawesome.github.io/Font-Awesome/icon/paper-plane/
-	static FText Paper_Plane;
+	static UE_API FText Paper_Plane;
 
 	// (wait) http://fortawesome.github.io/Font-Awesome/icon/pause/
-	static FText Pause;
+	static UE_API FText Pause;
 
 	// http://fortawesome.github.io/Font-Awesome/icon/pause-circle/
-	static FText Pause_Circle;
+	static UE_API FText Pause_Circle;
 
 	// (write, edit, update) http://fortawesome.github.io/Font-Awesome/icon/pencil/
-	static FText Pencil;
+	static UE_API FText Pencil;
 
 	// (write, edit, update) http://fortawesome.github.io/Font-Awesome/icon/pencil-square/
-	static FText Pencil_Square;
+	static UE_API FText Pencil_Square;
 
 	// (start, playing, music, sound) http://fortawesome.github.io/Font-Awesome/icon/play/
-	static FText Play;
+	static UE_API FText Play;
 
 	// (start, playing) http://fortawesome.github.io/Font-Awesome/icon/play-circle/
-	static FText Play_Circle;
+	static UE_API FText Play_Circle;
 
 	// http://fortawesome.github.io/Font-Awesome/icon/plug/
-	static FText Plug;
+	static UE_API FText Plug;
 
 	// (add, new, create, expand) http://fortawesome.github.io/Font-Awesome/icon/plus/
-	static FText Plus;
+	static UE_API FText Plus;
 
 	// (add, new, create, expand) http://fortawesome.github.io/Font-Awesome/icon/plus-circle/
-	static FText Plus_Circle;
+	static UE_API FText Plus_Circle;
 
 	// (help, information, unknown, support) http://fortawesome.github.io/Font-Awesome/icon/question/
-	static FText Question;
+	static UE_API FText Question;
 
 	// (sort, shuffle) http://fortawesome.github.io/Font-Awesome/icon/random/
-	static FText Random;
+	static UE_API FText Random;
 
 	// http://fortawesome.github.io/Font-Awesome/icon/recycle/
-	static FText Recycle;
+	static UE_API FText Recycle;
 
 	// (reload, sync) http://fortawesome.github.io/Font-Awesome/icon/refresh/
-	static FText Refresh;
+	static UE_API FText Refresh;
 
 	// http://fortawesome.github.io/Font-Awesome/icon/share/
-	static FText Share;
+	static UE_API FText Share;
 
 	// http://fortawesome.github.io/Font-Awesome/icon/share-alt/
-	static FText Share_Alt;
+	static UE_API FText Share_Alt;
 
 	// (enter, join, log in, login, sign up, sign in, signin, signup, arrow) http://fortawesome.github.io/Font-Awesome/icon/sign-in/
-	static FText Sign_In;
+	static UE_API FText Sign_In;
 
 	// (log out, logout, leave, exit, arrow) http://fortawesome.github.io/Font-Awesome/icon/sign-out/
-	static FText Sign_Out;
+	static UE_API FText Sign_Out;
 
 	// (block, box) http://fortawesome.github.io/Font-Awesome/icon/square/
-	static FText Square;
+	static UE_API FText Square;
 
 	// (award, achievement, night, rating, score) http://fortawesome.github.io/Font-Awesome/icon/star/
-	static FText Star;
+	static UE_API FText Star;
 
 	// (next, end, last) http://fortawesome.github.io/Font-Awesome/icon/step-forward/
-	static FText Step_Forward;
+	static UE_API FText Step_Forward;
 
 	// (block, box, square) http://fortawesome.github.io/Font-Awesome/icon/stop/
-	static FText Stop;
+	static UE_API FText Stop;
 
 	// (command, prompt, code) http://fortawesome.github.io/Font-Awesome/icon/terminal/
-	static FText Terminal;
+	static UE_API FText Terminal;
 
 	// (blocks, squares, boxes) http://fortawesome.github.io/Font-Awesome/icon/th/
-	static FText Th;
+	static UE_API FText Th;
 
 	// (marker, pin, location, coordinates) http://fortawesome.github.io/Font-Awesome/icon/thumb-tack/
-	static FText Thumb_Tack;
+	static UE_API FText Thumb_Tack;
 
 	// (close, exit, x, cross) http://fortawesome.github.io/Font-Awesome/icon/times/
-	static FText Times;
+	static UE_API FText Times;
 
 	// (close, exit, x) http://fortawesome.github.io/Font-Awesome/icon/times-circle/
-	static FText Times_Circle;
+	static UE_API FText Times_Circle;
 
 	// (garbage, delete, remove, hide) http://fortawesome.github.io/Font-Awesome/icon/trash/
-	static FText Trash;
+	static UE_API FText Trash;
 
 	// (back) http://fortawesome.github.io/Font-Awesome/icon/undo/
-	static FText Undo;
+	static UE_API FText Undo;
 
 	// (protect, admin, password, lock) http://fortawesome.github.io/Font-Awesome/icon/unlock/
-	static FText Unlock;
+	static UE_API FText Unlock;
 
 	// (film, movie, record) http://fortawesome.github.io/Font-Awesome/icon/video-camera/
-	static FText Video_Camera;
+	static UE_API FText Video_Camera;
 };
+
+#undef UE_API
diff --git a/Engine/Source/Editor/EditorStyle/Public/EditorStyleSet.h b/Engine/Source/Editor/EditorStyle/Public/EditorStyleSet.h
--- a/Engine/Source/Editor/EditorStyle/Public/EditorStyleSet.h
+++ b/Engine/Source/Editor/EditorStyle/Public/EditorStyleSet.h
@@ -20,6 +20,8 @@
 #include "Templates/SharedPointer.h"
 #include "UObject/NameTypes.h"
 
+#define UE_API EDITORSTYLE_API
+
 struct FLinearColor;
 struct FMargin;
 struct FSlateBrush;
@@ -29,7 +31,7 @@ struct FSlateSound;
 /**
  * A collection of named properties that guide the appearance of Slate.
  */
-class EDITORSTYLE_API FEditorStyle
+class FEditorStyle
 {
 public:
 
@@ -177,15 +179,17 @@ public:
 		}
 	}
 
-	static void ResetToDefault();
+	static UE_API void ResetToDefault();
 
 
 protected:
 
-	static void SetStyle( const TSharedRef< class ISlateStyle >& NewStyle );
+	static UE_API void SetStyle( const TSharedRef< class ISlateStyle >& NewStyle );
 
 private:
 
 	/** Singleton instance of the slate style */
-	static TSharedPtr< class ISlateStyle > Instance;
+	static UE_API TSharedPtr< class ISlateStyle > Instance;
 };
+
+#undef UE_API
diff --git a/Engine/Source/Editor/EditorSubsystem/Public/EditorSubsystem.h b/Engine/Source/Editor/EditorSubsystem/Public/EditorSubsystem.h
--- a/Engine/Source/Editor/EditorSubsystem/Public/EditorSubsystem.h
+++ b/Engine/Source/Editor/EditorSubsystem/Public/EditorSubsystem.h
@@ -8,6 +8,8 @@
 
 #include "EditorSubsystem.generated.h"
 
+#define UE_API EDITORSUBSYSTEM_API
+
 class UObject;
 
 /**
@@ -21,12 +23,13 @@ class UObject;
  * collection initialization time.
  */
 
-UCLASS(Abstract)
-class EDITORSUBSYSTEM_API UEditorSubsystem : public UDynamicSubsystem
+UCLASS(MinimalAPI, Abstract)
+class UEditorSubsystem : public UDynamicSubsystem
 {
 	GENERATED_BODY()
 
 public:
-	UEditorSubsystem();
+	UE_API UEditorSubsystem();
 };
 
+#undef UE_API
diff --git a/Engine/Source/Editor/EditorSubsystem/Public/EditorSubsystemModule.h b/Engine/Source/Editor/EditorSubsystem/Public/EditorSubsystemModule.h
--- a/Engine/Source/Editor/EditorSubsystem/Public/EditorSubsystemModule.h
+++ b/Engine/Source/Editor/EditorSubsystem/Public/EditorSubsystemModule.h
@@ -5,12 +5,16 @@
 #include "CoreMinimal.h"
 #include "Modules/ModuleInterface.h"
 
-class EDITORSUBSYSTEM_API FEditorSubsystemModule : public IModuleInterface
+#define UE_API EDITORSUBSYSTEM_API
+
+class FEditorSubsystemModule : public IModuleInterface
 {
 public:
 	// IModuleInterface
-	virtual void StartupModule() override;
-	virtual void ShutdownModule() override;
+	UE_API virtual void StartupModule() override;
+	UE_API virtual void ShutdownModule() override;
 
 private:
 };
+
+#undef UE_API
diff --git a/Engine/Source/Editor/EditorWidgets/Public/Filters/CustomClassFilterData.h b/Engine/Source/Editor/EditorWidgets/Public/Filters/CustomClassFilterData.h
--- a/Engine/Source/Editor/EditorWidgets/Public/Filters/CustomClassFilterData.h
+++ b/Engine/Source/Editor/EditorWidgets/Public/Filters/CustomClassFilterData.h
@@ -10,6 +10,8 @@
 #include "Templates/SharedPointer.h"
 #include "UObject/TopLevelAssetPath.h"
 
+#define UE_API EDITORWIDGETS_API
+
 class FFilterCategory;
 class IAssetTypeActions;
 class UAssetDefinition;
@@ -22,41 +24,41 @@ struct FARFilter;
  *  For example: filtering for StaticMeshActor and DirectionalLight will show assets that belong
  *  to either of those types.
  */
-class EDITORWIDGETS_API FCustomClassFilterData
+class FCustomClassFilterData
 {
 public:
 
 	/* You can provide an IAssetTypeActions for your type to get all the information from that */
-	FCustomClassFilterData(UAssetDefinition* InAssetDefinition, const FAssetFilterData& InAssetTypeAction);
+	UE_API FCustomClassFilterData(UAssetDefinition* InAssetDefinition, const FAssetFilterData& InAssetTypeAction);
 
 	/* Or you can provide a UClass, a Category and a Color to identify the filter */
-	FCustomClassFilterData(UClass* InClass, TSharedPtr<FFilterCategory> InCategory, FLinearColor InColor);
+	UE_API FCustomClassFilterData(UClass* InClass, TSharedPtr<FFilterCategory> InCategory, FLinearColor InColor);
 
 	/* Unlike normal filters, Type Filters are allowed to belong to multiple categories */
-	void AddCategory(TSharedPtr<FFilterCategory> InCategory);
+	UE_API void AddCategory(TSharedPtr<FFilterCategory> InCategory);
 	
 	/* Get the UClass associated with this filter */
-	UClass* GetClass() const;
+	UE_API UClass* GetClass() const;
 
 	/* Get the Catgories this filter belongs to */
-	TArray<TSharedPtr<FFilterCategory>> GetCategories() const;
+	UE_API TArray<TSharedPtr<FFilterCategory>> GetCategories() const;
 
 	/* Get the color of this filter */
-	FLinearColor GetColor() const;
+	UE_API FLinearColor GetColor() const;
 
 	/** Add this filter to the input BackendFilter */
-	void BuildBackendFilter(FARFilter& OutFilter);
+	UE_API void BuildBackendFilter(FARFilter& OutFilter);
 
 	/** Get the logical name of this filter class */
-	FText GetName() const;
+	UE_API FText GetName() const;
 
 	/** Get the display name of this filter class */
-	FString GetFilterName() const;
+	UE_API FString GetFilterName() const;
 
 	/**
 	 * Returns class path name as a package + class FName pair
 	 */
-	FTopLevelAssetPath GetClassPathName() const;
+	UE_API FTopLevelAssetPath GetClassPathName() const;
 
 private:
 
@@ -83,4 +85,6 @@ private:
 
 	/** An optional AssetDefinition to get information about the filter class */
 	TWeakObjectPtr<UAssetDefinition> AssetDefinitionPtr;
-};
\ No newline at end of file
+};
+
+#undef UE_API
diff --git a/Engine/Source/Editor/EditorWidgets/Public/Filters/FilterBarConfig.h b/Engine/Source/Editor/EditorWidgets/Public/Filters/FilterBarConfig.h
--- a/Engine/Source/Editor/EditorWidgets/Public/Filters/FilterBarConfig.h
+++ b/Engine/Source/Editor/EditorWidgets/Public/Filters/FilterBarConfig.h
@@ -14,6 +14,8 @@
 
 #include "FilterBarConfig.generated.h"
 
+#define UE_API EDITORWIDGETS_API
+
 class FName;
 class FString;
 class UObject;
@@ -75,14 +77,14 @@ public:
 	}
 };
 
-UCLASS(EditorConfig="FilterBar")
-class EDITORWIDGETS_API UFilterBarConfig : public UEditorConfigBase
+UCLASS(MinimalAPI, EditorConfig="FilterBar")
+class UFilterBarConfig : public UEditorConfigBase
 {
 	GENERATED_BODY()
 	
 public:
 
-	static void Initialize();
+	static UE_API void Initialize();
 	static UFilterBarConfig* Get() { return Instance; }
 	
 	UPROPERTY(meta=(EditorConfig))
@@ -90,5 +92,7 @@ public:
 	
 private:
 
-	static TObjectPtr<UFilterBarConfig> Instance;
+	static UE_API TObjectPtr<UFilterBarConfig> Instance;
 };
+
+#undef UE_API
diff --git a/Engine/Source/Editor/EditorWidgets/Public/Filters/SAssetFilterBar.h b/Engine/Source/Editor/EditorWidgets/Public/Filters/SAssetFilterBar.h
--- a/Engine/Source/Editor/EditorWidgets/Public/Filters/SAssetFilterBar.h
+++ b/Engine/Source/Editor/EditorWidgets/Public/Filters/SAssetFilterBar.h
@@ -59,6 +59,8 @@
 
 #include "SAssetFilterBar.generated.h"
 
+#define UE_API EDITORWIDGETS_API
+
 class SFilterSearchBox;
 class SWidget;
 
@@ -75,7 +77,7 @@ DECLARE_DELEGATE_ThreeParams(FOnPopulateAddAssetFilterMenu, UToolMenu*, TSharedP
 namespace UE::Editor::Widgets
 {
 	/** A local struct to describe a category in the filter menu */
-	struct EDITORWIDGETS_API FFilterCategoryMenu
+	struct FFilterCategoryMenu
 	{
 		/** The Classes that belong to this category */
 		TArray<TSharedPtr<FCustomClassFilterData>> Classes;
@@ -84,7 +86,7 @@ namespace UE::Editor::Widgets
 		FName SectionExtensionHook;
 		FText SectionHeading;
 
-		FFilterCategoryMenu(const FName& InSectionExtensionHook, const FText& InSectionHeading);
+		UE_API FFilterCategoryMenu(const FName& InSectionExtensionHook, const FText& InSectionHeading);
 	};
 
 	EDITORWIDGETS_API TMap<TSharedPtr<FFilterCategory>, FFilterCategoryMenu> BuildCategoryToMenuMap(
@@ -94,8 +96,8 @@ namespace UE::Editor::Widgets
 }
 
 /** ToolMenuContext that is used to create the Add Filter Menu */
-UCLASS()
-class EDITORWIDGETS_API UAssetFilterBarContext : public UObject
+UCLASS(MinimalAPI)
+class UAssetFilterBarContext : public UObject
 {
 	GENERATED_BODY()
 
@@ -109,21 +111,21 @@ public:
 /** A non-templated base class for SAssetFilterBar, where functionality that does not depend on the type of item
  *  being filtered lives
  */
-class EDITORWIDGETS_API FFilterBarBase
+class FFilterBarBase
 {
 protected:
 
 	/** Get a mutable version of this filter bar's config */
-	FFilterBarSettings* GetMutableConfig();
+	UE_API FFilterBarSettings* GetMutableConfig();
 
 	/** Get a const version of this filter bar's config */
-	const FFilterBarSettings* GetConstConfig() const;
+	UE_API const FFilterBarSettings* GetConstConfig() const;
 
 	/** Save this filte bar's data to the config file */
-	void SaveConfig();
+	UE_API void SaveConfig();
 
 	/** Initialize and load this filter bar's config */
-	void InitializeConfig();
+	UE_API void InitializeConfig();
 
 protected:
 
@@ -1283,3 +1285,5 @@ SAssetFilterBar<FilterType>::~SAssetFilterBar()
 }
 
 #undef LOCTEXT_NAMESPACE
+
+#undef UE_API
diff --git a/Engine/Source/Editor/EditorWidgets/Public/ObjectNameEditSinkRegistry.h b/Engine/Source/Editor/EditorWidgets/Public/ObjectNameEditSinkRegistry.h
--- a/Engine/Source/Editor/EditorWidgets/Public/ObjectNameEditSinkRegistry.h
+++ b/Engine/Source/Editor/EditorWidgets/Public/ObjectNameEditSinkRegistry.h
@@ -6,25 +6,27 @@
 #include "IObjectNameEditSink.h"
 #include "Templates/SharedPointer.h"
 
+#define UE_API EDITORWIDGETS_API
+
 class UClass;
 
 namespace UE::EditorWidgets
 {
 class IObjectNameEditSink;
 
-class EDITORWIDGETS_API FObjectNameEditSinkRegistry
+class FObjectNameEditSinkRegistry
 {
 public:
-	FObjectNameEditSinkRegistry();
+	UE_API FObjectNameEditSinkRegistry();
 
 	/** Registers an object name edit sink so it can provide a name editing interface for an object type. */
-	void RegisterObjectNameEditSink(const TSharedRef<IObjectNameEditSink>& NewSink);
+	UE_API void RegisterObjectNameEditSink(const TSharedRef<IObjectNameEditSink>& NewSink);
 
 	/** Unregisters a name edit sink. It will no longer provide a name editing interface for an object type. */
-	void UnregisterObjectNameEditSink(const TSharedRef<IObjectNameEditSink>& SinkToRemove);
+	UE_API void UnregisterObjectNameEditSink(const TSharedRef<IObjectNameEditSink>& SinkToRemove);
 
 	/** Gets the appropriate ObjectNameEditSink for the supplied class */
-	TSharedPtr<IObjectNameEditSink> GetObjectNameEditSinkForClass(const UClass* Class) const;
+	UE_API TSharedPtr<IObjectNameEditSink> GetObjectNameEditSinkForClass(const UClass* Class) const;
 
 private:
 	/** The list of all registered ObjectNameEditSinks */
@@ -32,3 +34,5 @@ private:
 };
 
 } // end namespace UE::EditorWidgets
+
+#undef UE_API
diff --git a/Engine/Source/Editor/EditorWidgets/Public/SAssetDropTarget.h b/Engine/Source/Editor/EditorWidgets/Public/SAssetDropTarget.h
--- a/Engine/Source/Editor/EditorWidgets/Public/SAssetDropTarget.h
+++ b/Engine/Source/Editor/EditorWidgets/Public/SAssetDropTarget.h
@@ -14,6 +14,8 @@
 #include "Templates/SharedPointer.h"
 #include "Widgets/DeclarativeSyntaxSupport.h"
 
+#define UE_API EDITORWIDGETS_API
+
 class FDragDropEvent;
 class FDragDropOperation;
 class SWidget;
@@ -23,7 +25,7 @@ struct FGeometry;
 /**
  * A widget that displays a hover cue and handles dropping assets of allowed types onto this widget
  */
-class EDITORWIDGETS_API SAssetDropTarget : public SDropTarget
+class SAssetDropTarget : public SDropTarget
 {
 public:
 	/** Called when a valid asset is dropped */
@@ -102,16 +104,16 @@ public:
 		SLATE_EVENT_DEPRECATED(5.0, "Use FAreAssetsAcceptableForDropWithReason instead.", FIsAssetAcceptableForDropWithReason, OnIsAssetAcceptableForDropWithReason, OnAreAssetsAcceptableForDropWithReason, ConvertObjectAcceptableDropWithReasonDelegate)
 	SLATE_END_ARGS()
 
-	void Construct(const FArguments& InArgs );
+	UE_API void Construct(const FArguments& InArgs );
 
 protected:
-	FReply OnDropped(const FGeometry& InGeometry, const FDragDropEvent& InDragDropEvent);
-	virtual bool OnAllowDrop(TSharedPtr<FDragDropOperation> DragDropOperation) const override;
-	virtual bool OnIsRecognized(TSharedPtr<FDragDropOperation> DragDropOperation) const override;
-	virtual void OnDragLeave(const FDragDropEvent& DragDropEvent) override;
+	UE_API FReply OnDropped(const FGeometry& InGeometry, const FDragDropEvent& InDragDropEvent);
+	UE_API virtual bool OnAllowDrop(TSharedPtr<FDragDropOperation> DragDropOperation) const override;
+	UE_API virtual bool OnIsRecognized(TSharedPtr<FDragDropOperation> DragDropOperation) const override;
+	UE_API virtual void OnDragLeave(const FDragDropEvent& DragDropEvent) override;
 
 private:
-	TArray<FAssetData> GetDroppedAssets(TSharedPtr<FDragDropOperation> DragDropOperation, bool& bOutRecognizedEvent) const;
+	UE_API TArray<FAssetData> GetDroppedAssets(TSharedPtr<FDragDropOperation> DragDropOperation, bool& bOutRecognizedEvent) const;
 
 private:
 	/** Delegate to call when an asset is dropped */
@@ -123,3 +125,5 @@ private:
 	/** Boolean to determine if this drop target can support multiple assets dropped, or only supports a single asset dropped at a time. False by default for legacy behavior. */
 	bool bSupportsMultiDrop;
 };
+
+#undef UE_API
diff --git a/Engine/Source/Editor/EditorWidgets/Public/SAssetSearchBox.h b/Engine/Source/Editor/EditorWidgets/Public/SAssetSearchBox.h
--- a/Engine/Source/Editor/EditorWidgets/Public/SAssetSearchBox.h
+++ b/Engine/Source/Editor/EditorWidgets/Public/SAssetSearchBox.h
@@ -27,6 +27,8 @@
 #include "Widgets/Views/STableRow.h"
 #include "Widgets/Views/STableViewBase.h"
 
+#define UE_API EDITORWIDGETS_API
+
 class ITableRow;
 class SMenuAnchor;
 class STableViewBase;
@@ -67,7 +69,7 @@ DECLARE_DELEGATE_RetVal_TwoParams(FText, FOnAssetSearchBoxSuggestionChosen, cons
  * SetOnSaveSearchHandler
  */
 
-class EDITORWIDGETS_API SAssetSearchBox : public SCompoundWidget
+class SAssetSearchBox : public SCompoundWidget
 {
 public:
 	SLATE_BEGIN_ARGS(SAssetSearchBox)
@@ -124,23 +126,23 @@ public:
 	SLATE_END_ARGS()
 
 	/** Constructs this widget with InArgs */
-	void Construct( const FArguments& InArgs );
+	UE_API void Construct( const FArguments& InArgs );
 
 	/** Sets the text string currently being edited */
-	void SetText(const TAttribute< FText >& InNewText);
+	UE_API void SetText(const TAttribute< FText >& InNewText);
 
 	/** Set or clear the current error reporting information for this search box */
-	void SetError( const FText& InError );
-	void SetError( const FString& InError );
+	UE_API void SetError( const FText& InError );
+	UE_API void SetError( const FString& InError );
 
 	/** Show a + button next to the current search and set the handler for when that is clicked */
-	void SetOnSaveSearchHandler(SFilterSearchBox::FOnSaveSearchClicked InOnSaveSearchHandler);
+	UE_API void SetOnSaveSearchHandler(SFilterSearchBox::FOnSaveSearchClicked InOnSaveSearchHandler);
 	
 	// SWidget implementation
-	virtual FReply OnPreviewKeyDown( const FGeometry& MyGeometry, const FKeyEvent& InKeyEvent ) override;
-	virtual bool SupportsKeyboardFocus() const override;
-	virtual bool HasKeyboardFocus() const override;
-	virtual FReply OnFocusReceived( const FGeometry& MyGeometry, const FFocusEvent& InFocusEvent ) override;
+	UE_API virtual FReply OnPreviewKeyDown( const FGeometry& MyGeometry, const FKeyEvent& InKeyEvent ) override;
+	UE_API virtual bool SupportsKeyboardFocus() const override;
+	UE_API virtual bool HasKeyboardFocus() const override;
+	UE_API virtual FReply OnFocusReceived( const FGeometry& MyGeometry, const FFocusEvent& InFocusEvent ) override;
 
 private:
 	struct FSuggestionListEntry
@@ -151,37 +153,37 @@ private:
 	};
 
 	/** First chance handler for key down events to the editable text widget */
-	FReply HandleKeyDown(const FGeometry& MyGeometry, const FKeyEvent& InKeyEvent);
+	UE_API FReply HandleKeyDown(const FGeometry& MyGeometry, const FKeyEvent& InKeyEvent);
 
 	/** Handler for when text in the editable text box changed */
-	void HandleTextChanged(const FText& NewText);
+	UE_API void HandleTextChanged(const FText& NewText);
 
 	/** Handler for when text in the editable text box changed */
-	void HandleTextCommitted(const FText& NewText, ETextCommit::Type CommitType);
+	UE_API void HandleTextCommitted(const FText& NewText, ETextCommit::Type CommitType);
 
 	/** Called by SListView when the selection changes in the suggestion list */
-	void OnSelectionChanged( TSharedPtr<FSuggestionListEntry> NewValue, ESelectInfo::Type SelectInfo );
+	UE_API void OnSelectionChanged( TSharedPtr<FSuggestionListEntry> NewValue, ESelectInfo::Type SelectInfo );
 
 	/** Makes the widget for a suggestion message in the list view */
-	TSharedRef<ITableRow> MakeSuggestionListItemWidget(TSharedPtr<FSuggestionListEntry> Suggestion, const TSharedRef<STableViewBase>& OwnerTable);
+	UE_API TSharedRef<ITableRow> MakeSuggestionListItemWidget(TSharedPtr<FSuggestionListEntry> Suggestion, const TSharedRef<STableViewBase>& OwnerTable);
 
 	/** Gets the text to highlight in the suggestion list */
-	FText GetHighlightText() const;
+	UE_API FText GetHighlightText() const;
 
 	/** Updates and shows or hides the suggestion list */
-	void UpdateSuggestionList();
+	UE_API void UpdateSuggestionList();
 
 	/** Sets the focus to the InputText box */
-	void FocusEditBox();
+	UE_API void FocusEditBox();
 
 	/** Returns the currently selected suggestion */
-	TSharedPtr<FSuggestionListEntry> GetSelectedSuggestion() const;
+	UE_API TSharedPtr<FSuggestionListEntry> GetSelectedSuggestion() const;
 
 	/** Default implementation of OnAssetSearchBoxSuggestionFilter, if no external implementation is provided */
-	static void DefaultSuggestionFilterImpl(const FText& SearchText, TArray<FAssetSearchBoxSuggestion>& PossibleSuggestions, FText& SuggestionHighlightText);
+	static UE_API void DefaultSuggestionFilterImpl(const FText& SearchText, TArray<FAssetSearchBoxSuggestion>& PossibleSuggestions, FText& SuggestionHighlightText);
 
 	/** Default implementation of OnAssetSearchBoxSuggestionChosen, if no external implementation is provided */
-	static FText DefaultSuggestionChosenImpl(const FText& SearchText, const FString& Suggestion);
+	static UE_API FText DefaultSuggestionChosenImpl(const FText& SearchText, const FString& Suggestion);
 
 private:
 	/** The editable text field */
@@ -223,3 +225,5 @@ private:
 	/** Determines whether or not the committed text should match a suggestion */
 	bool bMustMatchPossibleSuggestions;
 };
+
+#undef UE_API
diff --git a/Engine/Source/Editor/EditorWidgets/Public/SDropTarget.h b/Engine/Source/Editor/EditorWidgets/Public/SDropTarget.h
--- a/Engine/Source/Editor/EditorWidgets/Public/SDropTarget.h
+++ b/Engine/Source/Editor/EditorWidgets/Public/SDropTarget.h
@@ -18,6 +18,8 @@
 #include "Widgets/DeclarativeSyntaxSupport.h"
 #include "Widgets/SCompoundWidget.h"
 
+#define UE_API EDITORWIDGETS_API
+
 class FPaintArgs;
 class FSlateRect;
 class FSlateWindowElementList;
@@ -29,7 +31,7 @@ struct FSlateBrush;
 /** 
  * A widget that displays a hover cue and handles dropping assets of allowed types onto this widget
  */
-class EDITORWIDGETS_API SDropTarget : public SCompoundWidget
+class SDropTarget : public SCompoundWidget
 {
 public:
 	DECLARE_DELEGATE_RetVal_OneParam(bool, FVerifyDrag, TSharedPtr<FDragDropOperation>);
@@ -97,7 +99,7 @@ public:
 		SLATE_EVENT_DEPRECATED(5.0, "Use OnDropped instead.", FOnDropDeprecated, OnDrop, OnDropped, ConvertOnDropFn)
 	SLATE_END_ARGS()
 
-	void Construct(const FArguments& InArgs );
+	UE_API void Construct(const FArguments& InArgs );
 
 	void ClearAllowDropCache()
 	{
@@ -106,26 +108,26 @@ public:
 	
 protected:
 
-	bool AllowDrop(TSharedPtr<FDragDropOperation> DragDropOperation) const;
+	UE_API bool AllowDrop(TSharedPtr<FDragDropOperation> DragDropOperation) const;
 
-	virtual bool OnAllowDrop(TSharedPtr<FDragDropOperation> DragDropOperation) const;
-	virtual bool OnIsRecognized(TSharedPtr<FDragDropOperation> DragDropOperation) const;
+	UE_API virtual bool OnAllowDrop(TSharedPtr<FDragDropOperation> DragDropOperation) const;
+	UE_API virtual bool OnIsRecognized(TSharedPtr<FDragDropOperation> DragDropOperation) const;
 	
-	virtual void Tick(const FGeometry& AllottedGeometry, const double InCurrentTime, const float InDeltaTime) override;
+	UE_API virtual void Tick(const FGeometry& AllottedGeometry, const double InCurrentTime, const float InDeltaTime) override;
 protected:
 	// SWidget interface
-	virtual FReply OnDragOver( const FGeometry& MyGeometry, const FDragDropEvent& DragDropEvent ) override;
-	virtual FReply OnDrop( const FGeometry& MyGeometry, const FDragDropEvent& DragDropEvent ) override;
-	virtual void OnDragEnter( const FGeometry& MyGeometry, const FDragDropEvent& DragDropEvent ) override;
-	virtual void OnDragLeave( const FDragDropEvent& DragDropEvent ) override;
-	virtual int32 OnPaint(const FPaintArgs& Args, const FGeometry& AllottedGeometry, const FSlateRect& MyCullingRect, FSlateWindowElementList& OutDrawElements, int32 LayerId, const FWidgetStyle& InWidgetStyle, bool bParentEnabled) const override;
+	UE_API virtual FReply OnDragOver( const FGeometry& MyGeometry, const FDragDropEvent& DragDropEvent ) override;
+	UE_API virtual FReply OnDrop( const FGeometry& MyGeometry, const FDragDropEvent& DragDropEvent ) override;
+	UE_API virtual void OnDragEnter( const FGeometry& MyGeometry, const FDragDropEvent& DragDropEvent ) override;
+	UE_API virtual void OnDragLeave( const FDragDropEvent& DragDropEvent ) override;
+	UE_API virtual int32 OnPaint(const FPaintArgs& Args, const FGeometry& AllottedGeometry, const FSlateRect& MyCullingRect, FSlateWindowElementList& OutDrawElements, int32 LayerId, const FWidgetStyle& InWidgetStyle, bool bParentEnabled) const override;
 	// End of SWidget interface
 
 	/** @return Visibility of the overlay text when dragging is occurring. */
-	EVisibility GetDragOverlayVisibility() const;
+	UE_API EVisibility GetDragOverlayVisibility() const;
 
 	/** Get the brightness on the background. */
-	FSlateColor GetBackgroundBrightness() const;
+	UE_API FSlateColor GetBackgroundBrightness() const;
 
 	/** Returns true if this drop target is currently being hovered over by a drag drop event */
 	bool IsDragOver() const { return bIsDragOver; }
@@ -163,3 +165,5 @@ private:
 	bool bIsDragDropping = false;
 	bool bWasDragDroppingLastFrame = false;
 };
+
+#undef UE_API
diff --git a/Engine/Source/Editor/EditorWidgets/Public/SEnumCombo.h b/Engine/Source/Editor/EditorWidgets/Public/SEnumCombo.h
--- a/Engine/Source/Editor/EditorWidgets/Public/SEnumCombo.h
+++ b/Engine/Source/Editor/EditorWidgets/Public/SEnumCombo.h
@@ -15,11 +15,13 @@
 #include "Widgets/DeclarativeSyntaxSupport.h"
 #include "Widgets/Input/SComboButton.h"
 
+#define UE_API EDITORWIDGETS_API
+
 class SWidget;
 class UEnum;
 struct FButtonStyle;
 
-class EDITORWIDGETS_API SEnumComboBox : public SComboButton
+class SEnumComboBox : public SComboButton
 {
 public:
 	DECLARE_DELEGATE_TwoParams(FOnEnumSelectionChanged, int32, ESelectInfo::Type);
@@ -48,12 +50,12 @@ public:
 		SLATE_ARGUMENT(bool, bForceBitFlags)
 	SLATE_END_ARGS()
 
-	void Construct(const FArguments& InArgs, const UEnum* InEnum);
+	UE_API void Construct(const FArguments& InArgs, const UEnum* InEnum);
 
 private:
-	FText GetCurrentValueText() const;
-	FText GetCurrentValueTooltip() const;
-	TSharedRef<SWidget> OnGetMenuContent();
+	UE_API FText GetCurrentValueText() const;
+	UE_API FText GetCurrentValueTooltip() const;
+	UE_API TSharedRef<SWidget> OnGetMenuContent();
 
 private:
 
@@ -87,3 +89,5 @@ private:
 
 	FOnEnumSelectionChanged OnEnumSelectionChangedDelegate;
 };
+
+#undef UE_API
diff --git a/Engine/Source/Editor/EditorWidgets/Public/SInputChord.h b/Engine/Source/Editor/EditorWidgets/Public/SInputChord.h
--- a/Engine/Source/Editor/EditorWidgets/Public/SInputChord.h
+++ b/Engine/Source/Editor/EditorWidgets/Public/SInputChord.h
@@ -8,11 +8,13 @@
 #include "Styling/SlateWidgetStyleAsset.h"
 #include "Widgets/SCompoundWidget.h"
 
+#define UE_API EDITORWIDGETS_API
+
 /**
  * Displays a given input chord in a user-friendly way.
  * ie. "[Keyboard Icon] Ctrl + C | Copy"
  */
-class EDITORWIDGETS_API SInputChord : public SCompoundWidget
+class SInputChord : public SCompoundWidget
 {
 public:
 	SLATE_BEGIN_ARGS(SInputChord)
@@ -34,14 +36,16 @@ public:
 	SLATE_END_ARGS()
 
 	/** Constructs this widget with InArgs */
-	void Construct(const FArguments& InArgs);
+	UE_API void Construct(const FArguments& InArgs);
 
 	/** Create an SInputChord widget for the given InputChord and (optional) ActionLabel. The InputLabel (ie. "Ctrl + C") is derived from the InputChord. */
-	static TSharedRef<SInputChord> MakeForInputChord(const FInputChord& InInputChord, const FText& InActionLabel = { });
+	static UE_API TSharedRef<SInputChord> MakeForInputChord(const FInputChord& InInputChord, const FText& InActionLabel = { });
 
 	/** Create an SInputChord widget for the given CommandInfo. Optionally, provide a KeyBindingIndex for commands with multiple bindings. */
-	static TSharedRef<SInputChord> MakeForCommandInfo(const FUICommandInfo& InCommandInfo, const EMultipleKeyBindingIndex InKeyBindingIndex = EMultipleKeyBindingIndex::Primary);
+	static UE_API TSharedRef<SInputChord> MakeForCommandInfo(const FUICommandInfo& InCommandInfo, const EMultipleKeyBindingIndex InKeyBindingIndex = EMultipleKeyBindingIndex::Primary);
 
 private:
-	static const FSlateBrush* GetIconForInputChord(const FInputChord& InInputChord);
+	static UE_API const FSlateBrush* GetIconForInputChord(const FInputChord& InInputChord);
 };
+
+#undef UE_API
diff --git a/Engine/Source/Editor/EditorWidgets/Public/SMetaDataView.h b/Engine/Source/Editor/EditorWidgets/Public/SMetaDataView.h
--- a/Engine/Source/Editor/EditorWidgets/Public/SMetaDataView.h
+++ b/Engine/Source/Editor/EditorWidgets/Public/SMetaDataView.h
@@ -9,6 +9,8 @@
 #include "Widgets/DeclarativeSyntaxSupport.h"
 #include "Widgets/SCompoundWidget.h"
 
+#define UE_API EDITORWIDGETS_API
+
 class FName;
 class FString;
 class ITableRow;
@@ -18,7 +20,7 @@ struct FMetaDataLine;
 /**
  * The widget to display metadata as a table of tag/value rows
  */
-class EDITORWIDGETS_API SMetaDataView : public SCompoundWidget
+class SMetaDataView : public SCompoundWidget
 {
 public:
 	SLATE_BEGIN_ARGS(SMetaDataView)	{}
@@ -30,10 +32,12 @@ public:
 	 * @param	InArgs				Declaration used by the SNew() macro to construct this widget
 	 * @param	InMetaData			The metadata tags/values to display in the table view widget
 	 */
-	void Construct(const FArguments& InArgs, const TMap<FName, FString>& InMetadata);
+	UE_API void Construct(const FArguments& InArgs, const TMap<FName, FString>& InMetadata);
 
 private:
 	TArray< TSharedPtr< FMetaDataLine > > MetaDataLines;
 
-	TSharedRef< ITableRow > OnGenerateRow(const TSharedPtr< FMetaDataLine > Item, const TSharedRef< STableViewBase >& OwnerTable);
+	UE_API TSharedRef< ITableRow > OnGenerateRow(const TSharedPtr< FMetaDataLine > Item, const TSharedRef< STableViewBase >& OwnerTable);
 };
+
+#undef UE_API
diff --git a/Engine/Source/Editor/EditorWidgets/Public/STextPropertyEditableTextBox.h b/Engine/Source/Editor/EditorWidgets/Public/STextPropertyEditableTextBox.h
--- a/Engine/Source/Editor/EditorWidgets/Public/STextPropertyEditableTextBox.h
+++ b/Engine/Source/Editor/EditorWidgets/Public/STextPropertyEditableTextBox.h
@@ -36,6 +36,8 @@
 #include "Widgets/SWidget.h"
 #include "Widgets/Views/SListView.h"
 
+#define UE_API EDITORWIDGETS_API
+
 class SComboButton;
 class SEditableTextBox;
 class SMultiLineEditableTextBox;
@@ -103,7 +105,7 @@ protected:
 };
 
 /** A widget that can be used for editing the string table referenced for FText instances */
-class EDITORWIDGETS_API STextPropertyEditableStringTableReference : public SCompoundWidget
+class STextPropertyEditableStringTableReference : public SCompoundWidget
 {
 	SLATE_BEGIN_ARGS(STextPropertyEditableStringTableReference)
 		: _ComboStyle(&FCoreStyle::Get().GetWidgetStyle<FComboBoxStyle>("ComboBox"))
@@ -121,7 +123,7 @@ class EDITORWIDGETS_API STextPropertyEditableStringTableReference : public SComp
 	SLATE_END_ARGS()
 
 public:
-	void Construct(const FArguments& Arguments, const TSharedRef<IEditableTextProperty>& InEditableTextProperty);
+	UE_API void Construct(const FArguments& Arguments, const TSharedRef<IEditableTextProperty>& InEditableTextProperty);
 
 private:
 	struct FAvailableStringTable
@@ -130,29 +132,29 @@ private:
 		FText DisplayName;
 	};
 
-	void GetTableIdAndKey(FName& OutTableId, FString& OutKey) const;
-	void SetTableIdAndKey(const FName InTableId, const FString& InKey);
+	UE_API void GetTableIdAndKey(FName& OutTableId, FString& OutKey) const;
+	UE_API void SetTableIdAndKey(const FName InTableId, const FString& InKey);
 
-	void OnOptionsFilterTextChanged(const FText& InNewText);
-	void OnKeysFilterTextChanged(const FText& InNewText);
+	UE_API void OnOptionsFilterTextChanged(const FText& InNewText);
+	UE_API void OnKeysFilterTextChanged(const FText& InNewText);
 
-	TSharedRef<SWidget> OnGetStringTableComboOptions();
-	TSharedRef<class ITableRow> OnGenerateStringTableComboOption(TSharedPtr<FAvailableStringTable> InItem, const TSharedRef<class STableViewBase>& OwnerTable);
-	TSharedRef<SWidget> OnGetStringTableKeyOptions();
-	TSharedRef<class ITableRow> OnGenerateStringTableKeyOption(TSharedPtr<FString> InItem, const TSharedRef<class STableViewBase>& OwnerTable);
+	UE_API TSharedRef<SWidget> OnGetStringTableComboOptions();
+	UE_API TSharedRef<class ITableRow> OnGenerateStringTableComboOption(TSharedPtr<FAvailableStringTable> InItem, const TSharedRef<class STableViewBase>& OwnerTable);
+	UE_API TSharedRef<SWidget> OnGetStringTableKeyOptions();
+	UE_API TSharedRef<class ITableRow> OnGenerateStringTableKeyOption(TSharedPtr<FString> InItem, const TSharedRef<class STableViewBase>& OwnerTable);
 
-	void OnStringTableComboChanged(TSharedPtr<FAvailableStringTable> NewSelection, ESelectInfo::Type SelectInfo);
-	void UpdateStringTableComboOptions();
-	FText GetStringTableComboContent() const;
-	FText GetStringTableComboToolTip() const;
+	UE_API void OnStringTableComboChanged(TSharedPtr<FAvailableStringTable> NewSelection, ESelectInfo::Type SelectInfo);
+	UE_API void UpdateStringTableComboOptions();
+	UE_API FText GetStringTableComboContent() const;
+	UE_API FText GetStringTableComboToolTip() const;
 
-	void OnKeyComboChanged(TSharedPtr<FString> NewSelection, ESelectInfo::Type SelectInfo);
-	void UpdateStringTableKeyOptions();
-	FText GetKeyComboContent() const;
-	FText GetKeyComboToolTip() const;
+	UE_API void OnKeyComboChanged(TSharedPtr<FString> NewSelection, ESelectInfo::Type SelectInfo);
+	UE_API void UpdateStringTableKeyOptions();
+	UE_API FText GetKeyComboContent() const;
+	UE_API FText GetKeyComboToolTip() const;
 
-	bool IsUnlinkEnabled() const;
-	FReply OnUnlinkClicked();
+	UE_API bool IsUnlinkEnabled() const;
+	UE_API FReply OnUnlinkClicked();
 
 	TSharedPtr<IEditableTextProperty> EditableTextProperty;
 
@@ -174,7 +176,7 @@ private:
 };
 
 /** A widget that can be used for editing FText instances */
-class EDITORWIDGETS_API STextPropertyEditableTextBox : public SCompoundWidget
+class STextPropertyEditableTextBox : public SCompoundWidget
 {
 	SLATE_BEGIN_ARGS(STextPropertyEditableTextBox)
 		: _Style(&FCoreStyle::Get().GetWidgetStyle<FEditableTextBoxStyle>("NormalEditableTextBox"))
@@ -202,44 +204,44 @@ class EDITORWIDGETS_API STextPropertyEditableTextBox : public SCompoundWidget
 	SLATE_END_ARGS()
 
 public:
-	void Construct(const FArguments& Arguments, const TSharedRef<IEditableTextProperty>& InEditableTextProperty);
-	virtual bool SupportsKeyboardFocus() const override;
-	virtual FReply OnFocusReceived(const FGeometry& MyGeometry, const FFocusEvent& InFocusEvent) override;
+	UE_API void Construct(const FArguments& Arguments, const TSharedRef<IEditableTextProperty>& InEditableTextProperty);
+	UE_API virtual bool SupportsKeyboardFocus() const override;
+	UE_API virtual FReply OnFocusReceived(const FGeometry& MyGeometry, const FFocusEvent& InFocusEvent) override;
 
 private:
-	void GetDesiredWidth(float& OutMinDesiredWidth, float& OutMaxDesiredWidth);
-	bool CanEdit() const;
-	bool IsCultureInvariantFlagEnabled() const;
-	bool IsSourceTextReadOnly() const;
-	bool IsIdentityReadOnly() const;
-	FText GetToolTipText() const;
-	bool IsTextLocalizable() const;
-
-	FText GetTextValue() const;
-	void OnTextChanged(const FText& NewText);
-	void OnTextCommitted(const FText& NewText, ETextCommit::Type CommitInfo);
-	void SetTextError(const FText& InErrorMsg);
-
-	FText GetNamespaceValue() const;
-	void OnNamespaceChanged(const FText& NewText);
-	void OnNamespaceCommitted(const FText& NewText, ETextCommit::Type CommitInfo);
-
-	FText GetKeyValue() const;
+	UE_API void GetDesiredWidth(float& OutMinDesiredWidth, float& OutMaxDesiredWidth);
+	UE_API bool CanEdit() const;
+	UE_API bool IsCultureInvariantFlagEnabled() const;
+	UE_API bool IsSourceTextReadOnly() const;
+	UE_API bool IsIdentityReadOnly() const;
+	UE_API FText GetToolTipText() const;
+	UE_API bool IsTextLocalizable() const;
+
+	UE_API FText GetTextValue() const;
+	UE_API void OnTextChanged(const FText& NewText);
+	UE_API void OnTextCommitted(const FText& NewText, ETextCommit::Type CommitInfo);
+	UE_API void SetTextError(const FText& InErrorMsg);
+
+	UE_API FText GetNamespaceValue() const;
+	UE_API void OnNamespaceChanged(const FText& NewText);
+	UE_API void OnNamespaceCommitted(const FText& NewText, ETextCommit::Type CommitInfo);
+
+	UE_API FText GetKeyValue() const;
 #if USE_STABLE_LOCALIZATION_KEYS
-	void OnKeyChanged(const FText& NewText);
-	void OnKeyCommitted(const FText& NewText, ETextCommit::Type CommitInfo);
+	UE_API void OnKeyChanged(const FText& NewText);
+	UE_API void OnKeyCommitted(const FText& NewText, ETextCommit::Type CommitInfo);
 
-	FText GetPackageValue() const;
+	UE_API FText GetPackageValue() const;
 #endif // USE_STABLE_LOCALIZATION_KEYS
 
-	ECheckBoxState GetLocalizableCheckState() const;
+	UE_API ECheckBoxState GetLocalizableCheckState() const;
 
-	void HandleLocalizableCheckStateChanged(ECheckBoxState InCheckboxState);
+	UE_API void HandleLocalizableCheckStateChanged(ECheckBoxState InCheckboxState);
 
-	FText GetAdvancedTextSettingsComboToolTip() const;
-	const FSlateBrush* GetAdvancedTextSettingsComboImage() const;
+	UE_API FText GetAdvancedTextSettingsComboToolTip() const;
+	UE_API const FSlateBrush* GetAdvancedTextSettingsComboImage() const;
 
-	bool IsValidIdentity(const FText& InIdentity, FText* OutReason = nullptr, const FText* InErrorCtx = nullptr) const;
+	UE_API bool IsValidIdentity(const FText& InIdentity, FText* OutReason = nullptr, const FText* InErrorCtx = nullptr) const;
 
 	TSharedPtr<IEditableTextProperty> EditableTextProperty;
 
@@ -255,5 +257,7 @@ private:
 
 	bool bIsMultiLine = false;
 
-	static FText MultipleValuesText;
+	static UE_API FText MultipleValuesText;
 };
+
+#undef UE_API
diff --git a/Engine/Source/Editor/Experimental/EditorInteractiveToolsFramework/Public/Behaviors/2DViewportBehaviorTargets.h b/Engine/Source/Editor/Experimental/EditorInteractiveToolsFramework/Public/Behaviors/2DViewportBehaviorTargets.h
--- a/Engine/Source/Editor/Experimental/EditorInteractiveToolsFramework/Public/Behaviors/2DViewportBehaviorTargets.h
+++ b/Engine/Source/Editor/Experimental/EditorInteractiveToolsFramework/Public/Behaviors/2DViewportBehaviorTargets.h
@@ -7,24 +7,26 @@
 #include "BaseBehaviors/BehaviorTargetInterfaces.h"
 #include "VectorTypes.h"
 
+#define UE_API EDITORINTERACTIVETOOLSFRAMEWORK_API
+
 // TODO: Replace with a context object to avoid dependency on FEditorViewportClient (JIRA UE-163677)
 class FEditorViewportClient;
 
 /**
  * Allows click-dragging to move the camera in the XY plane.
  */
-class EDITORINTERACTIVETOOLSFRAMEWORK_API FEditor2DScrollBehaviorTarget : public IClickDragBehaviorTarget
+class FEditor2DScrollBehaviorTarget : public IClickDragBehaviorTarget
 {
 public:
 
-	FEditor2DScrollBehaviorTarget(FEditorViewportClient* ViewportClientIn);
+	UE_API FEditor2DScrollBehaviorTarget(FEditorViewportClient* ViewportClientIn);
 
 	// IClickDragBehaviorTarget
-	virtual FInputRayHit CanBeginClickDragSequence(const FInputDeviceRay& PressPos) override;
-	virtual void OnClickPress(const FInputDeviceRay& PressPos) override;
-	virtual void OnClickDrag(const FInputDeviceRay& DragPos) override;
-	virtual void OnClickRelease(const FInputDeviceRay& ReleasePos) override;
-	virtual void OnTerminateDragSequence() override;
+	UE_API virtual FInputRayHit CanBeginClickDragSequence(const FInputDeviceRay& PressPos) override;
+	UE_API virtual void OnClickPress(const FInputDeviceRay& PressPos) override;
+	UE_API virtual void OnClickDrag(const FInputDeviceRay& DragPos) override;
+	UE_API virtual void OnClickRelease(const FInputDeviceRay& ReleasePos) override;
+	UE_API virtual void OnTerminateDragSequence() override;
 
 protected:
 
@@ -38,16 +40,16 @@ protected:
  * Allows the mouse wheel to move the camera forwards/backwards relative to the XY plane,
  * in the direction pointed to by the mouse.
  */
-class EDITORINTERACTIVETOOLSFRAMEWORK_API FEditor2DMouseWheelZoomBehaviorTarget : public IMouseWheelBehaviorTarget
+class FEditor2DMouseWheelZoomBehaviorTarget : public IMouseWheelBehaviorTarget
 {
 public:
 
-	FEditor2DMouseWheelZoomBehaviorTarget(FEditorViewportClient* ViewportClientIn);
+	UE_API FEditor2DMouseWheelZoomBehaviorTarget(FEditorViewportClient* ViewportClientIn);
 
 	// IMouseWheelBehaviorTarget
-	virtual FInputRayHit ShouldRespondToMouseWheel(const FInputDeviceRay& CurrentPos);
-	virtual void OnMouseWheelScrollUp(const FInputDeviceRay& CurrentPos);
-	virtual void OnMouseWheelScrollDown(const FInputDeviceRay& CurrentPos);
+	UE_API virtual FInputRayHit ShouldRespondToMouseWheel(const FInputDeviceRay& CurrentPos);
+	UE_API virtual void OnMouseWheelScrollUp(const FInputDeviceRay& CurrentPos);
+	UE_API virtual void OnMouseWheelScrollDown(const FInputDeviceRay& CurrentPos);
 	
 	/**
 	 * @param PercentZoomIn How much to move forward on each mouse wheel forward scroll. For instance,
@@ -56,11 +58,11 @@ public:
 	 *  zoom out amount will be 25% since zooming in by 20% and then out by 25% of the result gets you
 	 *  back to where you were.
 	 */
-	virtual void SetZoomAmount(double PercentZoomIn);
-	void SetZoomLimits(double ZoomInLimitIn, double ZoomOutLimitIn);
+	UE_API virtual void SetZoomAmount(double PercentZoomIn);
+	UE_API void SetZoomLimits(double ZoomInLimitIn, double ZoomOutLimitIn);
 
-	void SetCameraFarPlaneWorldZ(double CameraFarPlaneWorldZIn);
-	void SetCameraNearPlaneProportionZ(double CameraFarPlaneProportionZIn );
+	UE_API void SetCameraFarPlaneWorldZ(double CameraFarPlaneWorldZIn);
+	UE_API void SetCameraNearPlaneProportionZ(double CameraFarPlaneProportionZIn );
 
 	inline static const double DEFAULT_ZOOM_AMOUNT = 20;
 protected:
@@ -73,4 +75,6 @@ protected:
 	double ZoomOutLimit;
 	double CameraFarPlaneWorldZ;
 	double CameraNearPlaneProportionZ;
-};
\ No newline at end of file
+};
+
+#undef UE_API
diff --git a/Engine/Source/Editor/Experimental/EditorInteractiveToolsFramework/Public/EditorDragTools/DragToolsBehaviourSource.h b/Engine/Source/Editor/Experimental/EditorInteractiveToolsFramework/Public/EditorDragTools/DragToolsBehaviourSource.h
--- a/Engine/Source/Editor/Experimental/EditorInteractiveToolsFramework/Public/EditorDragTools/DragToolsBehaviourSource.h
+++ b/Engine/Source/Editor/Experimental/EditorInteractiveToolsFramework/Public/EditorDragTools/DragToolsBehaviourSource.h
@@ -29,7 +29,7 @@ public:
 	static FOnOnViewportChangeToolToggleDelegate& OnViewportChangeToolDeactivated();
 
 	//~ Begin IInputBehaviorSource
-	EDITORINTERACTIVETOOLSFRAMEWORK_API virtual const UInputBehaviorSet* GetInputBehaviors() const override { return BehaviorSet; }
+	virtual const UInputBehaviorSet* GetInputBehaviors() const override { return BehaviorSet; }
 	//~ End IClickDragBehaviorTarget
 
 	/**
diff --git a/Engine/Source/Editor/Experimental/EditorInteractiveToolsFramework/Public/EditorDragTools/EditorMoveCameraWithObject.h b/Engine/Source/Editor/Experimental/EditorInteractiveToolsFramework/Public/EditorDragTools/EditorMoveCameraWithObject.h
--- a/Engine/Source/Editor/Experimental/EditorInteractiveToolsFramework/Public/EditorDragTools/EditorMoveCameraWithObject.h
+++ b/Engine/Source/Editor/Experimental/EditorInteractiveToolsFramework/Public/EditorDragTools/EditorMoveCameraWithObject.h
@@ -1,10 +1,12 @@
-﻿// Copyright Epic Games, Inc. All Rights Reserved.
+// Copyright Epic Games, Inc. All Rights Reserved.
 
 #pragma once
 
 #include "BaseGizmos/TransformProxy.h"
 #include "EditorGizmos/TransformGizmo.h"
 
+#define UE_API EDITORINTERACTIVETOOLSFRAMEWORK_API
+
 class FLevelEditorViewportClient;
 class UEditorTransformGizmo;
 class UInteractiveToolManager;
@@ -14,37 +16,37 @@ class UTransformProxy;
  * Listens for key press of a specified modifier key, and if Level Viewport Gizmo is dragging, it duplicates current selection.
  * Currently used with Shift (see UDragToolsBehaviorSource)
  */
-class EDITORINTERACTIVETOOLSFRAMEWORK_API FEditorMoveCameraWithObject
+class FEditorMoveCameraWithObject
 	: public IModifierToggleBehaviorTarget
 	, public FGCObject
 {
 public:
-	FEditorMoveCameraWithObject(UInteractiveToolManager* InToolManager);
+	UE_API FEditorMoveCameraWithObject(UInteractiveToolManager* InToolManager);
 
-	virtual ~FEditorMoveCameraWithObject() override;
+	UE_API virtual ~FEditorMoveCameraWithObject() override;
 
 protected:
 	//~Begin IModifierToggleBehaviorTarget
-	virtual void OnUpdateModifierState(int InModifierID, bool bInIsOn) override;
+	UE_API virtual void OnUpdateModifierState(int InModifierID, bool bInIsOn) override;
 	//~End IModifierToggleBehaviorTarget
 
 	//~Begin FGCObject
-	virtual void AddReferencedObjects(FReferenceCollector& InCollector) override;
+	UE_API virtual void AddReferencedObjects(FReferenceCollector& InCollector) override;
 	virtual FString GetReferencerName() const override
 	{
 		return TEXT("FEditorMoveCameraWithObject");
 	}
 	//~End FGCObject
 
-	void OnGizmoMovementBegin(UTransformProxy* InTransformProxy);
-	void OnGizmoMovementEnd(UTransformProxy* InTransformProxy);
-	void OnGizmoTransformChanged(UTransformProxy* InTransformProxy, UE::Math::TTransform<double> InTransform);
+	UE_API void OnGizmoMovementBegin(UTransformProxy* InTransformProxy);
+	UE_API void OnGizmoMovementEnd(UTransformProxy* InTransformProxy);
+	UE_API void OnGizmoTransformChanged(UTransformProxy* InTransformProxy, UE::Math::TTransform<double> InTransform);
 
-	void Initialize();
-	void Reset();
+	UE_API void Initialize();
+	UE_API void Reset();
 
-	void OnGizmoCreatedDelegate(UTransformGizmo* InTransformGizmo);
-	void OnUsesNewTRSGizmosChanged(bool bInUseNewTRSGizmos);
+	UE_API void OnGizmoCreatedDelegate(UTransformGizmo* InTransformGizmo);
+	UE_API void OnUsesNewTRSGizmosChanged(bool bInUseNewTRSGizmos);
 
 private:
 	TObjectPtr<UInteractiveToolManager> ToolManager;
@@ -57,3 +59,5 @@ private:
 	FDelegateHandle OnEndPivotEditDelegate;
 	FDelegateHandle OnTransformChangedDelegate;
 };
+
+#undef UE_API
diff --git a/Engine/Source/Editor/Experimental/EditorInteractiveToolsFramework/Public/EditorGizmos/EditorTransformGizmo.h b/Engine/Source/Editor/Experimental/EditorInteractiveToolsFramework/Public/EditorGizmos/EditorTransformGizmo.h
--- a/Engine/Source/Editor/Experimental/EditorInteractiveToolsFramework/Public/EditorGizmos/EditorTransformGizmo.h
+++ b/Engine/Source/Editor/Experimental/EditorInteractiveToolsFramework/Public/EditorGizmos/EditorTransformGizmo.h
@@ -10,33 +10,35 @@
 
 #include "EditorTransformGizmo.generated.h"
 
+#define UE_API EDITORINTERACTIVETOOLSFRAMEWORK_API
+
 class UObject;
 
 /**
  * UEditorTransformGizmo handles Editor-specific functionality for the TransformGizmo,
  * applied to a UEditorTransformProxy target object.
  */
-UCLASS()
-class EDITORINTERACTIVETOOLSFRAMEWORK_API UEditorTransformGizmo : public UTransformGizmo
+UCLASS(MinimalAPI)
+class UEditorTransformGizmo : public UTransformGizmo
 {
 	GENERATED_BODY()
 
 	/**  UTransformGizmo override */
-	virtual void Render(IToolsContextRenderAPI* RenderAPI) override;
+	UE_API virtual void Render(IToolsContextRenderAPI* RenderAPI) override;
 	
 protected:
 
 	/** Apply translate delta to transform proxy */
-	virtual void ApplyTranslateDelta(const FVector& InTranslateDelta) override;
+	UE_API virtual void ApplyTranslateDelta(const FVector& InTranslateDelta) override;
 
 	/** Apply rotate delta to transform proxy */
-	virtual void ApplyRotateDelta(const FQuat& InRotateDelta) override;
+	UE_API virtual void ApplyRotateDelta(const FQuat& InRotateDelta) override;
 
 	/** Apply scale delta to transform proxy */
-	virtual void ApplyScaleDelta(const FVector& InScaleDelta) override;
+	UE_API virtual void ApplyScaleDelta(const FVector& InScaleDelta) override;
 
 	/**  UTransformGizmo override */
-	virtual void SetActiveTarget(
+	UE_API virtual void SetActiveTarget(
 		UTransformProxy* Target,
 		IToolContextTransactionProvider* TransactionProvider = nullptr,
 		IGizmoStateTarget* InStateTarget = nullptr) override;
@@ -45,6 +47,8 @@ protected:
 	 * Functions to listen to gizmo transform begin/end.
 	 * They are currently used to set data on the legacy widget as some viewport clients rely on current axis to be set.
 	 */
-	void OnGizmoTransformBegin(UTransformProxy* TransformProxy) const;
-	void OnGizmoTransformEnd(UTransformProxy* TransformProxy) const;;
-};
\ No newline at end of file
+	UE_API void OnGizmoTransformBegin(UTransformProxy* TransformProxy) const;
+	UE_API void OnGizmoTransformEnd(UTransformProxy* TransformProxy) const;;
+};
+
+#undef UE_API
diff --git a/Engine/Source/Editor/Experimental/EditorInteractiveToolsFramework/Public/EditorGizmos/EditorTransformGizmoBuilder.h b/Engine/Source/Editor/Experimental/EditorInteractiveToolsFramework/Public/EditorGizmos/EditorTransformGizmoBuilder.h
--- a/Engine/Source/Editor/Experimental/EditorInteractiveToolsFramework/Public/EditorGizmos/EditorTransformGizmoBuilder.h
+++ b/Engine/Source/Editor/Experimental/EditorInteractiveToolsFramework/Public/EditorGizmos/EditorTransformGizmoBuilder.h
@@ -11,27 +11,30 @@
 
 #include "EditorTransformGizmoBuilder.generated.h"
 
+#define UE_API EDITORINTERACTIVETOOLSFRAMEWORK_API
+
 class UInteractiveGizmo;
 class UObject;
 class UEditorTransformGizmoContextObject;
 struct FToolBuilderState;
 
-UCLASS()
-class EDITORINTERACTIVETOOLSFRAMEWORK_API UEditorTransformGizmoBuilder : public UInteractiveGizmoBuilder, public IEditorInteractiveGizmoSelectionBuilder
+UCLASS(MinimalAPI)
+class UEditorTransformGizmoBuilder : public UInteractiveGizmoBuilder, public IEditorInteractiveGizmoSelectionBuilder
 {
 	GENERATED_BODY()
 
 public:
 
 	// UEditorInteractiveGizmoSelectionBuilder interface 
-	virtual UInteractiveGizmo* BuildGizmo(const FToolBuilderState& SceneState) const override;
-	virtual void UpdateGizmoForSelection(UInteractiveGizmo* Gizmo, const FToolBuilderState& SceneState) override;
+	UE_API virtual UInteractiveGizmo* BuildGizmo(const FToolBuilderState& SceneState) const override;
+	UE_API virtual void UpdateGizmoForSelection(UInteractiveGizmo* Gizmo, const FToolBuilderState& SceneState) override;
 
 	// If set, this function will be passed to UTransformGizmo instances to override the default material and display size.
 	TFunction<const FGizmoCustomization()> CustomizationFunction;
 
 private:
-	static const UEditorTransformGizmoContextObject* GetTransformGizmoContext(const FToolBuilderState& InSceneState);
+	static UE_API const UEditorTransformGizmoContextObject* GetTransformGizmoContext(const FToolBuilderState& InSceneState);
 };
 
 
+#undef UE_API
diff --git a/Engine/Source/Editor/Experimental/EditorInteractiveToolsFramework/Public/EditorGizmos/EditorTransformGizmoDataBinder.h b/Engine/Source/Editor/Experimental/EditorInteractiveToolsFramework/Public/EditorGizmos/EditorTransformGizmoDataBinder.h
--- a/Engine/Source/Editor/Experimental/EditorInteractiveToolsFramework/Public/EditorGizmos/EditorTransformGizmoDataBinder.h
+++ b/Engine/Source/Editor/Experimental/EditorInteractiveToolsFramework/Public/EditorGizmos/EditorTransformGizmoDataBinder.h
@@ -7,6 +7,8 @@
 #include "UObject/WeakObjectPtrTemplates.h"
 #include "Containers/Set.h"
 
+#define UE_API EDITORINTERACTIVETOOLSFRAMEWORK_API
+
 class UGizmoEditorStateTarget;
 class UEditorTransformGizmoContextObject;
 class UTransformGizmo;
@@ -17,27 +19,27 @@ class FEditorViewportClient;
  * FEditorTransformGizmoDataBinder is a helper class for binding a UTransformGizmo to a FEditorModeTools
  */
 
-class EDITORINTERACTIVETOOLSFRAMEWORK_API FEditorTransformGizmoDataBinder : public TSharedFromThis<FEditorTransformGizmoDataBinder>
+class FEditorTransformGizmoDataBinder : public TSharedFromThis<FEditorTransformGizmoDataBinder>
 {
 public:
 
-	virtual ~FEditorTransformGizmoDataBinder();
+	UE_API virtual ~FEditorTransformGizmoDataBinder();
 	
 	/**
 	 * Makes it so that the gizmo binder attaches to any gizmos created by the context object
 	 * in the future. The binding is automatically removed if FEditorTransformGizmoDataBinder is
 	 * destroyed.
 	 */
-	void BindToGizmoContextObject(UEditorTransformGizmoContextObject* InContextObject);
+	UE_API void BindToGizmoContextObject(UEditorTransformGizmoContextObject* InContextObject);
 
 	/** Used for binding to context object, called from the OnGizmoCreated delegate. */
-	void BindToUninitializedGizmo(UTransformGizmo* Gizmo);
+	UE_API void BindToUninitializedGizmo(UTransformGizmo* Gizmo);
 	
 	/** Binds to a specific gizmo for tracking. Requires ActiveTarget to be set. */
-	void BindToInitializedGizmo(UTransformGizmo* InGizmo, UTransformProxy* InProxy);
+	UE_API void BindToInitializedGizmo(UTransformGizmo* InGizmo, UTransformProxy* InProxy);
 	
 	/** Unbinds from a given gizmo. Done automatically for gizmos when their ActiveTarget is cleared. */
-	void UnbindFromGizmo(UTransformGizmo* InGizmo, UTransformProxy* InProxy);
+	UE_API void UnbindFromGizmo(UTransformGizmo* InGizmo, UTransformProxy* InProxy);
 
 private:
 
@@ -49,12 +51,13 @@ private:
 	/** Weak ptr to the context so we can interface with the current mode manager */
 	TWeakObjectPtr<UEditorTransformGizmoContextObject> WeakContext;
 
-	FEditorViewportClient* GetViewportClient() const;
+	UE_API FEditorViewportClient* GetViewportClient() const;
 
-	void OnProxyBeginTransformEdit(UTransformProxy* InTransformProxy);
-	void OnProxyTransformChanged(UTransformProxy* InTransformProxy, FTransform InTransform);
-	void OnProxyEndTransformEdit(UTransformProxy* InTransformProxy);
+	UE_API void OnProxyBeginTransformEdit(UTransformProxy* InTransformProxy);
+	UE_API void OnProxyTransformChanged(UTransformProxy* InTransformProxy, FTransform InTransform);
+	UE_API void OnProxyEndTransformEdit(UTransformProxy* InTransformProxy);
 
 	bool bHasTransformChanged = false;
 };
 
+#undef UE_API
diff --git a/Engine/Source/Editor/Experimental/EditorInteractiveToolsFramework/Public/EditorGizmos/EditorTransformGizmoSource.h b/Engine/Source/Editor/Experimental/EditorInteractiveToolsFramework/Public/EditorGizmos/EditorTransformGizmoSource.h
--- a/Engine/Source/Editor/Experimental/EditorInteractiveToolsFramework/Public/EditorGizmos/EditorTransformGizmoSource.h
+++ b/Engine/Source/Editor/Experimental/EditorInteractiveToolsFramework/Public/EditorGizmos/EditorTransformGizmoSource.h
@@ -15,6 +15,8 @@
 
 #include "EditorTransformGizmoSource.generated.h"
 
+#define UE_API EDITORINTERACTIVETOOLSFRAMEWORK_API
+
 class UEditorTransformGizmoContextObject;
 class FEditorModeTools;
 class FEditorViewportClient;
@@ -33,8 +35,8 @@ namespace FEditorTransformGizmoUtil
  * UEditorTransformGizmoSource is an ITransformGizmoSource implementation that provides
  * current state information used to configure the Editor transform gizmo.
  */
-UCLASS()
-class EDITORINTERACTIVETOOLSFRAMEWORK_API UEditorTransformGizmoSource : public UObject, public ITransformGizmoSource
+UCLASS(MinimalAPI)
+class UEditorTransformGizmoSource : public UObject, public ITransformGizmoSource
 {
 	GENERATED_BODY()
 public:
@@ -42,52 +44,53 @@ public:
 	/**
 	 * @return The current display mode for the Editor transform gizmo
 	 */
-	virtual EGizmoTransformMode GetGizmoMode() const override;
+	UE_API virtual EGizmoTransformMode GetGizmoMode() const override;
 
 	/**
 	 * @return The current axes to draw for the specified mode
 	 */
-	virtual EAxisList::Type GetGizmoAxisToDraw(EGizmoTransformMode InWidgetMode) const override;
+	UE_API virtual EAxisList::Type GetGizmoAxisToDraw(EGizmoTransformMode InWidgetMode) const override;
 
 	/**
 	 * @return The coordinate system space (world or local) to display the widget in
 	 */
-	virtual EToolContextCoordinateSystem GetGizmoCoordSystemSpace() const override;
+	UE_API virtual EToolContextCoordinateSystem GetGizmoCoordSystemSpace() const override;
 
 	/**
 	 * Returns a scale factor for the gizmo
 	 */
-	virtual float GetGizmoScale() const override;
+	UE_API virtual float GetGizmoScale() const override;
 
 	/**
 	 * Whether the gizmo is visible
 	 */
-	virtual bool GetVisible(const EViewportContext InViewportContext = EViewportContext::Focused) const override;
+	UE_API virtual bool GetVisible(const EViewportContext InViewportContext = EViewportContext::Focused) const override;
 
 	/* 
 	 * Returns whether the gizmo can interact.
 	 * Note that this can be true even if the gizmo is hidden to support indirect manipulation in game mode.
 	 */
-	virtual bool CanInteract(const EViewportContext InViewportContext = EViewportContext::Focused) const override;
+	UE_API virtual bool CanInteract(const EViewportContext InViewportContext = EViewportContext::Focused) const override;
 
 	/**
  	 * Get current scale type
 	 */
-	virtual EGizmoTransformScaleType GetScaleType() const override;
+	UE_API virtual EGizmoTransformScaleType GetScaleType() const override;
 
 	/** Get rotation context */
-	virtual const FRotationContext& GetRotationContext() const override;
+	UE_API virtual const FRotationContext& GetRotationContext() const override;
 	
-	static UEditorTransformGizmoSource* CreateNew(
+	static UE_API UEditorTransformGizmoSource* CreateNew(
 		UObject* Outer = (UObject*)GetTransientPackage(),
 		const UEditorTransformGizmoContextObject* InContext = nullptr);
 
 protected:
 
-	const FEditorModeTools& GetModeTools() const;
+	UE_API const FEditorModeTools& GetModeTools() const;
 
-	const FEditorViewportClient* GetViewportClient(const EViewportContext InViewportContext = EViewportContext::Focused) const;
+	UE_API const FEditorViewportClient* GetViewportClient(const EViewportContext InViewportContext = EViewportContext::Focused) const;
 	
 	TWeakObjectPtr<const UEditorTransformGizmoContextObject> WeakContext = nullptr;
 };
 
+#undef UE_API
diff --git a/Engine/Source/Editor/Experimental/EditorInteractiveToolsFramework/Public/EditorGizmos/EditorTransformProxy.h b/Engine/Source/Editor/Experimental/EditorInteractiveToolsFramework/Public/EditorGizmos/EditorTransformProxy.h
--- a/Engine/Source/Editor/Experimental/EditorInteractiveToolsFramework/Public/EditorGizmos/EditorTransformProxy.h
+++ b/Engine/Source/Editor/Experimental/EditorInteractiveToolsFramework/Public/EditorGizmos/EditorTransformProxy.h
@@ -14,6 +14,8 @@
 
 #include "EditorTransformProxy.generated.h"
 
+#define UE_API EDITORINTERACTIVETOOLSFRAMEWORK_API
+
 class UEditorTransformGizmoContextObject;
 class FEditorViewportClient;
 class UObject;
@@ -25,8 +27,8 @@ class UObject;
  * 
  * This defaults internally to GLevelEditorModeManager() if none provided when constructing the proxy.
  */
-UCLASS(Transient)
-class EDITORINTERACTIVETOOLSFRAMEWORK_API UEditorTransformProxy : public UTransformProxy
+UCLASS(MinimalAPI, Transient)
+class UEditorTransformProxy : public UTransformProxy
 {
 	GENERATED_BODY()
 public:
@@ -34,7 +36,7 @@ public:
 	/**
 	 * @return the stored transform for currently selected objects.
 	 */
-	virtual FTransform GetTransform() const override;
+	UE_API virtual FTransform GetTransform() const override;
 
 	/**
 	 * Unimplemented - all updates to the Editor transform proxy MUST be made by calling the Input delta methods.
@@ -45,23 +47,24 @@ public:
 	}
 
 	/** Input translate delta to be applied in world space of the current transform. */
-	virtual void InputTranslateDelta(const FVector& InDeltaTranslate, EAxisList::Type InAxisList);
+	UE_API virtual void InputTranslateDelta(const FVector& InDeltaTranslate, EAxisList::Type InAxisList);
 
 	/** Input rotation delta to be applied in local space of the current transform. */
-	virtual void InputRotateDelta(const FRotator& InDeltaRotate, EAxisList::Type InAxisList);
+	UE_API virtual void InputRotateDelta(const FRotator& InDeltaRotate, EAxisList::Type InAxisList);
 
 	/** Input scale delta to be applied in local space of the current transform. */
-	virtual void InputScaleDelta(const FVector& InDeltaScale, EAxisList::Type InAxisList);
+	UE_API virtual void InputScaleDelta(const FVector& InDeltaScale, EAxisList::Type InAxisList);
 
 	/** Set legacy widget axis temporarily because FEditorViewportClient overrides may expect it */
-	void SetCurrentAxis(const EAxisList::Type InAxisList) const;
+	UE_API void SetCurrentAxis(const EAxisList::Type InAxisList) const;
 
-	static UEditorTransformProxy* CreateNew(const UEditorTransformGizmoContextObject* InContext = nullptr);
+	static UE_API UEditorTransformProxy* CreateNew(const UEditorTransformGizmoContextObject* InContext = nullptr);
 	
 private:
 	
-	FEditorViewportClient* GetViewportClient() const;
+	UE_API FEditorViewportClient* GetViewportClient() const;
 
 	TWeakObjectPtr<const UEditorTransformGizmoContextObject> WeakContext;
 };
 
+#undef UE_API
