Add different padding to EditorViewportToolBar.Button.* variants to match EditorViewportToolBar.ToggleButton.*
This corrects some small visual issues when a toggle button and button are used side-by-side.

diff --git a/Engine/Source/Editor/EditorStyle/Private/StarshipStyle.cpp b/Engine/Source/Editor/EditorStyle/Private/StarshipStyle.cpp
--- a/Engine/Source/Editor/EditorStyle/Private/StarshipStyle.cpp
+++ b/Engine/Source/Editor/EditorStyle/Private/StarshipStyle.cpp
@@ -2953,19 +2953,25 @@ void FStarshipEditorStyle::FStyle::SetupViewportStyles()
 		FButtonStyle ViewportMenuButtonLeft = FButtonStyle(ViewportMenuButton)
 			.SetNormal(		BOX_BRUSH("Starship/EditorViewport/ToolBarLeftGroup", 12.f/25.f, FStyleColors::Dropdown))
 			.SetHovered(	BOX_BRUSH("Starship/EditorViewport/ToolBarLeftGroup", 12.f/25.f, FStyleColors::Hover))
-			.SetPressed(	BOX_BRUSH("Starship/EditorViewport/ToolBarLeftGroup", 12.f/25.f, FStyleColors::Recessed));
+			.SetPressed(	BOX_BRUSH("Starship/EditorViewport/ToolBarLeftGroup", 12.f/25.f, FStyleColors::Recessed))
+			.SetNormalPadding(ViewportMarginLeft)
+			.SetPressedPadding(ViewportMarginLeft);
 		Set("EditorViewportToolBar.Button.Start", ViewportMenuButtonLeft);
 
 		FButtonStyle ViewportMenuButtonMiddle = FButtonStyle(ViewportMenuButton)
         	.SetNormal(		BOX_BRUSH("Starship/EditorViewport/ToolBarMiddleGroup", 12.f/25.f, FStyleColors::Dropdown))
         	.SetHovered(	BOX_BRUSH("Starship/EditorViewport/ToolBarMiddleGroup", 12.f/25.f, FStyleColors::Hover))
-        	.SetPressed(	BOX_BRUSH("Starship/EditorViewport/ToolBarMiddleGroup", 12.f/25.f, FStyleColors::Recessed));
+        	.SetPressed(	BOX_BRUSH("Starship/EditorViewport/ToolBarMiddleGroup", 12.f/25.f, FStyleColors::Recessed))
+        	.SetNormalPadding(ViewportMarginCenter)
+        	.SetPressedPadding(ViewportMarginCenter);
 		Set("EditorViewportToolBar.Button.Middle", ViewportMenuButtonMiddle);
 
 		FButtonStyle ViewportMenuButtonRight = FButtonStyle(ViewportMenuButton)
 			.SetNormal(		BOX_BRUSH("Starship/EditorViewport/ToolBarRightGroup", 12.f/25.f, FStyleColors::Dropdown))
 			.SetHovered(	BOX_BRUSH("Starship/EditorViewport/ToolBarRightGroup", 12.f/25.f, FStyleColors::Hover))
-			.SetPressed(	BOX_BRUSH("Starship/EditorViewport/ToolBarRightGroup", 12.f/25.f, FStyleColors::Recessed));
+			.SetPressed(	BOX_BRUSH("Starship/EditorViewport/ToolBarRightGroup", 12.f/25.f, FStyleColors::Recessed))
+			.SetNormalPadding(ViewportMarginRight)
+			.SetPressedPadding(ViewportMarginRight);
 		Set("EditorViewportToolBar.Button.End", ViewportMenuButtonRight);
 
 		Set("EditorViewportToolBar.StartToolbarImage", new BOX_BRUSH("Starship/EditorViewport/ToolBarLeftGroup", 12.f/25.f, FStyleColors::Dropdown));
