Fix an issue where top-level toggle buttons in Viewport Toolbars looked like raised buttons.

diff --git a/Engine/Source/Editor/EditorStyle/Private/StarshipStyle.cpp b/Engine/Source/Editor/EditorStyle/Private/StarshipStyle.cpp
--- a/Engine/Source/Editor/EditorStyle/Private/StarshipStyle.cpp
+++ b/Engine/Source/Editor/EditorStyle/Private/StarshipStyle.cpp
@@ -2588,14 +2588,15 @@ void FStarshipEditorStyle::FStyle::SetupViewportStyles()
 			FStarshipCoreStyle::GetCoreStyle().GetWidgetStyle<FToolBarStyle>("SlimToolBar");
 
 		FCheckBoxStyle ToggleButtonStyle = FCheckBoxStyle(SlimToolBarStyle.ToggleButton)
-			.SetCheckedImage(TransparentBrush)
+			.SetCheckedImage(BackgroundBrush)
+			.SetUncheckedImage(BackgroundBrush)
 			.SetCheckedHoveredImage(HoveredButtonBrush)
 			.SetCheckedPressedImage(PressedButtonBrush)
 			.SetUncheckedPressedImage(PressedButtonBrush)
 			.SetCheckedForegroundColor(FStyleColors::AccentBlue)
 			.SetCheckedHoveredForegroundColor(FStyleColors::AccentBlue)
 			.SetCheckedPressedForegroundColor(FStyleColors::AccentBlue)
-			.SetPadding(FMargin(4.0f));
+			.SetPadding(FMargin(6.0f, 4.0f));
 		
 		FCheckBoxStyle ToggleButtonStyleStart = FCheckBoxStyle(ToggleButtonStyle)
 			.SetCheckedImage(BackgroundBrushStart)
@@ -2704,31 +2705,19 @@ void FStarshipEditorStyle::FStyle::SetupViewportStyles()
 			FButtonStyle ButtonStyleRaisedMiddle = ConvertButtonToRaised(ButtonStyleMiddle);
 			FButtonStyle ButtonStyleRaisedEnd = ConvertButtonToRaised(ButtonStyleEnd);
 			
-			FCheckBoxStyle ToggleButtonStyleRaised = FCheckBoxStyle(ToggleButtonStyle)
-				.SetUncheckedHoveredImage(HoveredButtonBrush)
-				.SetCheckedPressedImage(BackgroundBrush)
-				.SetUncheckedPressedImage(BackgroundBrush);
+			const auto ConvertToggleButtonToRaised = [&TransparentBrush](const FCheckBoxStyle& ButtonStyle)
+			{
+				return FCheckBoxStyle(ButtonStyle)
+					.SetPadding(4.0f)
+					.SetForegroundColor(FStyleColors::Foreground)
+					.SetUncheckedImage(TransparentBrush)
+					.SetCheckedImage(TransparentBrush);
+			};
 			
-			FCheckBoxStyle ToggleButtonStyleRaisedStart = FCheckBoxStyle(ToggleButtonStyleStart)
-				.SetUncheckedImage(TransparentBrush)
-				.SetCheckedImage(TransparentBrush)
-				.SetUncheckedHoveredImage(HoveredButtonBrushStart)
-				.SetCheckedPressedImage(BackgroundBrushStart)
-				.SetUncheckedPressedImage(BackgroundBrushStart);
-				
-			FCheckBoxStyle ToggleButtonStyleRaisedMiddle = FCheckBoxStyle(ToggleButtonStyleMiddle)
-				.SetUncheckedImage(TransparentBrush)
-				.SetCheckedImage(TransparentBrush)
-				.SetUncheckedHoveredImage(HoveredButtonBrushMiddle)
-				.SetCheckedPressedImage(BackgroundBrushMiddle)
-				.SetUncheckedPressedImage(BackgroundBrushMiddle);
-				
-			FCheckBoxStyle ToggleButtonStyleRaisedEnd = FCheckBoxStyle(ToggleButtonStyleEnd)
-				.SetUncheckedImage(TransparentBrush)
-				.SetCheckedImage(TransparentBrush)
-				.SetUncheckedHoveredImage(HoveredButtonBrushEnd)
-				.SetCheckedPressedImage(BackgroundBrushEnd)
-				.SetUncheckedPressedImage(BackgroundBrushEnd);
+			FCheckBoxStyle ToggleButtonStyleRaised = ConvertToggleButtonToRaised(ToggleButtonStyle);
+			FCheckBoxStyle ToggleButtonStyleRaisedStart = ConvertToggleButtonToRaised(ToggleButtonStyleStart);
+			FCheckBoxStyle ToggleButtonStyleRaisedMiddle = ConvertToggleButtonToRaised(ToggleButtonStyleMiddle);
+			FCheckBoxStyle ToggleButtonStyleRaisedEnd = ConvertToggleButtonToRaised(ToggleButtonStyleEnd);
 
 			FComboButtonStyle ComboStyleRaised = FComboButtonStyle(ViewportToolbarStyle.ComboButtonStyle)
 				.SetDownArrowPadding(FMargin(2.0f, 0, 0, 0));
