Sequencer: Update styles of keying icons in details rows

diff --git a/Engine/Plugins/Animation/ControlRig/Source/ControlRigEditor/Private/EditMode/ControlRigEditMode.cpp b/Engine/Plugins/Animation/ControlRig/Source/ControlRigEditor/Private/EditMode/ControlRigEditMode.cpp
--- a/Engine/Plugins/Animation/ControlRig/Source/ControlRigEditor/Private/EditMode/ControlRigEditMode.cpp
+++ b/Engine/Plugins/Animation/ControlRig/Source/ControlRigEditor/Private/EditMode/ControlRigEditMode.cpp
@@ -6961,11 +6961,11 @@ EPropertyKeyedStatus FDetailKeyFrameCacheAndHandler::GetPropertyKeyedStatus(cons
 		{
 			if (UAnimDetailsProxyBase* Proxy = Cast<UAnimDetailsProxyBase>(Object))
 			{
-				KeyedStatus = Proxy->GetPropertyKeyedStatus(PropertyHandle);
+				KeyedStatus = FMath::Max(KeyedStatus, Proxy->GetPropertyKeyedStatus(PropertyHandle));
 			}
 			else if (UAnimLayer* AnimLayer = (Object->GetTypedOuter<UAnimLayer>()))
 			{
-				KeyedStatus = AnimLayer->GetPropertyKeyedStatus(SequencerPtr, PropertyHandle);
+				KeyedStatus = FMath::Max(KeyedStatus, AnimLayer->GetPropertyKeyedStatus(SequencerPtr, PropertyHandle));
 			}
 		}
 		//else check to see if it's in sequencer
diff --git a/Engine/Source/Editor/EditorStyle/Private/StarshipStyle.cpp b/Engine/Source/Editor/EditorStyle/Private/StarshipStyle.cpp
--- a/Engine/Source/Editor/EditorStyle/Private/StarshipStyle.cpp
+++ b/Engine/Source/Editor/EditorStyle/Private/StarshipStyle.cpp
@@ -2422,10 +2422,10 @@ void FStarshipEditorStyle::FStyle::SetupSequencerStyles()
 
 		Set("Sequencer.AddKey.Details", new IMAGE_BRUSH("Sequencer/AddKey_Details", FVector2f(11.f, 11.f)));
 
-		Set("Sequencer.KeyedStatus.NotKeyed", new IMAGE_BRUSH_SVG("Sequencer/DetailsKeyUnkeyed", FVector2f(11.f, 11.f)));
-		Set("Sequencer.KeyedStatus.Keyed", new IMAGE_BRUSH_SVG("Sequencer/DetailsKeyKeyed", FVector2f(11.f, 11.f)));
-		Set("Sequencer.KeyedStatus.Animated", new IMAGE_BRUSH_SVG("Sequencer/DetailsKeyAnimated", FVector2f(11.f, 11.f)));
-		Set("Sequencer.KeyedStatus.PartialKey", new IMAGE_BRUSH_SVG("Sequencer/DetailsKeyPartialKey", FVector2f(11.f, 11.f)));
+		Set("Sequencer.KeyedStatus.NotKeyed", new IMAGE_BRUSH_SVG("Sequencer/DetailsKeyUnkeyed", FVector2f(11.f, 11.f), FLinearColor::FromSRGBColor(FColor(192, 192, 192, 127))));
+		Set("Sequencer.KeyedStatus.Keyed", new IMAGE_BRUSH_SVG("Sequencer/DetailsKeyKeyed", FVector2f(11.f, 11.f), FLinearColor::FromSRGBColor(FColor(254, 133, 57))));
+		Set("Sequencer.KeyedStatus.Animated", new IMAGE_BRUSH_SVG("Sequencer/DetailsKeyAnimated", FVector2f(11.f, 11.f), FLinearColor::FromSRGBColor(FColor(254, 133, 57))));
+		Set("Sequencer.KeyedStatus.PartialKey", new IMAGE_BRUSH_SVG("Sequencer/DetailsKeyPartialKey", FVector2f(11.f, 11.f), FLinearColor::FromSRGBColor(FColor(254, 133, 57))));
 
 		const FSplitterStyle OutlinerSplitterStyle = FSplitterStyle()
 		.SetHandleNormalBrush(FSlateNoResource())
