PR #11432: Added command to file menu that opens up the project directory

diff --git a/Engine/Source/Editor/EditorStyle/Private/StarshipStyle.cpp b/Engine/Source/Editor/EditorStyle/Private/StarshipStyle.cpp
--- a/Engine/Source/Editor/EditorStyle/Private/StarshipStyle.cpp
+++ b/Engine/Source/Editor/EditorStyle/Private/StarshipStyle.cpp
@@ -1338,24 +1338,24 @@ void FStarshipEditorStyle::FStyle::SetupGeneralStyles()
 
 	// Main frame
 	{
-		Set( "MainFrame.AutoSaveImage", 	       new IMAGE_BRUSH_SVG( "Starship/Common/SaveCurrent", Icon16x16 ) );
-		Set( "MainFrame.SaveAll",                  new IMAGE_BRUSH_SVG( "Starship/Common/SaveAll", Icon16x16 ) );
-		Set( "MainFrame.ChoosePackagesToSave",     new IMAGE_BRUSH_SVG( "Starship/Common/SaveChoose", Icon16x16 ) );
-		Set( "MainFrame.NewProject",               new IMAGE_BRUSH_SVG( "Starship/Common/ProjectNew", Icon16x16 ) );
-		Set( "MainFrame.OpenProject",              new IMAGE_BRUSH_SVG( "Starship/Common/ProjectOpen", Icon16x16 ) );
-		Set( "MainFrame.AddCodeToProject",         new IMAGE_BRUSH_SVG( "Starship/Common/ProjectC++", Icon16x16 ) );
-		Set( "MainFrame.Exit",                     new IMAGE_BRUSH_SVG( "Starship/Common/Exit", Icon16x16 ) );
-		Set( "MainFrame.CookContent",              new IMAGE_BRUSH_SVG( "Starship/Common/CookContent", Icon16x16 ) );
-		Set( "MainFrame.OpenVisualStudio",         new IMAGE_BRUSH_SVG( "Starship/Common/VisualStudio", Icon16x16 ) );
-		Set( "MainFrame.RefreshVisualStudio",      new IMAGE_BRUSH_SVG( "Starship/Common/RefreshVisualStudio", Icon16x16 ) );
-		Set( "MainFrame.OpenSourceCodeEditor",     new IMAGE_BRUSH_SVG( "Starship/Common/SourceCodeEditor", Icon16x16));
-		Set( "MainFrame.RefreshSourceCodeEditor",  new IMAGE_BRUSH_SVG( "Starship/Common/RefreshSourceCodeEditor", Icon16x16));
-		Set( "MainFrame.PackageProject",           new IMAGE_BRUSH_SVG( "Starship/Common/ProjectPackage", Icon16x16 ) );
-		Set( "MainFrame.RecentProjects",           new IMAGE_BRUSH_SVG( "Starship/Common/ProjectsRecent", Icon16x16 ) );
-		Set( "MainFrame.RecentLevels",             new IMAGE_BRUSH_SVG( "Starship/Common/LevelRecent", Icon16x16 ) );
-		Set( "MainFrame.FavoriteLevels",           new IMAGE_BRUSH_SVG( "Starship/Common/LevelFavorite", Icon16x16 ) );
-		Set( "MainFrame.ZipUpProject", 		       new IMAGE_BRUSH_SVG( "Starship/Common/ZipProject", Icon16x16 ) );
-
+		Set( "MainFrame.AutoSaveImage", 	         new IMAGE_BRUSH_SVG( "Starship/Common/SaveCurrent", Icon16x16 ) );
+		Set( "MainFrame.SaveAll",                    new IMAGE_BRUSH_SVG( "Starship/Common/SaveAll", Icon16x16 ) );
+		Set( "MainFrame.ChoosePackagesToSave",       new IMAGE_BRUSH_SVG( "Starship/Common/SaveChoose", Icon16x16 ) );
+		Set( "MainFrame.NewProject",                 new IMAGE_BRUSH_SVG( "Starship/Common/ProjectNew", Icon16x16 ) );
+		Set( "MainFrame.OpenProject",                new IMAGE_BRUSH_SVG( "Starship/Common/ProjectOpen", Icon16x16 ) );
+		Set( "MainFrame.AddCodeToProject",           new IMAGE_BRUSH_SVG( "Starship/Common/ProjectC++", Icon16x16 ) );
+		Set( "MainFrame.Exit",                       new IMAGE_BRUSH_SVG( "Starship/Common/Exit", Icon16x16 ) );
+		Set( "MainFrame.CookContent",                new IMAGE_BRUSH_SVG( "Starship/Common/CookContent", Icon16x16 ) );
+		Set( "MainFrame.OpenVisualStudio",           new IMAGE_BRUSH_SVG( "Starship/Common/VisualStudio", Icon16x16 ) );
+		Set( "MainFrame.RefreshVisualStudio",        new IMAGE_BRUSH_SVG( "Starship/Common/RefreshVisualStudio", Icon16x16 ) );
+		Set( "MainFrame.OpenSourceCodeEditor",       new IMAGE_BRUSH_SVG( "Starship/Common/SourceCodeEditor", Icon16x16));
+		Set( "MainFrame.RefreshSourceCodeEditor",    new IMAGE_BRUSH_SVG( "Starship/Common/RefreshSourceCodeEditor", Icon16x16));
+		Set( "MainFrame.PackageProject",             new IMAGE_BRUSH_SVG( "Starship/Common/ProjectPackage", Icon16x16 ) );
+		Set( "MainFrame.RecentProjects",             new IMAGE_BRUSH_SVG( "Starship/Common/ProjectsRecent", Icon16x16 ) );
+		Set( "MainFrame.RecentLevels",               new IMAGE_BRUSH_SVG( "Starship/Common/LevelRecent", Icon16x16 ) );
+		Set( "MainFrame.FavoriteLevels",             new IMAGE_BRUSH_SVG( "Starship/Common/LevelFavorite", Icon16x16 ) );
+		Set( "MainFrame.ZipUpProject", 		         new IMAGE_BRUSH_SVG( "Starship/Common/ZipProject", Icon16x16 ) );
+		Set("MainFrame.OpenCurrentProjectDirectory", new IMAGE_BRUSH_SVG( "Starship/Common/BrowseContent", Icon16x16));
 
 		Set( "MainFrame.ChooseFilesToSave",       new IMAGE_BRUSH_SVG( "Starship/Common/SaveChoose", Icon16x16 ) );
 		Set( "MainFrame.ConnectToSourceControl",  new CORE_IMAGE_BRUSH_SVG("Starship/SourceControl/Status/RevisionControl", Icon16x16) );
diff --git a/Engine/Source/Editor/MainFrame/Private/Frame/MainFrameActions.cpp b/Engine/Source/Editor/MainFrame/Private/Frame/MainFrameActions.cpp
--- a/Engine/Source/Editor/MainFrame/Private/Frame/MainFrameActions.cpp
+++ b/Engine/Source/Editor/MainFrame/Private/Frame/MainFrameActions.cpp
@@ -144,6 +144,9 @@ void FMainFrameCommands::RegisterCommands()
 	UI_COMMAND( ZipUpProject, "Zip Project", "Zips the project into a zip file.", EUserInterfaceActionType::Button, FInputChord() );
 	ActionList->MapAction(ZipUpProject, FExecuteAction::CreateStatic( &FMainFrameActionCallbacks::ZipUpProject ), DefaultExecuteAction);
 
+	UI_COMMAND(OpenCurrentProjectDirectory, "Open Current Project Directory", "Opens the current project's directory in explorer.", EUserInterfaceActionType::Button, FInputChord());
+	ActionList->MapAction(OpenCurrentProjectDirectory, FExecuteAction::CreateStatic(&FMainFrameActionCallbacks::OpenCurrentProjectDirectory), DefaultExecuteAction);
+
 	//UI_COMMAND( LocalizeProject, "Localize Project...", "Opens the dashboard for managing project localization data.", EUserInterfaceActionType::Button, FInputChord() );
 	//ActionList->MapAction( LocalizeProject, FExecuteAction::CreateStatic( &FMainFrameActionCallbacks::LocalizeProject ), DefaultExecuteAction );
 
@@ -548,6 +551,12 @@ void FMainFrameActionCallbacks::ZipUpProject()
 	}
 }
 
+void FMainFrameActionCallbacks::OpenCurrentProjectDirectory()
+{
+	FString ProjectDirectory = FPaths::ConvertRelativePathToFull(FPaths::ProjectDir());
+	FPlatformProcess::ExploreFolder(*ProjectDirectory);
+}
+
 void FMainFrameActionCallbacks::SwitchProjectByIndex(int32 ProjectIndex)
 {
 	FUnrealEdMisc::Get().SwitchProject(RecentProjects[ProjectIndex].ProjectName);
diff --git a/Engine/Source/Editor/MainFrame/Private/Frame/MainFrameActions.h b/Engine/Source/Editor/MainFrame/Private/Frame/MainFrameActions.h
--- a/Engine/Source/Editor/MainFrame/Private/Frame/MainFrameActions.h
+++ b/Engine/Source/Editor/MainFrame/Private/Frame/MainFrameActions.h
@@ -39,6 +39,7 @@ public:
 	TSharedPtr< FUICommandInfo > OpenIDE;
 	TSharedPtr< FUICommandInfo > RefreshCodeProject;
 	TSharedPtr< FUICommandInfo > ZipUpProject;
+	TSharedPtr< FUICommandInfo > OpenCurrentProjectDirectory;
 	TSharedPtr< FUICommandInfo > LocalizeProject;
 	TArray< TSharedPtr< FUICommandInfo > > SwitchProjectCommands;
 
@@ -246,6 +247,9 @@ public:
 	/** Zips up the project */
 	static void ZipUpProject();
 
+	/** Opens up the project directory */
+	static void OpenCurrentProjectDirectory();
+
 	/** Opens the Project Localization Dashboard */
 	static void LocalizeProject();
 	
diff --git a/Engine/Source/Editor/MainFrame/Private/Menus/MainMenu.cpp b/Engine/Source/Editor/MainFrame/Private/Menus/MainMenu.cpp
--- a/Engine/Source/Editor/MainFrame/Private/Menus/MainMenu.cpp
+++ b/Engine/Source/Editor/MainFrame/Private/Menus/MainMenu.cpp
@@ -547,6 +547,7 @@ void FMainMenu::RegisterFileProjectMenu()
 
 
 		Section.AddMenuEntry(FMainFrameCommands::Get().ZipUpProject);
+		Section.AddMenuEntry(FMainFrameCommands::Get().OpenCurrentProjectDirectory);
 
 		if (FMainFrameActionCallbacks::RecentProjects.Num() > 0)
 		{
