[Synthesis/AudioEditor] Update Sounds Icons with new ones

diff --git a/Engine/Source/Editor/EditorStyle/Private/StarshipStyle.cpp b/Engine/Source/Editor/EditorStyle/Private/StarshipStyle.cpp
--- a/Engine/Source/Editor/EditorStyle/Private/StarshipStyle.cpp
+++ b/Engine/Source/Editor/EditorStyle/Private/StarshipStyle.cpp
@@ -6665,6 +6665,7 @@ void FStarshipEditorStyle::FStyle::SetupClassIconsAndThumbnails()
 			{TEXT("FoliageType_Actor")},
 			{TEXT("Font")},
 			{TEXT("ForceFeedback")},
+			{TEXT("FusionPatch")},
 			{TEXT("GameModeBase")},
 			{TEXT("GameStateBase")},
 			{TEXT("GeometryCollection")},
@@ -6673,6 +6674,7 @@ void FStarshipEditorStyle::FStyle::SetupClassIconsAndThumbnails()
 			{TEXT("HLODLayer")},
 			{TEXT("HUD")},
 			{TEXT("ImagePlate")},
+			{TEXT("ImpulseResponse")},
 			{TEXT("InstancedStaticMeshComponent")},
 			{TEXT("Interface")},
 			{TEXT("KillZVolume")},
@@ -7006,18 +7008,20 @@ void FStarshipEditorStyle::FStyle::SetupContentBrowserStyle()
 		Set( "ContentBrowser.AssetViewCopyObjectPath", new IMAGE_BRUSH_SVG("../../Slate/Starship/Common/Copy", Icon16x16) );
 		Set( "ContentBrowser.AssetViewCopyPackageName", new IMAGE_BRUSH_SVG("../../Slate/Starship/Common/Copy", Icon16x16) );
 
-		Set( "MediaAsset.AssetActions.Play.Small", new IMAGE_BRUSH( "Icons/icon_SCueEd_PlayCue_16x", Icon16x16 ) );
-		Set( "MediaAsset.AssetActions.Stop.Small", new IMAGE_BRUSH( "Icons/icon_SCueEd_Stop_16x", Icon16x16 ) );
-		Set( "MediaAsset.AssetActions.Pause.Small", new IMAGE_BRUSH( "Icons/icon_SCueEd_Pause_16x", Icon16x16 ) );
-		Set( "MediaAsset.AssetActions.Solo.Small", new IMAGE_BRUSH( "Icons/icon_SCueEd_Solo_16x", Icon16x16));
-		Set( "MediaAsset.AssetActions.Mute.Small", new IMAGE_BRUSH( "Icons/icon_SCueEd_Mute_16x", Icon16x16));
+		Set( "MediaAsset.AssetActions.Play.Small", new IMAGE_BRUSH_SVG("Starship/Audio/PlayControlsPlay_16", Icon16x16));
+		Set( "MediaAsset.AssetActions.Stop.Small", new IMAGE_BRUSH_SVG("Starship/Audio/PlayControlsStop_16", Icon16x16));
+		Set( "MediaAsset.AssetActions.Pause.Small", new IMAGE_BRUSH("Icons/icon_SCueEd_Pause_16x", Icon16x16));
+		Set( "MediaAsset.AssetActions.Solo.Small", new IMAGE_BRUSH_SVG("Starship/Audio/Solo_16", Icon16x16));
+		Set( "MediaAsset.AssetActions.Mute.Small", new IMAGE_BRUSH_SVG("Starship/Audio/Mute_16", Icon16x16));
+		Set( "MediaAsset.AssetActions.SoundConsolidation.Small", new IMAGE_BRUSH_SVG("Starship/Audio/SoundConsolidateAttenuation_16", Icon16x16));
 
-		Set("MediaAsset.AssetActions.Play.Large", new IMAGE_BRUSH("Icons/icon_SCueEd_PlayCue_40x", Icon40x40));
-		Set("MediaAsset.AssetActions.Stop.Large", new IMAGE_BRUSH("Icons/icon_SCueEd_Stop_40x", Icon40x40));
+		Set("MediaAsset.AssetActions.Play.Large", new IMAGE_BRUSH_SVG("Starship/Audio/PlayControlsPlay_16", Icon40x40));
+		Set("MediaAsset.AssetActions.Stop.Large", new IMAGE_BRUSH_SVG("Starship/Audio/PlayControlsStop_16", Icon40x40));
 		Set("MediaAsset.AssetActions.Pause.Large", new IMAGE_BRUSH("Icons/icon_SCueEd_Pause_40x", Icon40x40));
-		Set("MediaAsset.AssetActions.Solo.Large", new IMAGE_BRUSH("Icons/icon_SCueEd_Solo_40x", Icon40x40));
-		Set("MediaAsset.AssetActions.Mute.Large", new IMAGE_BRUSH("Icons/icon_SCueEd_Mute_40x", Icon40x40));
-				
+		Set("MediaAsset.AssetActions.Solo.Large", new IMAGE_BRUSH_SVG("Starship/Audio/Solo_16", Icon40x40));
+		Set("MediaAsset.AssetActions.Mute.Large", new IMAGE_BRUSH_SVG("Starship/Audio/Mute_16", Icon40x40));
+		Set( "MediaAsset.AssetActions.SoundConsolidation.Large", new IMAGE_BRUSH_SVG("Starship/Audio/SoundConsolidateAttenuation_16", Icon40x40));
+
 		// Misc
 		/** Should be moved, shared */ Set( "ContentBrowser.ThumbnailShadow", new BOX_BRUSH( "ContentBrowser/ThumbnailShadow" , FMargin( 4.0f / 64.0f ) ) );
 
