SceneOutliner Unsaved Column: follow up on 45029358, added FileBadge.svg so it can be updated.

diff --git a/Engine/Source/Editor/EditorStyle/Private/StarshipStyle.cpp b/Engine/Source/Editor/EditorStyle/Private/StarshipStyle.cpp
--- a/Engine/Source/Editor/EditorStyle/Private/StarshipStyle.cpp
+++ b/Engine/Source/Editor/EditorStyle/Private/StarshipStyle.cpp
@@ -419,7 +419,7 @@ void FStarshipEditorStyle::FStyle::SetupGeneralStyles()
 		Set("Icons.SaveChanged", new IMAGE_BRUSH_SVG( "Starship/Common/SaveChanged", Icon16x16 ) );
 
 		Set("Icons.DirtyBadge", new IMAGE_BRUSH_SVG("Starship/Common/DirtyBadge", Icon12x12));
-		Set("Icons.SaveableBadge", new IMAGE_BRUSH_SVG("Starship/Common/DirtyBadge", Icon12x12, FLinearColor(0.5f, 0.5f, 0.5f, 0.3f)));
+		Set("Icons.SaveableBadge", new IMAGE_BRUSH_SVG("Starship/Common/FileBadge", Icon16x16, FLinearColor(0.3f, 0.3f, 0.3f, 0.3f)));
 		Set("Icons.MakeStaticMesh", new IMAGE_BRUSH_SVG("Starship/Common/MakeStaticMesh", Icon16x16));
 		Set("Icons.Documentation", new IMAGE_BRUSH_SVG("Starship/Common/Documentation", Icon16x16));
 		Set("Icons.Support", new IMAGE_BRUSH_SVG("Starship/Common/Support", Icon16x16));
