Move setup of depth bar styles to its own function.

diff --git a/Engine/Source/Editor/EditorStyle/Private/StarshipStyle.cpp b/Engine/Source/Editor/EditorStyle/Private/StarshipStyle.cpp
--- a/Engine/Source/Editor/EditorStyle/Private/StarshipStyle.cpp
+++ b/Engine/Source/Editor/EditorStyle/Private/StarshipStyle.cpp
@@ -217,6 +217,7 @@ void FStarshipEditorStyle::FStyle::Initialize()
 	SetupWorldPartitionStyles();
 	SetupSequencerStyles();
 	SetupViewportStyles();
+	SetupDepthBarStyles();
 	SetupMenuBarStyles();
 	SetupGeneralIcons();
 	SetupWindowStyles();
@@ -3244,44 +3245,44 @@ void FStarshipEditorStyle::FStyle::SetupViewportStyles()
 		Set("VRViewportActorPreview.Attached", new IMAGE_BRUSH("Common/ScreenAttach_VR", Icon64x64));
 		Set("VRViewportActorPreview.Detached", new IMAGE_BRUSH("Common/ScreenDetach_VR", Icon64x64));
 	}
-	
-	// Viewport Depth Bar
-	{
-		FDepthBarStyle DepthBarStyle = FDepthBarStyle()
-			.SetBackgroundBrush(FSlateRoundedBoxBrush(FStyleColors::Panel, 4.0f))
-			.SetTrack(FSlateRoundedBoxBrush(FStyleColors::Input, 1.0f))
-			.SetSliceNormal(FSlateColorBrush(FStyleColors::Dropdown))
-			.SetSliceHovered(FSlateColorBrush(FStyleColors::Hover))
-			.SetSliceTopBrush(FSlateRoundedBoxBrush(FStyleColors::DropdownOutline, FVector4(2.0f, 2.0f, 0.0f, 0.0f)))
-			.SetSliceTopHovered(FSlateRoundedBoxBrush(FStyleColors::Hover2, FVector4(2.0f, 2.0f, 0.0f, 0.0f)))
-			.SetSliceBottom(FSlateRoundedBoxBrush(FStyleColors::DropdownOutline, FVector4(0.0f, 0.0f, 2.0f, 2.0f)))
-			.SetSliceBottomHovered(FSlateRoundedBoxBrush(FStyleColors::Hover2, FVector4(0.0f, 0.0f, 2.0f, 2.0f)))
-			.SetFarPlaneButtonBrush(IMAGE_BRUSH_SVG("Starship/MainToolbar/landscape", Icon16x16))
-			.SetNearPlaneButtonBrush(IMAGE_BRUSH_SVG("Starship/AssetIcons/Camera_16", Icon16x16))
-			.SetFarPlaneButtonPadding(FMargin(4.0f, 4.0f, 4.0f, 8.0f))
-			.SetNearPlaneButtonPadding(FMargin(4.0f, 8.0f, 4.0f, 0.0f))
-		;
-		Set("EditorViewport.DepthBar", DepthBarStyle);
+}
+
+void FStarshipEditorStyle::FStyle::SetupDepthBarStyles()
+{
+	FDepthBarStyle DepthBarStyle = FDepthBarStyle()
+		.SetBackgroundBrush(FSlateRoundedBoxBrush(FStyleColors::Panel, 4.0f))
+		.SetTrack(FSlateRoundedBoxBrush(FStyleColors::Input, 1.0f))
+		.SetSliceNormal(FSlateColorBrush(FStyleColors::Dropdown))
+		.SetSliceHovered(FSlateColorBrush(FStyleColors::Hover))
+		.SetSliceTopBrush(FSlateRoundedBoxBrush(FStyleColors::DropdownOutline, FVector4(2.0f, 2.0f, 0.0f, 0.0f)))
+		.SetSliceTopHovered(FSlateRoundedBoxBrush(FStyleColors::Hover2, FVector4(2.0f, 2.0f, 0.0f, 0.0f)))
+		.SetSliceBottom(FSlateRoundedBoxBrush(FStyleColors::DropdownOutline, FVector4(0.0f, 0.0f, 2.0f, 2.0f)))
+		.SetSliceBottomHovered(FSlateRoundedBoxBrush(FStyleColors::Hover2, FVector4(0.0f, 0.0f, 2.0f, 2.0f)))
+		.SetFarPlaneButtonBrush(IMAGE_BRUSH_SVG("Starship/MainToolbar/landscape", Icon16x16))
+		.SetNearPlaneButtonBrush(IMAGE_BRUSH_SVG("Starship/AssetIcons/Camera_16", Icon16x16))
+		.SetFarPlaneButtonPadding(FMargin(4.0f, 4.0f, 4.0f, 8.0f))
+		.SetNearPlaneButtonPadding(FMargin(4.0f, 8.0f, 4.0f, 0.0f))
+	;
+	Set("EditorViewport.DepthBar", DepthBarStyle);
 		
-		FToolBarStyle DepthBarToolbarStyle = GetWidgetStyle<FToolBarStyle>("ViewportToolbar.Raised");
-		DepthBarToolbarStyle
-			.SetBackgroundPadding(FMargin(0.0f))
-			.SetButtonPadding(FMargin(2.0f, 4.0f))
-			.SetComboButtonPadding(FMargin(2.0f, 4.0f))
-			.SetSeparatorPadding(FMargin(2.0f, 4.0f))
-			.SetShowLabels(false)
-		;
+	FToolBarStyle DepthBarToolbarStyle = GetWidgetStyle<FToolBarStyle>("ViewportToolbar.Raised");
+	DepthBarToolbarStyle
+		.SetBackgroundPadding(FMargin(0.0f))
+		.SetButtonPadding(FMargin(2.0f, 4.0f))
+		.SetComboButtonPadding(FMargin(2.0f, 4.0f))
+		.SetSeparatorPadding(FMargin(2.0f, 4.0f))
+		.SetShowLabels(false)
+	;
 		
-		DepthBarToolbarStyle.ComboButtonStyle
-			.SetDownArrowImage(FSlateNoResource())
-			.SetDownArrowPadding(FMargin(0.0f))
-		;
+	DepthBarToolbarStyle.ComboButtonStyle
+		.SetDownArrowImage(FSlateNoResource())
+		.SetDownArrowPadding(FMargin(0.0f))
+	;
 		
-		Set("EditorViewport.DepthBar.Toolbar", DepthBarToolbarStyle);
+	Set("EditorViewport.DepthBar.Toolbar", DepthBarToolbarStyle);
 		
-		Set("EditorViewport.DepthBar.Background", new FSlateRoundedBoxBrush(FStyleColors::Panel, 8.0f, COLOR("555555FF"), 1.0f));
-		Set("EditorViewport.DepthBar.SelectionHighlight", new FSlateColorBrush(FSlateColor(COLOR("EBA20AFF"))));
-	}
+	Set("EditorViewport.DepthBar.Background", new FSlateRoundedBoxBrush(FStyleColors::Panel, 8.0f, COLOR("555555FF"), 1.0f));
+	Set("EditorViewport.DepthBar.SelectionHighlight", new FSlateColorBrush(FSlateColor(COLOR("EBA20AFF"))));
 }
 
 void FStarshipEditorStyle::FStyle::SetupMenuBarStyles()
diff --git a/Engine/Source/Editor/EditorStyle/Private/StarshipStyle.h b/Engine/Source/Editor/EditorStyle/Private/StarshipStyle.h
--- a/Engine/Source/Editor/EditorStyle/Private/StarshipStyle.h
+++ b/Engine/Source/Editor/EditorStyle/Private/StarshipStyle.h
@@ -50,6 +50,7 @@ public:
 		void SetupWorldPartitionStyles();
 		void SetupSequencerStyles();
 		void SetupViewportStyles();
+		void SetupDepthBarStyles();
 		void SetupMenuBarStyles();
 		void SetupGeneralIcons();
 		void SetupWindowStyles();
