[AutoViz] Specialized Viewport Modes for Geometry Evaluation of Mechanical Models
- Adding Geometry Inspection viewport modes icon

diff --git a/Engine/Source/Editor/EditorStyle/Private/StarshipStyle.cpp b/Engine/Source/Editor/EditorStyle/Private/StarshipStyle.cpp
--- a/Engine/Source/Editor/EditorStyle/Private/StarshipStyle.cpp
+++ b/Engine/Source/Editor/EditorStyle/Private/StarshipStyle.cpp
@@ -4923,6 +4923,8 @@ void FStarshipEditorStyle::FStyle::SetupLevelEditorStyle()
 
 		Set( "EditorViewport.QuadOverdrawMode", new IMAGE_BRUSH_SVG("Starship/Common/OptimizationViewmodes", Icon16x16 ) );
 
+		Set( "EditorViewport.GeometryInspection", new IMAGE_BRUSH_SVG( "Starship/Common/GeometryInspection_16", Icon16x16 ) );
+
 		Set( "EditorViewport.ShaderComplexityWithQuadOverdrawMode", new IMAGE_BRUSH( "Icons/icon_ViewMode_Shadercomplexity_16px", Icon16x16 ) );
 		Set( "EditorViewport.TexStreamAccPrimitiveDistanceMode", new IMAGE_BRUSH( "Icons/icon_ViewMode_TextureStreamingAccuracy_16px", Icon16x16 ) );
 		Set( "EditorViewport.TexStreamAccMeshUVDensityMode", new IMAGE_BRUSH("Icons/icon_ViewMode_TextureStreamingAccuracy_16px", Icon16x16));
diff --git a/Engine/Source/Editor/UnrealEd/Private/ViewportToolbar/UnrealEdViewportToolbar.cpp b/Engine/Source/Editor/UnrealEd/Private/ViewportToolbar/UnrealEdViewportToolbar.cpp
--- a/Engine/Source/Editor/UnrealEd/Private/ViewportToolbar/UnrealEdViewportToolbar.cpp
+++ b/Engine/Source/Editor/UnrealEd/Private/ViewportToolbar/UnrealEdViewportToolbar.cpp
@@ -2274,7 +2274,7 @@ void PopulateViewModesMenu(UToolMenu* InMenu)
 					MenuAction,
 					EUserInterfaceActionType::RadioButton,
 					/* bInOpenSubMenuOnClick = */ false,
-					FSlateIcon(FAppStyle::GetAppStyleSetName(), "EditorViewport.CollisionVisibility")
+					FSlateIcon(FAppStyle::GetAppStyleSetName(), "EditorViewport.GeometryInspection")
 				);
 			}
 		}
diff --git a/Engine/Source/Runtime/Engine/Private/ViewModeNames.cpp b/Engine/Source/Runtime/Engine/Private/ViewModeNames.cpp
--- a/Engine/Source/Runtime/Engine/Private/ViewModeNames.cpp
+++ b/Engine/Source/Runtime/Engine/Private/ViewModeNames.cpp
@@ -274,7 +274,7 @@ FName GetViewModeIconName(const EViewModeIndex InViewMode)
 	static FName RayTracingDebugIcon("EditorViewport.RayTracingDebugMode");
 	static FName VisualizeVirtualTextureIcon("EditorViewport.VisualizeVirtualTextureMode");
 	static FName VisualizeActorColorationIcon("EditorViewport.VisualizeActorColorationMode");
-	static FName GeometryInspectionIcon("EditorViewport.CollisionVisibility");
+	static FName GeometryInspectionIcon("EditorViewport.GeometryInspection");
 
 	FName Icon = NAME_None;
 
