Cleanup landscape editor icons

diff --git a/Engine/Source/Editor/EditorStyle/Private/StarshipStyle.cpp b/Engine/Source/Editor/EditorStyle/Private/StarshipStyle.cpp
--- a/Engine/Source/Editor/EditorStyle/Private/StarshipStyle.cpp
+++ b/Engine/Source/Editor/EditorStyle/Private/StarshipStyle.cpp
@@ -7282,11 +7282,6 @@ void FStarshipEditorStyle::FStyle::SetupLandscapeEditorStyle()
 
 		}
 
-		Set("LandscapeEditor.Brushes.Alpha.UseRChannel", new IMAGE_BRUSH("Icons/icon_Landscape_Brush_Alpha_UseRChannel_20x", Icon20x20));
-		Set("LandscapeEditor.Brushes.Alpha.UseGChannel", new IMAGE_BRUSH("Icons/icon_Landscape_Brush_Alpha_UseGChannel_20x", Icon20x20));
-		Set("LandscapeEditor.Brushes.Alpha.UseBChannel", new IMAGE_BRUSH("Icons/icon_Landscape_Brush_Alpha_UseBChannel_20x", Icon20x20));
-		Set("LandscapeEditor.Brushes.Alpha.UseAChannel", new IMAGE_BRUSH("Icons/icon_Landscape_Brush_Alpha_UseAChannel_20x", Icon20x20));
-
 		Set("LandscapeEditor.Brush.AffectsHeightmap", new IMAGE_BRUSH("Icons/icon_Landscape_Affects_Heightmap_16x", Icon16x16));
 		Set("LandscapeEditor.Brush.AffectsHeightmap.Disabled", new IMAGE_BRUSH("Icons/icon_Landscape_Affects_Heightmap_Disabled_16x", Icon16x16));
 
@@ -7299,7 +7294,6 @@ void FStarshipEditorStyle::FStyle::SetupLandscapeEditorStyle()
 		Set("LandscapeEditor.Brush.AffectsWeight.Enabled", new IMAGE_BRUSH("Icons/icon_Landscape_Affects_Weight_Enabled_16x", Icon16x16));
 		Set("LandscapeEditor.Brush.AffectsWeight.Disabled", new IMAGE_BRUSH("Icons/icon_Landscape_Affects_Weight_Disabled_16x", Icon16x16));
 		
-		Set("LandscapeEditor.Brush.AffectsVisibilityLayer", new IMAGE_BRUSH("Icons/icon_Landscape_Affects_VisibilityLayer_16x", Icon16x16));
 		Set("LandscapeEditor.Brush.AffectsVisibilityLayer.Enabled", new IMAGE_BRUSH("Icons/icon_Landscape_Affects_VisibilityLayer_Enabled_16x", Icon16x16));
 		Set("LandscapeEditor.Brush.AffectsVisibilityLayer.Disabled", new IMAGE_BRUSH("Icons/icon_Landscape_Affects_VisibilityLayer_Disabled_16x", Icon16x16));
 
diff --git a/Engine/Source/Editor/LandscapeEditor/Private/LandscapeEditorDetailCustomization_LayersBrushStack.cpp b/Engine/Source/Editor/LandscapeEditor/Private/LandscapeEditorDetailCustomization_LayersBrushStack.cpp
--- a/Engine/Source/Editor/LandscapeEditor/Private/LandscapeEditorDetailCustomization_LayersBrushStack.cpp
+++ b/Engine/Source/Editor/LandscapeEditor/Private/LandscapeEditorDetailCustomization_LayersBrushStack.cpp
@@ -385,7 +385,7 @@ bool FLandscapeEditorCustomNodeBuilder_LayersBrushStack::IsAffectingHeightmap(in
 
 const FSlateBrush* FLandscapeEditorCustomNodeBuilder_LayersBrushStack::GetAffectsVisibilityLayerBrush(int32 InBrushIndex) const
 {
-	return IsAffectingVisibilityLayer(InBrushIndex) ? FAppStyle::GetBrush("LandscapeEditor.Brush.AffectsVisibilityLayer") : FAppStyle::GetBrush("LandscapeEditor.Brush.AffectsVisibilityLayer.Disabled");
+	return IsAffectingVisibilityLayer(InBrushIndex) ? FAppStyle::GetBrush("LandscapeEditor.Brush.AffectsVisibilityLayer.Enabled") : FAppStyle::GetBrush("LandscapeEditor.Brush.AffectsVisibilityLayer.Disabled");
 }
 
 bool FLandscapeEditorCustomNodeBuilder_LayersBrushStack::IsAffectingVisibilityLayer(int32 InBrushIndex) const
