Fixed landscape icons

#rnx

diff --git a/Engine/Source/Editor/EditorStyle/Private/StarshipStyle.cpp b/Engine/Source/Editor/EditorStyle/Private/StarshipStyle.cpp
--- a/Engine/Source/Editor/EditorStyle/Private/StarshipStyle.cpp
+++ b/Engine/Source/Editor/EditorStyle/Private/StarshipStyle.cpp
@@ -7233,8 +7233,8 @@ void FStarshipEditorStyle::FStyle::SetupLandscapeEditorStyle()
 			Set("LandscapeEditor.HydroErosionTool.Small", 		new IMAGE_BRUSH("Icons/Landscape/Landscape_HydroErosion_x40",      Icon20x20));
 			Set("LandscapeEditor.NoiseTool",        			new IMAGE_BRUSH("Icons/Landscape/Landscape_Noise_x40",             Icon20x20));
 			Set("LandscapeEditor.NoiseTool.Small",        		new IMAGE_BRUSH("Icons/Landscape/Landscape_Noise_x40",             Icon20x20));
-			Set("LandscapeEditor.VisibilityTool",   			new IMAGE_BRUSH_SVG("Icons/icon_Landscape_Target_Visibility_16x",  Icon20x20));
-			Set("LandscapeEditor.VisibilityTool.Small",   		new IMAGE_BRUSH_SVG("Icons/icon_Landscape_Target_Visibility_16x",  Icon20x20));
+			Set("LandscapeEditor.VisibilityTool",   			new IMAGE_BRUSH_SVG("Icons/Landscape/Landscape_VisibilityTool_x40",  Icon20x20));
+			Set("LandscapeEditor.VisibilityTool.Small",   		new IMAGE_BRUSH_SVG("Icons/Landscape/Landscape_VisibilityTool_x40",  Icon20x20));
 			Set("LandscapeEditor.BlueprintBrushTool", 			new IMAGE_BRUSH("Icons/Landscape/Landscape_BlueprintTool_x40",     Icon20x20));
 			Set("LandscapeEditor.BlueprintBrushTool.Small", 	new IMAGE_BRUSH("Icons/Landscape/Landscape_BlueprintTool_x40",     Icon20x20));
 
