Minor - adjust the height of the kebab icon in the Play toolbar to match its actual image dimensions and its height in other toolbars.

diff --git a/Engine/Source/Editor/EditorStyle/Private/StarshipStyle.cpp b/Engine/Source/Editor/EditorStyle/Private/StarshipStyle.cpp
--- a/Engine/Source/Editor/EditorStyle/Private/StarshipStyle.cpp
+++ b/Engine/Source/Editor/EditorStyle/Private/StarshipStyle.cpp
@@ -471,7 +471,7 @@ void FStarshipEditorStyle::FStyle::SetupGeneralStyles()
 		Set("Icons.ViewportScalability", new IMAGE_BRUSH_SVG("Starship/Common/ViewportScalability_16", Icon16x16));
 		Set("Icons.ViewportScalabilityReset", new IMAGE_BRUSH_SVG("Starship/Common/ViewportScalabilityReset_16", Icon16x16));
 		Set("Icons.EllipsisVerticalNarrow",
-			new CORE_IMAGE_BRUSH_SVG("Starship/Common/ellipsis-vertical-narrow", FVector2f(6.f, 24.f)));
+			new CORE_IMAGE_BRUSH_SVG("Starship/Common/ellipsis-vertical-narrow", FVector2f(6.f, 20.f)));
 
 		Set("Icons.Toolbar.Play", new IMAGE_BRUSH_SVG("Starship/Common/play", Icon20x20));
 		Set("Icons.Toolbar.Pause", new IMAGE_BRUSH_SVG("Starship/MainToolbar/pause", Icon20x20));
@@ -6160,7 +6160,7 @@ void FStarshipEditorStyle::FStyle::SetupPersonaStyle()
 		FComboButtonStyle PlayToolbarComboButton =
 			FComboButtonStyle(FStarshipCoreStyle::GetCoreStyle().GetWidgetStyle<FComboButtonStyle>("ComboButton"))
 				.SetDownArrowPadding(FMargin(1.f, 0.f, 2.f, 0.f))
-				.SetDownArrowImage(CORE_IMAGE_BRUSH_SVG("Starship/Common/ellipsis-vertical-narrow", FVector2f(6.f, 24.f)));
+				.SetDownArrowImage(CORE_IMAGE_BRUSH_SVG("Starship/Common/ellipsis-vertical-narrow", FVector2f(6.f, 20.f)));
 		PlayToolbarComboButton.ButtonStyle = RightToolbarComboButton;
 
 		MainToolbarRightComboButton.SetButtonStyle(RightToolbarComboButton);
