FStarshipEditorStyle : Add "Icons.OpenFile"

#FailedValidations Preflight

diff --git a/Engine/Source/Editor/EditorStyle/Private/StarshipStyle.cpp b/Engine/Source/Editor/EditorStyle/Private/StarshipStyle.cpp
--- a/Engine/Source/Editor/EditorStyle/Private/StarshipStyle.cpp
+++ b/Engine/Source/Editor/EditorStyle/Private/StarshipStyle.cpp
@@ -458,6 +458,7 @@ void FStarshipEditorStyle::FStyle::SetupGeneralStyles()
 		Set("Icons.OpenInExternalEditor", new IMAGE_BRUSH_SVG("Starship/Common/OpenInExternalEditor", Icon16x16));
 		Set("Icons.OpenSourceLocation", new IMAGE_BRUSH_SVG("Starship/Common/OpenSourceLocation", Icon16x16));
 		Set("Icons.OpenInBrowser", new IMAGE_BRUSH_SVG("Starship/Common/WebBrowser", Icon16x16));
+		Set("Icons.OpenFile", new IMAGE_BRUSH_SVG("Starship/Common/OpenFile", Icon16x16));
 		Set("Icons.Find", new IMAGE_BRUSH_SVG("Starship/Common/Find", Icon16x16));
 		Set("Icons.Validate", new IMAGE_BRUSH_SVG("Starship/Common/validate", Icon16x16));
 		Set("Icons.Pinned", new IMAGE_BRUSH_SVG("Starship/Common/Pinned", Icon16x16));
