Added missing EULA style

#rnx
#Synced-CL 49041953

diff --git a/Engine/Source/Editor/EditorStyle/Private/StarshipStyle.cpp b/Engine/Source/Editor/EditorStyle/Private/StarshipStyle.cpp
--- a/Engine/Source/Editor/EditorStyle/Private/StarshipStyle.cpp
+++ b/Engine/Source/Editor/EditorStyle/Private/StarshipStyle.cpp
@@ -339,9 +339,13 @@ void FStarshipEditorStyle::FStyle::SetupGeneralStyles()
 		Set("EULA.Header", FTextBlockStyle(FAppStyle::Get().GetWidgetStyle<FTextBlockStyle>("Log.Normal"))
 			.SetFont(FCoreStyle::GetDefaultFontStyle("Bold", 10)));
 
+		Set("EULA.Highlight", FTextBlockStyle(NormalText)
+			.SetFont(DEFAULT_FONT("Bold", FStarshipCoreStyle::RegularTextSize))
+			.SetColorAndOpacity(FLinearColor(1.0f, 1.0f, 1.0f)));
+
 		Set("EULA.HighlightItalic", FTextBlockStyle(NormalText)
 			.SetFont(DEFAULT_FONT("Italic", FStarshipCoreStyle::RegularTextSize))
-			.SetColorAndOpacity(FLinearColor(1.0f,1.0f, 1.0f)));
+			.SetColorAndOpacity(FLinearColor(1.0f, 1.0f, 1.0f)));
 
 		const FButtonStyle EulaHyperlinkButton = FButtonStyle()
 			.SetNormal(BORDER_BRUSH("Old/HyperlinkDotted", FMargin(0, 0, 0, 3 / 16.0f), FLinearColor(0.25f, 0.5f, 1.0f)))
