EditorUsability : EditorStyles
- Fixed horde issue : function SetupViewportStyles uses '304332' bytes of stack. Consider moving some data to heap.

#Synced-CL 49549735

diff --git a/Engine/Source/Editor/EditorStyle/Private/StarshipStyle.cpp b/Engine/Source/Editor/EditorStyle/Private/StarshipStyle.cpp
--- a/Engine/Source/Editor/EditorStyle/Private/StarshipStyle.cpp
+++ b/Engine/Source/Editor/EditorStyle/Private/StarshipStyle.cpp
@@ -2576,433 +2576,448 @@ void FStarshipEditorStyle::FStyle::SetupSequencerStyles()
 
 void FStarshipEditorStyle::FStyle::SetupViewportStyles()
 {
-	// New viewport toolbar.
+	SetupNewViewportToolbarStyles();
+	SetupOldViewportToolbarStyles();
+	SetupLegacyViewportToolbarStyles();
+
+	// Viewport actor preview's pin/unpin and attach/detach buttons
 	{
-		const FSlateRoundedBoxBrush TransparentBrush(FStyleColors::Transparent, 0.f, FStyleColors::Transparent, 0.f);
+		Set("ViewportActorPreview.Pinned", new IMAGE_BRUSH("Common/PushPin_Down", Icon16x16));
+		Set("ViewportActorPreview.Unpinned", new IMAGE_BRUSH("Common/PushPin_Up", Icon16x16));
+		Set("VRViewportActorPreview.Pinned", new IMAGE_BRUSH("Common/PushPin_Down_VR", Icon64x64));
+		Set("VRViewportActorPreview.Unpinned", new IMAGE_BRUSH("Common/PushPin_Up_VR", Icon64x64));
+		Set("VRViewportActorPreview.Attached", new IMAGE_BRUSH("Common/ScreenAttach_VR", Icon64x64));
+		Set("VRViewportActorPreview.Detached", new IMAGE_BRUSH("Common/ScreenDetach_VR", Icon64x64));
+	}
+}
+
+void FStarshipEditorStyle::FStyle::SetupNewViewportToolbarStyles()
+{
+	FSlateColor ToolbarBackgroundColor = FStyleColors::Dropdown;
+	const FSlateRoundedBoxBrush BackgroundBrush(ToolbarBackgroundColor, 5.f, ToolbarBackgroundColor, 1.0f);
+	const FSlateRoundedBoxBrush BackgroundBrushStart(ToolbarBackgroundColor, FVector4(5.f, 0.f, 0.f, 5.f), ToolbarBackgroundColor, 1.0f);
+	const FSlateRoundedBoxBrush BackgroundBrushMiddle(ToolbarBackgroundColor, FVector4(0.f, 0.f, 0.f, 0.f), ToolbarBackgroundColor, 1.0f);
+	const FSlateRoundedBoxBrush BackgroundBrushEnd(ToolbarBackgroundColor, FVector4(0.f, 5.f, 5.f, 0.f), ToolbarBackgroundColor, 1.0f);
+
+	FSlateColor ToolbarHoveredColor = FStyleColors::Hover;
+	const FSlateRoundedBoxBrush HoveredButtonBrush(ToolbarHoveredColor, 5.f, ToolbarHoveredColor, 1.0f);
+	const FSlateRoundedBoxBrush HoveredButtonBrushStart(ToolbarHoveredColor, FVector4(5.f, 0.f, 0.f, 5.f), ToolbarHoveredColor, 1.0f);
+	const FSlateRoundedBoxBrush HoveredButtonBrushMiddle(ToolbarHoveredColor, FVector4(0.f, 0.f, 0.f, 0.f), ToolbarHoveredColor, 1.0f);
+	const FSlateRoundedBoxBrush HoveredButtonBrushEnd(ToolbarHoveredColor, FVector4(0.f, 5.f, 5.f, 0.f), ToolbarHoveredColor, 1.0f);
+
+	FSlateColor ToolbarPressedColor = FStyleColors::Header;
+	const FSlateRoundedBoxBrush PressedButtonBrush(ToolbarPressedColor, 5.f, ToolbarPressedColor, 1.0f);
+	const FSlateRoundedBoxBrush PressedButtonBrushStart(ToolbarPressedColor, FVector4(5.f, 0.f, 0.f, 5.f), ToolbarPressedColor, 1.0f);
+	const FSlateRoundedBoxBrush PressedButtonBrushMiddle(ToolbarPressedColor, FVector4(0.f, 0.f, 0.f, 0.f), ToolbarPressedColor, 1.0f);
+	const FSlateRoundedBoxBrush PressedButtonBrushEnd(ToolbarPressedColor, FVector4(0.f, 5.f, 5.f, 0.f), ToolbarPressedColor, 1.0f);
+
+	FButtonStyle ButtonStyle;
+	{
+		ButtonStyle.SetNormal(BackgroundBrush)
+			.SetHovered(HoveredButtonBrush)
+			.SetPressed(PressedButtonBrush)
+			.SetNormalForeground(FStyleColors::ForegroundHeader)
+			.SetHoveredForeground(FStyleColors::ForegroundHover)
+			.SetPressedForeground(FStyleColors::ForegroundHover)
+			.SetDisabledForeground(FStyleColors::Hover2)
+			.SetNormalPadding(FMargin(6.0f, 4.0f))
+			.SetPressedPadding(FMargin(6.0f, 4.0f));
+	}
 	
-		FSlateColor ToolbarBackgroundColor = FStyleColors::Dropdown;
-		const FSlateRoundedBoxBrush BackgroundBrush(ToolbarBackgroundColor, 5.f, ToolbarBackgroundColor, 1.0f);
-		const FSlateRoundedBoxBrush BackgroundBrushStart(ToolbarBackgroundColor, FVector4(5.f, 0.f, 0.f, 5.f), ToolbarBackgroundColor, 1.0f);
-		const FSlateRoundedBoxBrush BackgroundBrushMiddle(ToolbarBackgroundColor, FVector4(0.f, 0.f, 0.f, 0.f), ToolbarBackgroundColor, 1.0f);
-		const FSlateRoundedBoxBrush BackgroundBrushEnd(ToolbarBackgroundColor, FVector4(0.f, 5.f, 5.f, 0.f), ToolbarBackgroundColor, 1.0f);
-
-		FSlateColor ToolbarHoveredColor = FStyleColors::Hover;
-		const FSlateRoundedBoxBrush HoveredButtonBrush(ToolbarHoveredColor, 5.f, ToolbarHoveredColor, 1.0f);
-		const FSlateRoundedBoxBrush HoveredButtonBrushStart(ToolbarHoveredColor, FVector4(5.f, 0.f, 0.f, 5.f), ToolbarHoveredColor, 1.0f);
-		const FSlateRoundedBoxBrush HoveredButtonBrushMiddle(ToolbarHoveredColor, FVector4(0.f, 0.f, 0.f, 0.f), ToolbarHoveredColor, 1.0f);
-		const FSlateRoundedBoxBrush HoveredButtonBrushEnd(ToolbarHoveredColor, FVector4(0.f, 5.f, 5.f, 0.f), ToolbarHoveredColor, 1.0f);
-
-		FSlateColor ToolbarPressedColor = FStyleColors::Header;
-		const FSlateRoundedBoxBrush PressedButtonBrush(ToolbarPressedColor, 5.f, ToolbarPressedColor, 1.0f);
-		const FSlateRoundedBoxBrush PressedButtonBrushStart(ToolbarPressedColor, FVector4(5.f, 0.f, 0.f, 5.f), ToolbarPressedColor, 1.0f);
-		const FSlateRoundedBoxBrush PressedButtonBrushMiddle(ToolbarPressedColor, FVector4(0.f, 0.f, 0.f, 0.f), ToolbarPressedColor, 1.0f);
-		const FSlateRoundedBoxBrush PressedButtonBrushEnd(ToolbarPressedColor, FVector4(0.f, 5.f, 5.f, 0.f), ToolbarPressedColor, 1.0f);
-
-		FButtonStyle ButtonStyle;
-		{
-			ButtonStyle.SetNormal(BackgroundBrush)
-				.SetHovered(HoveredButtonBrush)
-				.SetPressed(PressedButtonBrush)
-				.SetNormalForeground(FStyleColors::ForegroundHeader)
-				.SetHoveredForeground(FStyleColors::ForegroundHover)
-				.SetPressedForeground(FStyleColors::ForegroundHover)
-				.SetDisabledForeground(FStyleColors::Hover2)
-				.SetNormalPadding(FMargin(6.0f, 4.0f))
-				.SetPressedPadding(FMargin(6.0f, 4.0f));
-		}
+	FButtonStyle ButtonStyleStart = FButtonStyle(ButtonStyle)
+		.SetNormal(BackgroundBrushStart)
+		.SetHovered(HoveredButtonBrushStart)
+		.SetPressed(PressedButtonBrushStart);
+	
+	FButtonStyle ButtonStyleMiddle = FButtonStyle(ButtonStyle)
+		.SetNormal(BackgroundBrushMiddle)
+		.SetHovered(HoveredButtonBrushMiddle)
+		.SetPressed(PressedButtonBrushMiddle);
+	
+	FButtonStyle ButtonStyleEnd = FButtonStyle(ButtonStyle)
+		.SetNormal(BackgroundBrushEnd)
+		.SetHovered(HoveredButtonBrushEnd)
+		.SetPressed(PressedButtonBrushEnd);
+	
+
+	const FToolBarStyle SlimToolBarStyle =
+		FStarshipCoreStyle::GetCoreStyle().GetWidgetStyle<FToolBarStyle>("SlimToolBar");
+
+	FCheckBoxStyle ToggleButtonStyle = FCheckBoxStyle(SlimToolBarStyle.ToggleButton)
+		.SetCheckedImage(BackgroundBrush)
+		.SetUncheckedImage(BackgroundBrush)
+		.SetCheckedHoveredImage(HoveredButtonBrush)
+		.SetCheckedPressedImage(PressedButtonBrush)
+		.SetUncheckedPressedImage(PressedButtonBrush)
+		.SetCheckedForegroundColor(FStyleColors::AccentBlue)
+		.SetCheckedHoveredForegroundColor(FStyleColors::AccentBlue)
+		.SetCheckedPressedForegroundColor(FStyleColors::AccentBlue)
+		.SetPadding(FMargin(6.0f, 4.0f));
+	
+	FCheckBoxStyle ToggleButtonStyleStart = FCheckBoxStyle(ToggleButtonStyle)
+		.SetCheckedImage(BackgroundBrushStart)
+		.SetCheckedHoveredImage(HoveredButtonBrushStart)
+		.SetCheckedPressedImage(PressedButtonBrushStart)
+		.SetUncheckedImage(BackgroundBrushStart)
+		.SetUncheckedHoveredImage(HoveredButtonBrushStart)
+		.SetUncheckedPressedImage(PressedButtonBrushStart);
+	
+	FCheckBoxStyle ToggleButtonStyleMiddle = FCheckBoxStyle(ToggleButtonStyle)
+		.SetCheckedImage(BackgroundBrushMiddle)
+		.SetCheckedHoveredImage(HoveredButtonBrushMiddle)
+		.SetCheckedPressedImage(PressedButtonBrushMiddle)
+		.SetUncheckedImage(BackgroundBrushMiddle)
+		.SetUncheckedHoveredImage(HoveredButtonBrushMiddle)
+		.SetUncheckedPressedImage(PressedButtonBrushMiddle);
+	
+	FCheckBoxStyle ToggleButtonStyleEnd = FCheckBoxStyle(ToggleButtonStyle)
+		.SetCheckedImage(BackgroundBrushEnd)
+		.SetCheckedHoveredImage(HoveredButtonBrushEnd)
+		.SetCheckedPressedImage(PressedButtonBrushEnd)
+		.SetUncheckedImage(BackgroundBrushEnd)
+		.SetUncheckedHoveredImage(HoveredButtonBrushEnd)
+		.SetUncheckedPressedImage(PressedButtonBrushEnd);
+
+	FComboButtonStyle ComboButtonStyle = SlimToolBarStyle.ComboButtonStyle;
+	{
+		FButtonStyle InnerButtonStyle = FButtonStyle(ButtonStyle).SetNormalPadding(0.f).SetPressedPadding(0.f);
+
+		ComboButtonStyle
+			.SetButtonStyle(InnerButtonStyle)
+			.SetDownArrowImage(IMAGE_BRUSH_SVG("Starship/EditorViewport/small-chevron-down", FVector2f(6.f, 6.f)))
+			.SetDownArrowPadding(FMargin(4.0f, 0.f, 0.f, 0.f));
+	}
+
+	FToolBarStyle ViewportToolbarStyle = SlimToolBarStyle;
+	{
+		FSlateBrush SeparatorBrush;
+		SeparatorBrush.DrawAs = ESlateBrushDrawType::Box;
+		SeparatorBrush.TintColor = FStyleColors::Hover;
 		
-		FButtonStyle ButtonStyleStart = FButtonStyle(ButtonStyle)
-			.SetNormal(BackgroundBrushStart)
-			.SetHovered(HoveredButtonBrushStart)
-			.SetPressed(PressedButtonBrushStart);
+		FComboButtonStyle WrapButtonComboStyle = FComboButtonStyle(ComboButtonStyle)
+		.SetButtonStyle(FButtonStyle()
+			.SetNormal(FSlateColorBrush(FStyleColors::Dropdown))
+			.SetHovered(FSlateColorBrush(FStyleColors::Hover))
+			.SetPressed(FSlateColorBrush(FStyleColors::Header))
+			.SetNormalForeground(FStyleColors::Foreground)
+			.SetHoveredForeground(FStyleColors::ForegroundHover)
+			.SetPressedForeground(FStyleColors::ForegroundHover)
+			.SetDisabledForeground(FStyleColors::Hover2)
+			.SetNormalPadding(FMargin(6.0f, 4.0f))
+			.SetPressedPadding(FMargin(6.0f, 4.0f))
+		);
+
+		ViewportToolbarStyle.SetIconSize(Icon16x16)
+			.SetButtonStyle(ButtonStyle)
+			.SetToggleButtonStyle(ToggleButtonStyle)
+			.SetButtonPadding(FMargin(4.f, 0.f))
+			.SetComboButtonStyle(ComboButtonStyle)
+			.SetComboButtonPadding(FMargin(4.0f, 0.f))
+			.SetSeparatorBrush(SeparatorBrush)
+			.SetSeparatorThickness(1.0f)
+			.SetSeparatorPadding(FMargin(4.0f, 0.f))
+			.SetBackgroundPadding(FMargin(4.0f))
+			.SetWrapButtonStyle(
+				FWrapButtonStyle(SlimToolBarStyle.WrapButtonStyle)
+				.SetComboButtonStyle(WrapButtonComboStyle)
+				.SetWrapButtonIndex(-1)
+				.SetWrapButtonPadding(FMargin(4.f, -4.f, -4.f, -4.f))
+				.SetSeparatorBrush(SlimToolBarStyle.SeparatorBrush)
+				.SetSeparatorThickness(2.0f)
+			)
+			.SetRaisedChildrenRightPadding(16.0f);
+
+		Set("ViewportToolbar", ViewportToolbarStyle);
 		
-		FButtonStyle ButtonStyleMiddle = FButtonStyle(ButtonStyle)
-			.SetNormal(BackgroundBrushMiddle)
-			.SetHovered(HoveredButtonBrushMiddle)
-			.SetPressed(PressedButtonBrushMiddle);
+		Set("ViewportToolbar.Button.Start", ButtonStyleStart);
+		Set("ViewportToolbar.Button.Middle", ButtonStyleMiddle);
+		Set("ViewportToolbar.Button.End", ButtonStyleEnd);
 		
-		FButtonStyle ButtonStyleEnd = FButtonStyle(ButtonStyle)
-			.SetNormal(BackgroundBrushEnd)
-			.SetHovered(HoveredButtonBrushEnd)
-			.SetPressed(PressedButtonBrushEnd);
+		Set("ViewportToolbar.ToggleButton.Start", ToggleButtonStyleStart);
+		Set("ViewportToolbar.ToggleButton.Middle", ToggleButtonStyleMiddle);
+		Set("ViewportToolbar.ToggleButton.End", ToggleButtonStyleEnd);
 		
+		// Ensure split buttons don't have rounded corners on the inside
+		Set("ViewportToolbar.SplitToggleButton", ToggleButtonStyleStart);
+		FComboButtonStyle SplitComboButton = ViewportToolbarStyle.ComboButtonStyle;
+		SplitComboButton.ButtonStyle = ButtonStyleEnd;
+		Set("ViewportToolbar.SplitComboButton", SplitComboButton);
+	}
+
+	// Variant with no background
+	{
+		FToolBarStyle ViewportToolbarNoBackground(ViewportToolbarStyle);
+		ViewportToolbarNoBackground.SetBackground(FSlateNoResource());
+		ViewportToolbarNoBackground.BackgroundPadding = FMargin(0.f);
+		Set("NoBackgroundViewportToolbar", ViewportToolbarNoBackground);
+		Set("NoBackgroundViewportToolbar.Button.Start", ButtonStyleStart);
+        Set("NoBackgroundViewportToolbar.Button.Middle", ButtonStyleMiddle);
+        Set("NoBackgroundViewportToolbar.Button.End", ButtonStyleEnd);
+        
+        Set("NoBackgroundViewportToolbar.ToggleButton.Start", ToggleButtonStyleStart);
+        Set("NoBackgroundViewportToolbar.ToggleButton.Middle", ToggleButtonStyleMiddle);
+        Set("NoBackgroundViewportToolbar.ToggleButton.End", ToggleButtonStyleEnd);
+        
+        // Ensure split buttons don't have rounded corners on the inside
+        Set("NoBackgroundViewportToolbar.SplitToggleButton", ToggleButtonStyleStart);
+        FComboButtonStyle SplitComboButton = ViewportToolbarStyle.ComboButtonStyle;
+        SplitComboButton.ButtonStyle = ButtonStyleEnd;
+        Set("NoBackgroundViewportToolbar.SplitComboButton", SplitComboButton);
+	}
+
+	// Special styling for top-level raised buttons.
+	FToolBarStyle ViewportToolbarRaisedStyle = ViewportToolbarStyle;
+	{
+		const FSlateRoundedBoxBrush TransparentBrush(FStyleColors::Transparent, 0.f, FStyleColors::Transparent, 0.f);
 
-		const FToolBarStyle SlimToolBarStyle =
-			FStarshipCoreStyle::GetCoreStyle().GetWidgetStyle<FToolBarStyle>("SlimToolBar");
-
-		FCheckBoxStyle ToggleButtonStyle = FCheckBoxStyle(SlimToolBarStyle.ToggleButton)
-			.SetCheckedImage(BackgroundBrush)
-			.SetUncheckedImage(BackgroundBrush)
-			.SetCheckedHoveredImage(HoveredButtonBrush)
-			.SetCheckedPressedImage(PressedButtonBrush)
-			.SetUncheckedPressedImage(PressedButtonBrush)
-			.SetCheckedForegroundColor(FStyleColors::AccentBlue)
-			.SetCheckedHoveredForegroundColor(FStyleColors::AccentBlue)
-			.SetCheckedPressedForegroundColor(FStyleColors::AccentBlue)
-			.SetPadding(FMargin(6.0f, 4.0f));
-		
-		FCheckBoxStyle ToggleButtonStyleStart = FCheckBoxStyle(ToggleButtonStyle)
-			.SetCheckedImage(BackgroundBrushStart)
-			.SetCheckedHoveredImage(HoveredButtonBrushStart)
-			.SetCheckedPressedImage(PressedButtonBrushStart)
-			.SetUncheckedImage(BackgroundBrushStart)
-			.SetUncheckedHoveredImage(HoveredButtonBrushStart)
-			.SetUncheckedPressedImage(PressedButtonBrushStart);
-		
-		FCheckBoxStyle ToggleButtonStyleMiddle = FCheckBoxStyle(ToggleButtonStyle)
-			.SetCheckedImage(BackgroundBrushMiddle)
-			.SetCheckedHoveredImage(HoveredButtonBrushMiddle)
-			.SetCheckedPressedImage(PressedButtonBrushMiddle)
-			.SetUncheckedImage(BackgroundBrushMiddle)
-			.SetUncheckedHoveredImage(HoveredButtonBrushMiddle)
-			.SetUncheckedPressedImage(PressedButtonBrushMiddle);
+		const auto& ConvertButtonToRaised = [&TransparentBrush](const FButtonStyle& ButtonStyle)
+		{
+			return FButtonStyle(ButtonStyle)
+				.SetNormal(TransparentBrush)
+				.SetNormalForeground(FStyleColors::Foreground)
+				.SetNormalPadding(4.0f)
+				.SetPressedPadding(4.0f);
+		};
+	
+		FButtonStyle ButtonStyleRaised = ConvertButtonToRaised(ButtonStyle);
+		FButtonStyle ButtonStyleRaisedStart = ConvertButtonToRaised(ButtonStyleStart);
+		FButtonStyle ButtonStyleRaisedMiddle = ConvertButtonToRaised(ButtonStyleMiddle);
+		FButtonStyle ButtonStyleRaisedEnd = ConvertButtonToRaised(ButtonStyleEnd);
 		
-		FCheckBoxStyle ToggleButtonStyleEnd = FCheckBoxStyle(ToggleButtonStyle)
-			.SetCheckedImage(BackgroundBrushEnd)
-			.SetCheckedHoveredImage(HoveredButtonBrushEnd)
-			.SetCheckedPressedImage(PressedButtonBrushEnd)
-			.SetUncheckedImage(BackgroundBrushEnd)
-			.SetUncheckedHoveredImage(HoveredButtonBrushEnd)
-			.SetUncheckedPressedImage(PressedButtonBrushEnd);
-
-		FComboButtonStyle ComboButtonStyle = SlimToolBarStyle.ComboButtonStyle;
+		const auto ConvertToggleButtonToRaised = [&TransparentBrush](const FCheckBoxStyle& ButtonStyle)
 		{
-			FButtonStyle InnerButtonStyle = FButtonStyle(ButtonStyle).SetNormalPadding(0.f).SetPressedPadding(0.f);
-
-			ComboButtonStyle
-				.SetButtonStyle(InnerButtonStyle)
-				.SetDownArrowImage(IMAGE_BRUSH_SVG("Starship/EditorViewport/small-chevron-down", FVector2f(6.f, 6.f)))
-				.SetDownArrowPadding(FMargin(4.0f, 0.f, 0.f, 0.f));
-		}
+			return FCheckBoxStyle(ButtonStyle)
+				.SetPadding(4.0f)
+				.SetForegroundColor(FStyleColors::Foreground)
+				.SetUncheckedImage(TransparentBrush)
+				.SetCheckedImage(TransparentBrush);
+		};
 		
-		FComboButtonStyle WrapButtonComboStyle = FComboButtonStyle(ComboButtonStyle)
-			.SetButtonStyle(FButtonStyle()
-				.SetNormal(FSlateColorBrush(FStyleColors::Dropdown))
-				.SetHovered(FSlateColorBrush(FStyleColors::Hover))
-				.SetPressed(FSlateColorBrush(FStyleColors::Header))
-				.SetNormalForeground(FStyleColors::Foreground)
-				.SetHoveredForeground(FStyleColors::ForegroundHover)
-				.SetPressedForeground(FStyleColors::ForegroundHover)
-				.SetDisabledForeground(FStyleColors::Hover2)
-				.SetNormalPadding(FMargin(6.0f, 4.0f))
-				.SetPressedPadding(FMargin(6.0f, 4.0f))
-			);
+		FCheckBoxStyle ToggleButtonStyleRaised = ConvertToggleButtonToRaised(ToggleButtonStyle);
+		FCheckBoxStyle ToggleButtonStyleRaisedStart = ConvertToggleButtonToRaised(ToggleButtonStyleStart);
+		FCheckBoxStyle ToggleButtonStyleRaisedMiddle = ConvertToggleButtonToRaised(ToggleButtonStyleMiddle);
+		FCheckBoxStyle ToggleButtonStyleRaisedEnd = ConvertToggleButtonToRaised(ToggleButtonStyleEnd);
 
-		FToolBarStyle ViewportToolbarStyle = SlimToolBarStyle;
-		{
-			FSlateBrush SeparatorBrush;
-			SeparatorBrush.DrawAs = ESlateBrushDrawType::Box;
-			SeparatorBrush.TintColor = FStyleColors::Hover;
-
-			ViewportToolbarStyle.SetIconSize(Icon16x16)
-				.SetButtonStyle(ButtonStyle)
-				.SetToggleButtonStyle(ToggleButtonStyle)
-				.SetButtonPadding(FMargin(4.f, 0.f))
-				.SetComboButtonStyle(ComboButtonStyle)
-				.SetComboButtonPadding(FMargin(4.0f, 0.f))
-				.SetSeparatorBrush(SeparatorBrush)
-				.SetSeparatorThickness(1.0f)
-				.SetSeparatorPadding(FMargin(4.0f, 0.f))
-				.SetBackgroundPadding(FMargin(4.0f))
-				.SetWrapButtonStyle(
-					FWrapButtonStyle(SlimToolBarStyle.WrapButtonStyle)
-					.SetComboButtonStyle(WrapButtonComboStyle)
-					.SetWrapButtonIndex(-1)
-					.SetWrapButtonPadding(FMargin(4.f, -4.f, -4.f, -4.f))
-					.SetSeparatorBrush(SlimToolBarStyle.SeparatorBrush)
-					.SetSeparatorThickness(2.0f)
-				)
-				.SetRaisedChildrenRightPadding(16.0f);
-
-			Set("ViewportToolbar", ViewportToolbarStyle);
-			
-			Set("ViewportToolbar.Button.Start", ButtonStyleStart);
-			Set("ViewportToolbar.Button.Middle", ButtonStyleMiddle);
-			Set("ViewportToolbar.Button.End", ButtonStyleEnd);
-			
-			Set("ViewportToolbar.ToggleButton.Start", ToggleButtonStyleStart);
-			Set("ViewportToolbar.ToggleButton.Middle", ToggleButtonStyleMiddle);
-			Set("ViewportToolbar.ToggleButton.End", ToggleButtonStyleEnd);
-			
-			// Ensure split buttons don't have rounded corners on the inside
-			Set("ViewportToolbar.SplitToggleButton", ToggleButtonStyleStart);
-			FComboButtonStyle SplitComboButton = ViewportToolbarStyle.ComboButtonStyle;
-			SplitComboButton.ButtonStyle = ButtonStyleEnd;
-			Set("ViewportToolbar.SplitComboButton", SplitComboButton);
-		}
+		FComboButtonStyle ComboStyleRaised = FComboButtonStyle(ViewportToolbarStyle.ComboButtonStyle)
+			.SetDownArrowPadding(FMargin(2.0f, 0, 0, 0));
 
-		// Variant with no background
+		FComboButtonStyle SettingsComboStyleRaised = ViewportToolbarStyle.SettingsComboButton;
 		{
-			FToolBarStyle ViewportToolbarNoBackground(ViewportToolbarStyle);
-			ViewportToolbarNoBackground.SetBackground(FSlateNoResource());
-			ViewportToolbarNoBackground.BackgroundPadding = FMargin(0.f);
-			Set("NoBackgroundViewportToolbar", ViewportToolbarNoBackground);
-			Set("NoBackgroundViewportToolbar.Button.Start", ButtonStyleStart);
-            Set("NoBackgroundViewportToolbar.Button.Middle", ButtonStyleMiddle);
-            Set("NoBackgroundViewportToolbar.Button.End", ButtonStyleEnd);
-            
-            Set("NoBackgroundViewportToolbar.ToggleButton.Start", ToggleButtonStyleStart);
-            Set("NoBackgroundViewportToolbar.ToggleButton.Middle", ToggleButtonStyleMiddle);
-            Set("NoBackgroundViewportToolbar.ToggleButton.End", ToggleButtonStyleEnd);
-            
-            // Ensure split buttons don't have rounded corners on the inside
-            Set("NoBackgroundViewportToolbar.SplitToggleButton", ToggleButtonStyleStart);
-            FComboButtonStyle SplitComboButton = ViewportToolbarStyle.ComboButtonStyle;
-            SplitComboButton.ButtonStyle = ButtonStyleEnd;
-            Set("NoBackgroundViewportToolbar.SplitComboButton", SplitComboButton);
+			SettingsComboStyleRaised.ButtonStyle
+				.SetNormalPadding(FMargin(2.f, 0.f))
+				.SetPressedPadding(FMargin(2.f, 0.f))
+				.SetHovered(HoveredButtonBrush)
+				.SetHoveredForeground(FStyleColors::ForegroundHover);
 		}
 
-		// Special styling for top-level raised buttons.
-		FToolBarStyle ViewportToolbarRaisedStyle = ViewportToolbarStyle;
-		{
-			const auto& ConvertButtonToRaised = [&TransparentBrush](const FButtonStyle& ButtonStyle)
-			{
-				return FButtonStyle(ButtonStyle)
-					.SetNormal(TransparentBrush)
-					.SetNormalForeground(FStyleColors::Foreground)
-					.SetNormalPadding(4.0f)
-					.SetPressedPadding(4.0f);
-			};
+		ViewportToolbarRaisedStyle
+			.SetButtonStyle(ButtonStyleRaised)
+			.SetToggleButtonStyle(ToggleButtonStyleRaised)
+			.SetComboButtonStyle(ComboStyleRaised)
+			.SetBlockHovered(BackgroundBrush)
+			// The offset gap allows raised buttons to be 6px away from themselves & their top-level parents
+			.SetButtonPadding(FMargin(2.f, 0.f, 4.f, 0.f))
+			.SetComboButtonPadding(FMargin(2.f, 0.f, 4.f, 0.f))
+			.SetSettingsComboButtonStyle(SettingsComboStyleRaised);
 		
-			FButtonStyle ButtonStyleRaised = ConvertButtonToRaised(ButtonStyle);
-			FButtonStyle ButtonStyleRaisedStart = ConvertButtonToRaised(ButtonStyleStart);
-			FButtonStyle ButtonStyleRaisedMiddle = ConvertButtonToRaised(ButtonStyleMiddle);
-			FButtonStyle ButtonStyleRaisedEnd = ConvertButtonToRaised(ButtonStyleEnd);
-			
-			const auto ConvertToggleButtonToRaised = [&TransparentBrush](const FCheckBoxStyle& ButtonStyle)
-			{
-				return FCheckBoxStyle(ButtonStyle)
-					.SetPadding(4.0f)
-					.SetForegroundColor(FStyleColors::Foreground)
-					.SetUncheckedImage(TransparentBrush)
-					.SetCheckedImage(TransparentBrush);
-			};
-			
-			FCheckBoxStyle ToggleButtonStyleRaised = ConvertToggleButtonToRaised(ToggleButtonStyle);
-			FCheckBoxStyle ToggleButtonStyleRaisedStart = ConvertToggleButtonToRaised(ToggleButtonStyleStart);
-			FCheckBoxStyle ToggleButtonStyleRaisedMiddle = ConvertToggleButtonToRaised(ToggleButtonStyleMiddle);
-			FCheckBoxStyle ToggleButtonStyleRaisedEnd = ConvertToggleButtonToRaised(ToggleButtonStyleEnd);
-
-			FComboButtonStyle ComboStyleRaised = FComboButtonStyle(ViewportToolbarStyle.ComboButtonStyle)
-				.SetDownArrowPadding(FMargin(2.0f, 0, 0, 0));
-
-			FComboButtonStyle SettingsComboStyleRaised = ViewportToolbarStyle.SettingsComboButton;
-			{
-				SettingsComboStyleRaised.ButtonStyle
-					.SetNormalPadding(FMargin(2.f, 0.f))
-					.SetPressedPadding(FMargin(2.f, 0.f))
-					.SetHovered(HoveredButtonBrush)
-					.SetHoveredForeground(FStyleColors::ForegroundHover);
-			}
-
-			ViewportToolbarRaisedStyle
-				.SetButtonStyle(ButtonStyleRaised)
-				.SetToggleButtonStyle(ToggleButtonStyleRaised)
-				.SetComboButtonStyle(ComboStyleRaised)
-				.SetBlockHovered(BackgroundBrush)
-				// The offset gap allows raised buttons to be 6px away from themselves & their top-level parents
-				.SetButtonPadding(FMargin(2.f, 0.f, 4.f, 0.f))
-				.SetComboButtonPadding(FMargin(2.f, 0.f, 4.f, 0.f))
-				.SetSettingsComboButtonStyle(SettingsComboStyleRaised);
-			
-			Set("ViewportToolbar.Raised", ViewportToolbarRaisedStyle);
-			
-			Set("ViewportToolbar.Raised.Button.Start", ButtonStyleRaisedStart);
-			Set("ViewportToolbar.Raised.Button.Middle", ButtonStyleRaisedMiddle);
-			Set("ViewportToolbar.Raised.Button.End", ButtonStyleRaisedEnd);
-			
-			Set("ViewportToolbar.Raised.ToggleButton.Start", ToggleButtonStyleRaisedStart);
-			Set("ViewportToolbar.Raised.ToggleButton.Middle", ToggleButtonStyleRaisedMiddle);
-			Set("ViewportToolbar.Raised.ToggleButton.End", ToggleButtonStyleRaisedEnd);
-			
-			// Ensure split buttons don't have rounded corners on the inside
-			Set("ViewportToolbar.Raised.SplitToggleButton", ToggleButtonStyleRaisedStart);
-			FComboButtonStyle SplitComboButton = ViewportToolbarRaisedStyle.ComboButtonStyle;
-			SplitComboButton.ButtonStyle = ButtonStyleRaisedEnd;
-			Set("ViewportToolbar.Raised.SplitComboButton", SplitComboButton);
-		}
-
-		// Special style for raised transform buttons
-		{
-			FToolBarStyle TransformToolsParent = FToolBarStyle(GetWidgetStyle<FToolBarStyle>("ViewportToolbar"))
-				// The raised children all have no padding, so the raised gap padding needs to be extended
-				// to make the sizing consistent.
-				// The design target is 24px from the last raised item to the next top-level item.
-				.SetRaisedChildrenRightPadding(20.f);
+		Set("ViewportToolbar.Raised", ViewportToolbarRaisedStyle);
 		
-			FToolBarStyle RaisedTransformTools = FToolBarStyle(GetWidgetStyle<FToolBarStyle>("ViewportToolbar.Raised"))
-				.SetButtonPadding(0.f)
-				.SetComboButtonPadding(0.f);
-				
-			Set("ViewportToolbar.TransformTools", TransformToolsParent);
-			Set("ViewportToolbar.TransformTools.Raised", RaisedTransformTools);
-		}
-
-		// Special styling for warnings.
-		{
-			// Use normal styling for non-raised buttons.
-			Set("ViewportToolbarWarning", ViewportToolbarStyle);
+		Set("ViewportToolbar.Raised.Button.Start", ButtonStyleRaisedStart);
+		Set("ViewportToolbar.Raised.Button.Middle", ButtonStyleRaisedMiddle);
+		Set("ViewportToolbar.Raised.Button.End", ButtonStyleRaisedEnd);
+		
+		Set("ViewportToolbar.Raised.ToggleButton.Start", ToggleButtonStyleRaisedStart);
+		Set("ViewportToolbar.Raised.ToggleButton.Middle", ToggleButtonStyleRaisedMiddle);
+		Set("ViewportToolbar.Raised.ToggleButton.End", ToggleButtonStyleRaisedEnd);
+		
+		// Ensure split buttons don't have rounded corners on the inside
+		Set("ViewportToolbar.Raised.SplitToggleButton", ToggleButtonStyleRaisedStart);
+		FComboButtonStyle SplitComboButton = ViewportToolbarRaisedStyle.ComboButtonStyle;
+		SplitComboButton.ButtonStyle = ButtonStyleRaisedEnd;
+		Set("ViewportToolbar.Raised.SplitComboButton", SplitComboButton);
+	}
 
-			{
-				FLinearColor WarningColor = FStyleColors::Warning.GetSpecifiedColor();
-				const FSlateRoundedBoxBrush WarningBrush(WarningColor, 5.0f, WarningColor, 1.0f);
-
-				FLinearColor WarningHoveredColor = FStyleColors::Warning.GetSpecifiedColor();
-				WarningHoveredColor.R = FMath::Min(1.0f, WarningHoveredColor.R * 1.5f);
-				WarningHoveredColor.G = FMath::Min(1.0f, WarningHoveredColor.G * 1.5f);
-				WarningHoveredColor.B = FMath::Min(1.0f, WarningHoveredColor.B * 1.5f);
-				const FSlateRoundedBoxBrush WarningHoveredBrush(WarningHoveredColor, 5.0f, WarningHoveredColor, 1.0f);
-
-				FLinearColor WarningPressedColor = FStyleColors::Warning.GetSpecifiedColor();
-				WarningPressedColor.A = .50f;
-				const FSlateRoundedBoxBrush WarningPressedBrush(WarningPressedColor, 5.0f, WarningPressedColor, 1.0f);
-
-				const FCheckBoxStyle ToggleButtonStyleRaisedWarning =
-					FCheckBoxStyle(ViewportToolbarRaisedStyle.ToggleButton)
-						.SetBackgroundImage(WarningBrush)
-						.SetBackgroundHoveredImage(WarningHoveredBrush)
-						.SetUncheckedImage(WarningBrush)
-						.SetUncheckedHoveredImage(WarningHoveredBrush)
-						.SetUncheckedPressedImage(WarningPressedBrush)
-						.SetCheckedImage(WarningBrush)
-						.SetCheckedHoveredImage(WarningHoveredBrush)
-						.SetCheckedPressedImage(WarningPressedBrush)
-						.SetForegroundColor(FStyleColors::ForegroundInverted)
-						.SetHoveredForegroundColor(FStyleColors::ForegroundInverted)
-						.SetUndeterminedForegroundColor(FStyleColors::ForegroundInverted)
-						.SetPressedForegroundColor(FStyleColors::ForegroundInverted);
-
-				const FButtonStyle ButtonStyleRaisedWarning = FButtonStyle(ViewportToolbarRaisedStyle.ButtonStyle)
-																  .SetHovered(WarningHoveredBrush)
-																  .SetPressed(WarningPressedBrush)
-																  .SetNormal(WarningBrush)
-																  .SetNormalForeground(FStyleColors::ForegroundInverted)
-																  .SetHoveredForeground(FStyleColors::ForegroundInverted)
-																  .SetPressedForeground(FStyleColors::ForegroundInverted);
-
-				const FToolBarStyle NewViewportToolbarTopLevelWarningRaisedStyle =
-					FToolBarStyle(ViewportToolbarRaisedStyle)
-						.SetToggleButtonStyle(ToggleButtonStyleRaisedWarning)
-						.SetButtonStyle(ButtonStyleRaisedWarning);
-
-				Set("ViewportToolbarWarning.Raised", NewViewportToolbarTopLevelWarningRaisedStyle);
-			}
-		}
+	// Special style for raised transform buttons
+	{
+		FToolBarStyle TransformToolsParent = FToolBarStyle(GetWidgetStyle<FToolBarStyle>("ViewportToolbar"))
+			// The raised children all have no padding, so the raised gap padding needs to be extended
+			// to make the sizing consistent.
+			// The design target is 24px from the last raised item to the next top-level item.
+			.SetRaisedChildrenRightPadding(20.f);
+	
+		FToolBarStyle RaisedTransformTools = FToolBarStyle(GetWidgetStyle<FToolBarStyle>("ViewportToolbar.Raised"))
+			.SetButtonPadding(0.f)
+			.SetComboButtonPadding(0.f);
+			
+		Set("ViewportToolbar.TransformTools", TransformToolsParent);
+		Set("ViewportToolbar.TransformTools.Raised", RaisedTransformTools);
+	}
 
-		// Special style for the "Viewport Sizing" submenu, with rounded corners on the left.
-		{
-			const FButtonStyle ViewportSizingButtonStyle =
-				FButtonStyle(ButtonStyleStart)
-					.SetNormalPadding(FMargin(0.f))
-					.SetPressedPadding(FMargin(0.f));
-
-			const FComboButtonStyle ViewportSizingComboButtonStyle =
-				FComboButtonStyle(ComboButtonStyle)
-					.SetDownArrowPadding(FMargin(5.f, 0.f, 5.f, 0.f))
-					.SetDownArrowImage(CORE_IMAGE_BRUSH_SVG("Starship/Common/ellipsis-vertical-narrow", FVector2f(6.f, 20.f)))
-					.SetButtonStyle(ViewportSizingButtonStyle);
-
-			const FToolBarStyle ViewportSizingStyle = FToolBarStyle(ViewportToolbarStyle)
-				.SetButtonStyle(ViewportSizingButtonStyle)
-				.SetComboButtonStyle(ViewportSizingComboButtonStyle)
-				.SetComboButtonPadding(FMargin(4.0f, 0, 1.0f, 0))
-				.SetSeparatorPadding(FMargin(0.f));
-
-			Set("ViewportToolbarViewportSizingSubmenu", ViewportSizingStyle);
-
-			// A raised button keeps the gray background here, with rounded corners on the right.
-			{
-				FCheckBoxStyle ToggleButtonStyleRaised = FCheckBoxStyle(ToggleButtonStyleEnd)
-					// The icon cannot be perfectly centered, so it needs a 1px shove to the right
-					// to match the design spec.
-					.SetPadding(FMargin(7.0f, 4.0f, 6.0f, 4.0f));
-					
-				const FToolBarStyle ViewportSizingRaisedStyle = FToolBarStyle(ViewportToolbarStyle)
-					.SetToggleButtonStyle(ToggleButtonStyleRaised)
-					.SetButtonPadding(FMargin(0.f, 0.f, 4.f, 0.f))
-					.SetComboButtonPadding(FMargin(2.0f, 0.f));
-
-				Set("ViewportToolbarViewportSizingSubmenu.Raised", ViewportSizingRaisedStyle);
-			}
-		}
+	// Special styling for warnings.
+	{
+		// Use normal styling for non-raised buttons.
+		Set("ViewportToolbarWarning", ViewportToolbarStyle);
 
-		// Viewport Toolbar Icons
 		{
-			Set("ViewportToolbar.CameraSpeed", new IMAGE_BRUSH_SVG("Starship/Common/CameraSpeed_16", Icon16x16));
-			Set("ViewportToolbar.EV100", new IMAGE_BRUSH_SVG("Starship/Common/EV100_16", Icon16x16));
-			Set("ViewportToolbar.ExactCameraView", new IMAGE_BRUSH_SVG("Starship/Common/ExactCameraView_16", Icon16x16));
-			Set("ViewportToolbar.Exposure", new IMAGE_BRUSH_SVG("Starship/Common/Exposure_16", Icon16x16));
-			Set("ViewportToolbar.FarViewPlane", new IMAGE_BRUSH_SVG("Starship/Common/FarViewPlane_16", Icon16x16));
-			Set("ViewportToolbar.FieldOfView", new IMAGE_BRUSH_SVG("Starship/Common/FieldOfView_16", Icon16x16));
-			Set("ViewportToolbar.GizmoScale", new IMAGE_BRUSH_SVG("Starship/Common/GizmoScale_16", Icon16x16));
-			Set("ViewportToolbar.NearViewPlane", new IMAGE_BRUSH_SVG("Starship/Common/NearViewPlane_16", Icon16x16));
-			Set("ViewportToolbar.PreviewSceneSettings", new IMAGE_BRUSH_SVG("Starship/AssetEditors/PreviewSceneSettings_16", Icon16x16));
-			Set("ViewportToolbar.SetShowGrid", new IMAGE_BRUSH_SVG("Starship/Common/Grid", Icon16x16));
-			Set("ViewportToolbar.Snap", new IMAGE_BRUSH_SVG("Starship/Common/Snap_16", Icon16x16));
-			Set("ViewportToolbar.SnapLocation", new IMAGE_BRUSH_SVG("Starship/Common/SnapLocation_16", Icon16x16));
-			Set("ViewportToolbar.SnapPlanar", new IMAGE_BRUSH_SVG("Starship/Common/SnapPlanar_16", Icon16x16));
-			Set("ViewportToolbar.SnapRotation", new IMAGE_BRUSH_SVG("Starship/Common/SnapRotation_16", Icon16x16));
-			Set("ViewportToolbar.SnapScale", new IMAGE_BRUSH_SVG("Starship/Common/SnapScale_16", Icon16x16));
-			Set("ViewportToolbar.SpeedScalar", new IMAGE_BRUSH_SVG("Starship/Common/SpeedScalar_16", Icon16x16));
-			Set("ViewportToolbar.SurfaceSnap", new IMAGE_BRUSH_SVG("Starship/Common/SurfaceSnap_16", Icon16x16));
-			Set("ViewportToolbar.SurfaceSnapRotateToNormal", new IMAGE_BRUSH_SVG("Starship/Common/SurfaceSnapRotateToNormal_16", Icon16x16));
-			Set("ViewportToolbar.TransformRotate", new IMAGE_BRUSH_SVG("Starship/Common/TransformRotate_16", Icon16x16));
-			Set("ViewportToolbar.TransformScale", new IMAGE_BRUSH_SVG("Starship/Common/TransformScale_16", Icon16x16));
-			Set("ViewportToolbar.TransformSelect", new IMAGE_BRUSH_SVG("Starship/Common/TransformSelect_16", Icon16x16));
-			Set("ViewportToolbar.TransformTranslate", new IMAGE_BRUSH_SVG("Starship/Common/TransformMove_16", Icon16x16));
-			Set("ViewportToolbar.VPCamera", new IMAGE_BRUSH_SVG("Starship/Common/VPCamera_16", Icon16x16));
-			Set("ViewportToolbar.VPCineCamera", new IMAGE_BRUSH_SVG("Starship/Common/VPCineCamera_16", Icon16x16));
+			FLinearColor WarningColor = FStyleColors::Warning.GetSpecifiedColor();
+			const FSlateRoundedBoxBrush WarningBrush(WarningColor, 5.0f, WarningColor, 1.0f);
+
+			FLinearColor WarningHoveredColor = FStyleColors::Warning.GetSpecifiedColor();
+			WarningHoveredColor.R = FMath::Min(1.0f, WarningHoveredColor.R * 1.5f);
+			WarningHoveredColor.G = FMath::Min(1.0f, WarningHoveredColor.G * 1.5f);
+			WarningHoveredColor.B = FMath::Min(1.0f, WarningHoveredColor.B * 1.5f);
+			const FSlateRoundedBoxBrush WarningHoveredBrush(WarningHoveredColor, 5.0f, WarningHoveredColor, 1.0f);
+
+			FLinearColor WarningPressedColor = FStyleColors::Warning.GetSpecifiedColor();
+			WarningPressedColor.A = .50f;
+			const FSlateRoundedBoxBrush WarningPressedBrush(WarningPressedColor, 5.0f, WarningPressedColor, 1.0f);
+
+			const FCheckBoxStyle ToggleButtonStyleRaisedWarning =
+				FCheckBoxStyle(ViewportToolbarRaisedStyle.ToggleButton)
+					.SetBackgroundImage(WarningBrush)
+					.SetBackgroundHoveredImage(WarningHoveredBrush)
+					.SetUncheckedImage(WarningBrush)
+					.SetUncheckedHoveredImage(WarningHoveredBrush)
+					.SetUncheckedPressedImage(WarningPressedBrush)
+					.SetCheckedImage(WarningBrush)
+					.SetCheckedHoveredImage(WarningHoveredBrush)
+					.SetCheckedPressedImage(WarningPressedBrush)
+					.SetForegroundColor(FStyleColors::ForegroundInverted)
+					.SetHoveredForegroundColor(FStyleColors::ForegroundInverted)
+					.SetUndeterminedForegroundColor(FStyleColors::ForegroundInverted)
+					.SetPressedForegroundColor(FStyleColors::ForegroundInverted);
+
+			const FButtonStyle ButtonStyleRaisedWarning = FButtonStyle(ViewportToolbarRaisedStyle.ButtonStyle)
+															  .SetHovered(WarningHoveredBrush)
+															  .SetPressed(WarningPressedBrush)
+															  .SetNormal(WarningBrush)
+															  .SetNormalForeground(FStyleColors::ForegroundInverted)
+															  .SetHoveredForeground(FStyleColors::ForegroundInverted)
+															  .SetPressedForeground(FStyleColors::ForegroundInverted);
+
+			const FToolBarStyle NewViewportToolbarTopLevelWarningRaisedStyle =
+				FToolBarStyle(ViewportToolbarRaisedStyle)
+					.SetToggleButtonStyle(ToggleButtonStyleRaisedWarning)
+					.SetButtonStyle(ButtonStyleRaisedWarning);
+
+			Set("ViewportToolbarWarning.Raised", NewViewportToolbarTopLevelWarningRaisedStyle);
 		}
 	}
 
-	// Old viewport toolbar.
+	// Special style for the "Viewport Sizing" submenu, with rounded corners on the left.
 	{
-		FToolBarStyle ViewportToolbarStyle = FStarshipCoreStyle::GetCoreStyle().GetWidgetStyle<FToolBarStyle>("SlimToolBar");
+		const FButtonStyle ViewportSizingButtonStyle =
+			FButtonStyle(ButtonStyleStart)
+				.SetNormalPadding(FMargin(0.f))
+				.SetPressedPadding(FMargin(0.f));
 
-		FMargin ViewportMarginLeft(6.f, 4.f, 3.f, 4.f);
-		FMargin ViewportMarginCenter(6.f, 4.f, 3.f, 4.f);
-		FMargin ViewportMarginRight(4.f, 4.f, 5.f, 4.f);
-
-		const FCheckBoxStyle ViewportToggleButton = FCheckBoxStyle()
-			.SetCheckBoxType(ESlateCheckBoxType::ToggleButton)
-			.SetCheckedImage(FSlateNoResource())
-			.SetCheckedHoveredImage(FSlateNoResource())
-			.SetCheckedPressedImage(FSlateNoResource())
-			.SetUncheckedImage(FSlateNoResource())
-			.SetUncheckedHoveredImage(FSlateNoResource())
-			.SetUncheckedPressedImage(FSlateNoResource())
-			.SetForegroundColor(FStyleColors::Foreground)
-			.SetHoveredForegroundColor(FStyleColors::ForegroundHover)
-			.SetPressedForegroundColor(FStyleColors::ForegroundHover)
-			.SetCheckedForegroundColor(FStyleColors::ForegroundHover)
-			.SetCheckedHoveredForegroundColor(FStyleColors::ForegroundHover)
-			.SetPadding(0);
+		const FComboButtonStyle ViewportSizingComboButtonStyle =
+			FComboButtonStyle(ComboButtonStyle)
+				.SetDownArrowPadding(FMargin(5.f, 0.f, 5.f, 0.f))
+				.SetDownArrowImage(CORE_IMAGE_BRUSH_SVG("Starship/Common/ellipsis-vertical-narrow", FVector2f(6.f, 20.f)))
+				.SetButtonStyle(ViewportSizingButtonStyle);
 
+		const FToolBarStyle ViewportSizingStyle = FToolBarStyle(ViewportToolbarStyle)
+			.SetButtonStyle(ViewportSizingButtonStyle)
+			.SetComboButtonStyle(ViewportSizingComboButtonStyle)
+			.SetComboButtonPadding(FMargin(4.0f, 0, 1.0f, 0))
+			.SetSeparatorPadding(FMargin(0.f));
 
-		FLinearColor ToolbarBackgroundColor = FStyleColors::Dropdown.GetSpecifiedColor();
-		ToolbarBackgroundColor.A = .80f;
+		Set("ViewportToolbarViewportSizingSubmenu", ViewportSizingStyle);
 
-		FLinearColor ToolbarPressedColor = FStyleColors::Recessed.GetSpecifiedColor();
-		ToolbarPressedColor.A = .80f;
+		// A raised button keeps the gray background here, with rounded corners on the right.
+		{
+			FCheckBoxStyle ToggleButtonStyleRaised = FCheckBoxStyle(ToggleButtonStyleEnd)
+				// The icon cannot be perfectly centered, so it needs a 1px shove to the right
+				// to match the design spec.
+				.SetPadding(FMargin(7.0f, 4.0f, 6.0f, 4.0f));
+				
+			const FToolBarStyle ViewportSizingRaisedStyle = FToolBarStyle(ViewportToolbarStyle)
+				.SetToggleButtonStyle(ToggleButtonStyleRaised)
+				.SetButtonPadding(FMargin(0.f, 0.f, 4.f, 0.f))
+				.SetComboButtonPadding(FMargin(2.0f, 0.f));
 
-		FSlateRoundedBoxBrush* ViewportGroupBrush = new FSlateRoundedBoxBrush(ToolbarBackgroundColor, 12.f, FLinearColor(0.f,0.f,0.f,.8f), 1.0f);
-		Set("EditorViewportToolBar.Group", ViewportGroupBrush);
+			Set("ViewportToolbarViewportSizingSubmenu.Raised", ViewportSizingRaisedStyle);
+		}
+	}
 
-		FSlateRoundedBoxBrush* ViewportGroupPressedBrush = new FSlateRoundedBoxBrush(ToolbarPressedColor, 12.f, FLinearColor(0.f, 0.f, 0.f, .8f), 1.0f);
-		Set("EditorViewportToolBar.Group.Pressed", ViewportGroupPressedBrush);
+	// Viewport Toolbar Icons
+	{
+		Set("ViewportToolbar.CameraSpeed", new IMAGE_BRUSH_SVG("Starship/Common/CameraSpeed_16", Icon16x16));
+		Set("ViewportToolbar.EV100", new IMAGE_BRUSH_SVG("Starship/Common/EV100_16", Icon16x16));
+		Set("ViewportToolbar.ExactCameraView", new IMAGE_BRUSH_SVG("Starship/Common/ExactCameraView_16", Icon16x16));
+		Set("ViewportToolbar.Exposure", new IMAGE_BRUSH_SVG("Starship/Common/Exposure_16", Icon16x16));
+		Set("ViewportToolbar.FarViewPlane", new IMAGE_BRUSH_SVG("Starship/Common/FarViewPlane_16", Icon16x16));
+		Set("ViewportToolbar.FieldOfView", new IMAGE_BRUSH_SVG("Starship/Common/FieldOfView_16", Icon16x16));
+		Set("ViewportToolbar.GizmoScale", new IMAGE_BRUSH_SVG("Starship/Common/GizmoScale_16", Icon16x16));
+		Set("ViewportToolbar.NearViewPlane", new IMAGE_BRUSH_SVG("Starship/Common/NearViewPlane_16", Icon16x16));
+		Set("ViewportToolbar.PreviewSceneSettings", new IMAGE_BRUSH_SVG("Starship/AssetEditors/PreviewSceneSettings_16", Icon16x16));
+		Set("ViewportToolbar.SetShowGrid", new IMAGE_BRUSH_SVG("Starship/Common/Grid", Icon16x16));
+		Set("ViewportToolbar.Snap", new IMAGE_BRUSH_SVG("Starship/Common/Snap_16", Icon16x16));
+		Set("ViewportToolbar.SnapLocation", new IMAGE_BRUSH_SVG("Starship/Common/SnapLocation_16", Icon16x16));
+		Set("ViewportToolbar.SnapPlanar", new IMAGE_BRUSH_SVG("Starship/Common/SnapPlanar_16", Icon16x16));
+		Set("ViewportToolbar.SnapRotation", new IMAGE_BRUSH_SVG("Starship/Common/SnapRotation_16", Icon16x16));
+		Set("ViewportToolbar.SnapScale", new IMAGE_BRUSH_SVG("Starship/Common/SnapScale_16", Icon16x16));
+		Set("ViewportToolbar.SpeedScalar", new IMAGE_BRUSH_SVG("Starship/Common/SpeedScalar_16", Icon16x16));
+		Set("ViewportToolbar.SurfaceSnap", new IMAGE_BRUSH_SVG("Starship/Common/SurfaceSnap_16", Icon16x16));
+		Set("ViewportToolbar.SurfaceSnapRotateToNormal", new IMAGE_BRUSH_SVG("Starship/Common/SurfaceSnapRotateToNormal_16", Icon16x16));
+		Set("ViewportToolbar.TransformRotate", new IMAGE_BRUSH_SVG("Starship/Common/TransformRotate_16", Icon16x16));
+		Set("ViewportToolbar.TransformScale", new IMAGE_BRUSH_SVG("Starship/Common/TransformScale_16", Icon16x16));
+		Set("ViewportToolbar.TransformSelect", new IMAGE_BRUSH_SVG("Starship/Common/TransformSelect_16", Icon16x16));
+		Set("ViewportToolbar.TransformTranslate", new IMAGE_BRUSH_SVG("Starship/Common/TransformMove_16", Icon16x16));
+		Set("ViewportToolbar.VPCamera", new IMAGE_BRUSH_SVG("Starship/Common/VPCamera_16", Icon16x16));
+		Set("ViewportToolbar.VPCineCamera", new IMAGE_BRUSH_SVG("Starship/Common/VPCineCamera_16", Icon16x16));
+	}
+}
 
-		FButtonStyle ViewportMenuButton = FButtonStyle()
-			.SetNormal(*ViewportGroupBrush)
-			.SetHovered(*ViewportGroupBrush)
-			.SetPressed(*ViewportGroupPressedBrush)
-			.SetNormalForeground(FStyleColors::Foreground)
-			.SetHoveredForeground(FStyleColors::ForegroundHover)
-			.SetPressedForeground(FStyleColors::ForegroundHover)
-			.SetDisabledForeground(FStyleColors::Foreground)
-			.SetNormalPadding(FMargin(4.0f, 4.0f, 3.0f, 4.0f))
-			.SetPressedPadding(FMargin(4.0f, 4.0f, 3.0f, 4.0f));
-		Set("EditorViewportToolBar.Button", ViewportMenuButton);
+void FStarshipEditorStyle::FStyle::SetupOldViewportToolbarStyles()
+{
+	FToolBarStyle ViewportToolbarStyle = FStarshipCoreStyle::GetCoreStyle().GetWidgetStyle<FToolBarStyle>("SlimToolBar");
+
+	FMargin ViewportMarginLeft(6.f, 4.f, 3.f, 4.f);
+	FMargin ViewportMarginCenter(6.f, 4.f, 3.f, 4.f);
+	FMargin ViewportMarginRight(4.f, 4.f, 5.f, 4.f);
+
+	const FCheckBoxStyle ViewportToggleButton = FCheckBoxStyle()
+		.SetCheckBoxType(ESlateCheckBoxType::ToggleButton)
+		.SetCheckedImage(FSlateNoResource())
+		.SetCheckedHoveredImage(FSlateNoResource())
+		.SetCheckedPressedImage(FSlateNoResource())
+		.SetUncheckedImage(FSlateNoResource())
+		.SetUncheckedHoveredImage(FSlateNoResource())
+		.SetUncheckedPressedImage(FSlateNoResource())
+		.SetForegroundColor(FStyleColors::Foreground)
+		.SetHoveredForegroundColor(FStyleColors::ForegroundHover)
+		.SetPressedForegroundColor(FStyleColors::ForegroundHover)
+		.SetCheckedForegroundColor(FStyleColors::ForegroundHover)
+		.SetCheckedHoveredForegroundColor(FStyleColors::ForegroundHover)
+		.SetPadding(0);
+
+	FLinearColor ToolbarBackgroundColor = FStyleColors::Dropdown.GetSpecifiedColor();
+	ToolbarBackgroundColor.A = .80f;
+
+	FLinearColor ToolbarPressedColor = FStyleColors::Recessed.GetSpecifiedColor();
+	ToolbarPressedColor.A = .80f;
+
+	FSlateRoundedBoxBrush* ViewportGroupBrush = new FSlateRoundedBoxBrush(ToolbarBackgroundColor, 12.f, FLinearColor(0.f,0.f,0.f,.8f), 1.0f);
+	Set("EditorViewportToolBar.Group", ViewportGroupBrush);
+
+	FSlateRoundedBoxBrush* ViewportGroupPressedBrush = new FSlateRoundedBoxBrush(ToolbarPressedColor, 12.f, FLinearColor(0.f, 0.f, 0.f, .8f), 1.0f);
+	Set("EditorViewportToolBar.Group.Pressed", ViewportGroupPressedBrush);
+
+	FButtonStyle ViewportMenuButton = FButtonStyle()
+		.SetNormal(*ViewportGroupBrush)
+		.SetHovered(*ViewportGroupBrush)
+		.SetPressed(*ViewportGroupPressedBrush)
+		.SetNormalForeground(FStyleColors::Foreground)
+		.SetHoveredForeground(FStyleColors::ForegroundHover)
+		.SetPressedForeground(FStyleColors::ForegroundHover)
+		.SetDisabledForeground(FStyleColors::Foreground)
+		.SetNormalPadding(FMargin(4.0f, 4.0f, 3.0f, 4.0f))
+		.SetPressedPadding(FMargin(4.0f, 4.0f, 3.0f, 4.0f));
+	Set("EditorViewportToolBar.Button", ViewportMenuButton);
 
+	{
 		FButtonStyle ViewportMenuButtonLeft = FButtonStyle(ViewportMenuButton)
 			.SetNormal(		BOX_BRUSH("Starship/EditorViewport/ToolBarLeftGroup", 12.f/25.f, FStyleColors::Dropdown))
 			.SetHovered(	BOX_BRUSH("Starship/EditorViewport/ToolBarLeftGroup", 12.f/25.f, FStyleColors::Hover))
@@ -3012,11 +3027,11 @@ void FStarshipEditorStyle::FStyle::SetupViewportStyles()
 		Set("EditorViewportToolBar.Button.Start", ViewportMenuButtonLeft);
 
 		FButtonStyle ViewportMenuButtonMiddle = FButtonStyle(ViewportMenuButton)
-        	.SetNormal(		BOX_BRUSH("Starship/EditorViewport/ToolBarMiddleGroup", 12.f/25.f, FStyleColors::Dropdown))
-        	.SetHovered(	BOX_BRUSH("Starship/EditorViewport/ToolBarMiddleGroup", 12.f/25.f, FStyleColors::Hover))
-        	.SetPressed(	BOX_BRUSH("Starship/EditorViewport/ToolBarMiddleGroup", 12.f/25.f, FStyleColors::Recessed))
-        	.SetNormalPadding(ViewportMarginCenter)
-        	.SetPressedPadding(ViewportMarginCenter);
+			.SetNormal(		BOX_BRUSH("Starship/EditorViewport/ToolBarMiddleGroup", 12.f/25.f, FStyleColors::Dropdown))
+			.SetHovered(	BOX_BRUSH("Starship/EditorViewport/ToolBarMiddleGroup", 12.f/25.f, FStyleColors::Hover))
+			.SetPressed(	BOX_BRUSH("Starship/EditorViewport/ToolBarMiddleGroup", 12.f/25.f, FStyleColors::Recessed))
+			.SetNormalPadding(ViewportMarginCenter)
+			.SetPressedPadding(ViewportMarginCenter);
 		Set("EditorViewportToolBar.Button.Middle", ViewportMenuButtonMiddle);
 
 		FButtonStyle ViewportMenuButtonRight = FButtonStyle(ViewportMenuButton)
@@ -3026,9 +3041,11 @@ void FStarshipEditorStyle::FStyle::SetupViewportStyles()
 			.SetNormalPadding(ViewportMarginRight)
 			.SetPressedPadding(ViewportMarginRight);
 		Set("EditorViewportToolBar.Button.End", ViewportMenuButtonRight);
+	}
 
-		Set("EditorViewportToolBar.StartToolbarImage", new BOX_BRUSH("Starship/EditorViewport/ToolBarLeftGroup", 12.f/25.f, FStyleColors::Dropdown));
+	Set("EditorViewportToolBar.StartToolbarImage", new BOX_BRUSH("Starship/EditorViewport/ToolBarLeftGroup", 12.f/25.f, FStyleColors::Dropdown));
 
+	{
 		const FCheckBoxStyle ViewportMenuToggleLeftButtonStyle = FCheckBoxStyle(ViewportToggleButton)
 			.SetCheckBoxType(ESlateCheckBoxType::ToggleButton)
 			.SetUncheckedImage(		  BOX_BRUSH("Starship/EditorViewport/ToolBarLeftGroup", 12.f/25.f, FStyleColors::Dropdown))
@@ -3061,7 +3078,9 @@ void FStarshipEditorStyle::FStyle::SetupViewportStyles()
 			.SetCheckedImage(         BOX_BRUSH("Starship/EditorViewport/ToolBarRightGroup", 12.f/25.f, FStyleColors::Primary))
 			.SetPadding(ViewportMarginRight);
 		Set("EditorViewportToolBar.ToggleButton.End", ViewportMenuToggleRightButtonStyle);
+	}
 
+	{
 		// We want a background-less version as the ComboMenu has its own unified background
 		const FToolBarStyle& SlimCoreToolBarStyle = FStarshipCoreStyle::GetCoreStyle().GetWidgetStyle<FToolBarStyle>("SlimToolBar");
 
@@ -3075,7 +3094,9 @@ void FStarshipEditorStyle::FStyle::SetupViewportStyles()
 		Set("EditorViewportToolBar.ComboMenu.ButtonStyle", ComboMenuButtonStyle);
 		Set("EditorViewportToolBar.ComboMenu.ToggleButton", ViewportToggleButton);
 		Set("EditorViewportToolBar.ComboMenu.LabelStyle", SlimCoreToolBarStyle.LabelStyle);
+	}
 
+	{
 		FCheckBoxStyle MaximizeRestoreButton = FCheckBoxStyle(ViewportToolbarStyle.ToggleButton)
 			.SetUncheckedImage(*ViewportGroupBrush)
 			.SetUncheckedPressedImage(*ViewportGroupPressedBrush)
@@ -3090,16 +3111,11 @@ void FStarshipEditorStyle::FStyle::SetupViewportStyles()
 			.SetCheckedPressedForegroundColor(FStyleColors::ForegroundHover)
 			.SetCheckedHoveredForegroundColor(FStyleColors::ForegroundHover)
 			.SetPadding(FMargin(4.0f, 4.0f, 3.0f, 4.0f));
+
 		Set("EditorViewportToolBar.MaximizeRestoreButton", MaximizeRestoreButton);
 
 		Set("EditorViewportToolBar.Heading.Padding", FMargin(4.f));
 
-
-		// SComboBox 
-		FComboButtonStyle ViewportComboButton = FComboButtonStyle()
-			.SetButtonStyle(ViewportMenuButton)
-			.SetContentPadding(ViewportMarginCenter);
-
 		// Non-grouped Toggle Button
 		FCheckBoxStyle SoloToggleButton = FCheckBoxStyle(ViewportToolbarStyle.ToggleButton)
 			.SetUncheckedImage(*ViewportGroupBrush)
@@ -3116,7 +3132,11 @@ void FStarshipEditorStyle::FStyle::SetupViewportStyles()
 			.SetCheckedHoveredForegroundColor(FStyleColors::ForegroundHover)
 			.SetPadding(FMargin(6.0f, 4.0f, 6.0f, 4.0f));
 
-
+		// SComboBox 
+		FComboButtonStyle ViewportComboButton = FComboButtonStyle()
+			.SetButtonStyle(ViewportMenuButton)
+			.SetContentPadding(ViewportMarginCenter);
+	
 		ViewportToolbarStyle
 			.SetBackground(FSlateNoResource())
 			.SetIconSize(Icon16x16)
@@ -3132,16 +3152,19 @@ void FStarshipEditorStyle::FStyle::SetupViewportStyles()
 			.SetButtonStyle(ViewportMenuButton)
 			.SetSeparatorBrush(FSlateNoResource())
 			.SetSeparatorPadding(FMargin(2.0f, 0.0f));
-		
+	
 		ViewportToolbarStyle.WrapButtonStyle.SetExpandBrush(IMAGE_BRUSH("Icons/toolbar_expand_16x", Icon8x8));
-		
+	
 		Set("EditorViewportToolBar", ViewportToolbarStyle);
+	}
 
+	{
 		FButtonStyle ViewportMenuWarningButton = FButtonStyle(ViewportMenuButton)
 			.SetNormalForeground(FStyleColors::AccentYellow)
 			.SetHoveredForeground(FStyleColors::ForegroundHover)
 			.SetPressedForeground(FStyleColors::ForegroundHover)
 			.SetDisabledForeground(FStyleColors::AccentYellow);
+
 		Set("EditorViewportToolBar.WarningButton", ViewportMenuWarningButton);
 
 		Set("EditorViewportToolBar.Background", new FSlateNoResource());
@@ -3155,7 +3178,6 @@ void FStarshipEditorStyle::FStyle::SetupViewportStyles()
 			.SetPressed(BOX_BRUSH("Common/SmallRoundedButton", FMargin(7.f / 16.f)))
 		);
 
-		
 		Set("EditorViewportToolBar.MenuDropdown", new IMAGE_BRUSH("Common/ComboArrow", Icon8x8));
 		Set("EditorViewportToolBar.Maximize.Normal", new IMAGE_BRUSH_SVG("Starship/EditorViewport/square", Icon16x16));
 		Set("EditorViewportToolBar.Maximize.Checked", new IMAGE_BRUSH_SVG("Starship/EditorViewport/quad", Icon16x16));
@@ -3184,128 +3206,118 @@ void FStarshipEditorStyle::FStyle::SetupViewportStyles()
 			.SetInactiveHoveredBrush(ActionableListViewBrush)
 		);
 	}
+}
 
-	// Legacy Viewport ToolbarBar
-	{
-		FToolBarStyle ViewportToolbarStyle =
-			FToolBarStyle()
-			.SetBackground(BOX_BRUSH("Old/Menu_Background", FMargin(8.0f / 64.0f), FLinearColor::Transparent))
-			.SetComboButtonPadding(FMargin(0.f))
-			.SetButtonPadding(FMargin(0.f))
-			.SetCheckBoxPadding(FMargin(4.f))
-			.SetSeparatorBrush(BOX_BRUSH("Old/Button", 8.0f / 32.0f, FLinearColor::Transparent))
-			.SetSeparatorPadding(FMargin(1.f, 0.f, 0.f, 0.f))
-			.SetIconSize(Icon16x16)
-			.SetLabelPadding(FMargin(0.0f, 0.0f, 3.0f, 0.0f))
-			.SetEditableTextStyle(FEditableTextBoxStyle(NormalEditableTextBoxStyle).SetFont(DEFAULT_FONT("Regular", 9)))
-			.SetIndentedBlockPadding(FMargin(0.f))
-			.SetBlockPadding(FMargin(0.f))
-			.SetLabelStyle(
-				FTextBlockStyle(NormalText)
-				.SetFont(DEFAULT_FONT("Bold", 9))
-				.SetColorAndOpacity(FLinearColor(0.0f, 0.0f, 0.0f, 1.0f))
-			);
-			
-		ViewportToolbarStyle.WrapButtonStyle.SetExpandBrush(IMAGE_BRUSH("Icons/toolbar_expand_16x", Icon8x8));
-		
-		const FString SmallRoundedButton(TEXT("Common/SmallRoundedButton"));
-		const FString SmallRoundedButtonStart(TEXT("Common/SmallRoundedButtonLeft"));
-		const FString SmallRoundedButtonMiddle(TEXT("Common/SmallRoundedButtonCentre"));
-		const FString SmallRoundedButtonEnd(TEXT("Common/SmallRoundedButtonRight"));
-
-		const FLinearColor NormalColor(1.f, 1.f, 1.f, 0.75f);
-		const FLinearColor PressedColor(1.f, 1.f, 1.f, 1.f);
-
-		/* Create style for "LegacyViewportMenu.ToggleButton" ... */
-		const FCheckBoxStyle ViewportMenuToggleButtonStyle = FCheckBoxStyle()
-			.SetCheckBoxType(ESlateCheckBoxType::ToggleButton)
-			.SetUncheckedImage(BOX_BRUSH(*SmallRoundedButton, FMargin(7.f / 16.f), NormalColor))
-			.SetUncheckedPressedImage(BOX_BRUSH(*SmallRoundedButton, FMargin(7.f / 16.f), PressedColor))
-			.SetUncheckedHoveredImage(BOX_BRUSH(*SmallRoundedButton, FMargin(7.f / 16.f), PressedColor))
-			.SetCheckedHoveredImage(BOX_BRUSH(*SmallRoundedButton, FMargin(7.f / 16.f), SelectionColor_Pressed))
-			.SetCheckedPressedImage(BOX_BRUSH(*SmallRoundedButton, FMargin(7.f / 16.f), SelectionColor_Pressed))
-			.SetCheckedImage(BOX_BRUSH(*SmallRoundedButton, FMargin(7.f / 16.f), SelectionColor_Pressed));
-		/* ... and add new style */
-		ViewportToolbarStyle.SetToggleButtonStyle(ViewportMenuToggleButtonStyle);
-
-		/* Create style for "LegacyViewportMenu.ToggleButton.Start" ... */
-		const FCheckBoxStyle ViewportMenuToggleStartButtonStyle = FCheckBoxStyle()
-			.SetCheckBoxType(ESlateCheckBoxType::ToggleButton)
-			.SetUncheckedImage(BOX_BRUSH(*SmallRoundedButtonStart, FMargin(7.f / 16.f), NormalColor))
-			.SetUncheckedPressedImage(BOX_BRUSH(*SmallRoundedButtonStart, FMargin(7.f / 16.f), PressedColor))
-			.SetUncheckedHoveredImage(BOX_BRUSH(*SmallRoundedButtonStart, FMargin(7.f / 16.f), PressedColor))
-			.SetCheckedHoveredImage(BOX_BRUSH(*SmallRoundedButtonStart, FMargin(7.f / 16.f), SelectionColor_Pressed))
-			.SetCheckedPressedImage(BOX_BRUSH(*SmallRoundedButtonStart, FMargin(7.f / 16.f), SelectionColor_Pressed))
-			.SetCheckedImage(BOX_BRUSH(*SmallRoundedButtonStart, FMargin(7.f / 16.f), SelectionColor_Pressed));
-		/* ... and add new style */
-		Set("LegacyViewportMenu.ToggleButton.Start", ViewportMenuToggleStartButtonStyle);
-
-		/* Create style for "LegacyViewportMenu.ToggleButton.Middle" ... */
-		const FCheckBoxStyle ViewportMenuToggleMiddleButtonStyle = FCheckBoxStyle()
-			.SetCheckBoxType(ESlateCheckBoxType::ToggleButton)
-			.SetUncheckedImage(BOX_BRUSH(*SmallRoundedButtonMiddle, FMargin(7.f / 16.f), NormalColor))
-			.SetUncheckedPressedImage(BOX_BRUSH(*SmallRoundedButtonMiddle, FMargin(7.f / 16.f), PressedColor))
-			.SetUncheckedHoveredImage(BOX_BRUSH(*SmallRoundedButtonMiddle, FMargin(7.f / 16.f), PressedColor))
-			.SetCheckedHoveredImage(BOX_BRUSH(*SmallRoundedButtonMiddle, FMargin(7.f / 16.f), SelectionColor_Pressed))
-			.SetCheckedPressedImage(BOX_BRUSH(*SmallRoundedButtonMiddle, FMargin(7.f / 16.f), SelectionColor_Pressed))
-			.SetCheckedImage(BOX_BRUSH(*SmallRoundedButtonMiddle, FMargin(7.f / 16.f), SelectionColor_Pressed));
-		/* ... and add new style */
-		Set("LegacyViewportMenu.ToggleButton.Middle", ViewportMenuToggleMiddleButtonStyle);
-
-		/* Create style for "LegacyViewportMenu.ToggleButton.End" ... */
-		const FCheckBoxStyle ViewportMenuToggleEndButtonStyle = FCheckBoxStyle()
-			.SetCheckBoxType(ESlateCheckBoxType::ToggleButton)
-			.SetUncheckedImage(BOX_BRUSH(*SmallRoundedButtonEnd, FMargin(7.f / 16.f), NormalColor))
-			.SetUncheckedPressedImage(BOX_BRUSH(*SmallRoundedButtonEnd, FMargin(7.f / 16.f), PressedColor))
-			.SetUncheckedHoveredImage(BOX_BRUSH(*SmallRoundedButtonEnd, FMargin(7.f / 16.f), PressedColor))
-			.SetCheckedHoveredImage(BOX_BRUSH(*SmallRoundedButtonEnd, FMargin(7.f / 16.f), SelectionColor_Pressed))
-			.SetCheckedPressedImage(BOX_BRUSH(*SmallRoundedButtonEnd, FMargin(7.f / 16.f), SelectionColor_Pressed))
-			.SetCheckedImage(BOX_BRUSH(*SmallRoundedButtonEnd, FMargin(7.f / 16.f), SelectionColor_Pressed));
-		/* ... and add new style */
-		Set("LegacyViewportMenu.ToggleButton.End", ViewportMenuToggleEndButtonStyle);
-
-		const FMargin NormalPadding = FMargin(4.0f, 4.0f, 4.0f, 4.0f);
-		const FMargin PressedPadding = FMargin(4.0f, 4.0f, 4.0f, 4.0f);
-
-		const FButtonStyle ViewportMenuButton = FButtonStyle(Button)
-		.SetNormal(BOX_BRUSH(*SmallRoundedButton, 7.0f / 16.0f, NormalColor))
-		.SetPressed(BOX_BRUSH(*SmallRoundedButton, 7.0f / 16.0f, PressedColor))
-		.SetHovered(BOX_BRUSH(*SmallRoundedButton, 7.0f / 16.0f, PressedColor))
-		.SetPressedPadding(PressedPadding)
-		.SetNormalPadding(NormalPadding);
-
-		ViewportToolbarStyle.SetButtonStyle(ViewportMenuButton);
-
-		Set("LegacyViewportMenu.Button.Start", FButtonStyle(ViewportMenuButton)
-			.SetNormal(BOX_BRUSH(*SmallRoundedButtonStart, 7.0f / 16.0f, NormalColor))
-			.SetPressed(BOX_BRUSH(*SmallRoundedButtonStart, 7.0f / 16.0f, PressedColor))
-			.SetHovered(BOX_BRUSH(*SmallRoundedButtonStart, 7.0f / 16.0f, PressedColor))
-		);
-
-		Set("LegacyViewportMenu.Button.Middle", FButtonStyle(ViewportMenuButton)
-			.SetNormal(BOX_BRUSH(*SmallRoundedButtonMiddle, 7.0f / 16.0f, NormalColor))
-			.SetPressed(BOX_BRUSH(*SmallRoundedButtonMiddle, 7.0f / 16.0f, PressedColor))
-			.SetHovered(BOX_BRUSH(*SmallRoundedButtonMiddle, 7.0f / 16.0f, PressedColor))
+void FStarshipEditorStyle::FStyle::SetupLegacyViewportToolbarStyles()
+{
+	FToolBarStyle ViewportToolbarStyle =
+		FToolBarStyle()
+		.SetBackground(BOX_BRUSH("Old/Menu_Background", FMargin(8.0f / 64.0f), FLinearColor::Transparent))
+		.SetComboButtonPadding(FMargin(0.f))
+		.SetButtonPadding(FMargin(0.f))
+		.SetCheckBoxPadding(FMargin(4.f))
+		.SetSeparatorBrush(BOX_BRUSH("Old/Button", 8.0f / 32.0f, FLinearColor::Transparent))
+		.SetSeparatorPadding(FMargin(1.f, 0.f, 0.f, 0.f))
+		.SetIconSize(Icon16x16)
+		.SetLabelPadding(FMargin(0.0f, 0.0f, 3.0f, 0.0f))
+		.SetEditableTextStyle(FEditableTextBoxStyle(NormalEditableTextBoxStyle).SetFont(DEFAULT_FONT("Regular", 9)))
+		.SetIndentedBlockPadding(FMargin(0.f))
+		.SetBlockPadding(FMargin(0.f))
+		.SetLabelStyle(
+			FTextBlockStyle(NormalText)
+			.SetFont(DEFAULT_FONT("Bold", 9))
+			.SetColorAndOpacity(FLinearColor(0.0f, 0.0f, 0.0f, 1.0f))
 		);
+		
+	ViewportToolbarStyle.WrapButtonStyle.SetExpandBrush(IMAGE_BRUSH("Icons/toolbar_expand_16x", Icon8x8));
+	
+	const FString SmallRoundedButton(TEXT("Common/SmallRoundedButton"));
+	const FString SmallRoundedButtonStart(TEXT("Common/SmallRoundedButtonLeft"));
+	const FString SmallRoundedButtonMiddle(TEXT("Common/SmallRoundedButtonCentre"));
+	const FString SmallRoundedButtonEnd(TEXT("Common/SmallRoundedButtonRight"));
+
+	const FLinearColor NormalColor(1.f, 1.f, 1.f, 0.75f);
+	const FLinearColor PressedColor(1.f, 1.f, 1.f, 1.f);
+
+	/* Create style for "LegacyViewportMenu.ToggleButton" ... */
+	const FCheckBoxStyle ViewportMenuToggleButtonStyle = FCheckBoxStyle()
+		.SetCheckBoxType(ESlateCheckBoxType::ToggleButton)
+		.SetUncheckedImage(BOX_BRUSH(*SmallRoundedButton, FMargin(7.f / 16.f), NormalColor))
+		.SetUncheckedPressedImage(BOX_BRUSH(*SmallRoundedButton, FMargin(7.f / 16.f), PressedColor))
+		.SetUncheckedHoveredImage(BOX_BRUSH(*SmallRoundedButton, FMargin(7.f / 16.f), PressedColor))
+		.SetCheckedHoveredImage(BOX_BRUSH(*SmallRoundedButton, FMargin(7.f / 16.f), SelectionColor_Pressed))
+		.SetCheckedPressedImage(BOX_BRUSH(*SmallRoundedButton, FMargin(7.f / 16.f), SelectionColor_Pressed))
+		.SetCheckedImage(BOX_BRUSH(*SmallRoundedButton, FMargin(7.f / 16.f), SelectionColor_Pressed));
+	/* ... and add new style */
+	ViewportToolbarStyle.SetToggleButtonStyle(ViewportMenuToggleButtonStyle);
+
+	/* Create style for "LegacyViewportMenu.ToggleButton.Start" ... */
+	const FCheckBoxStyle ViewportMenuToggleStartButtonStyle = FCheckBoxStyle()
+		.SetCheckBoxType(ESlateCheckBoxType::ToggleButton)
+		.SetUncheckedImage(BOX_BRUSH(*SmallRoundedButtonStart, FMargin(7.f / 16.f), NormalColor))
+		.SetUncheckedPressedImage(BOX_BRUSH(*SmallRoundedButtonStart, FMargin(7.f / 16.f), PressedColor))
+		.SetUncheckedHoveredImage(BOX_BRUSH(*SmallRoundedButtonStart, FMargin(7.f / 16.f), PressedColor))
+		.SetCheckedHoveredImage(BOX_BRUSH(*SmallRoundedButtonStart, FMargin(7.f / 16.f), SelectionColor_Pressed))
+		.SetCheckedPressedImage(BOX_BRUSH(*SmallRoundedButtonStart, FMargin(7.f / 16.f), SelectionColor_Pressed))
+		.SetCheckedImage(BOX_BRUSH(*SmallRoundedButtonStart, FMargin(7.f / 16.f), SelectionColor_Pressed));
+	/* ... and add new style */
+	Set("LegacyViewportMenu.ToggleButton.Start", ViewportMenuToggleStartButtonStyle);
+
+	/* Create style for "LegacyViewportMenu.ToggleButton.Middle" ... */
+	const FCheckBoxStyle ViewportMenuToggleMiddleButtonStyle = FCheckBoxStyle()
+		.SetCheckBoxType(ESlateCheckBoxType::ToggleButton)
+		.SetUncheckedImage(BOX_BRUSH(*SmallRoundedButtonMiddle, FMargin(7.f / 16.f), NormalColor))
+		.SetUncheckedPressedImage(BOX_BRUSH(*SmallRoundedButtonMiddle, FMargin(7.f / 16.f), PressedColor))
+		.SetUncheckedHoveredImage(BOX_BRUSH(*SmallRoundedButtonMiddle, FMargin(7.f / 16.f), PressedColor))
+		.SetCheckedHoveredImage(BOX_BRUSH(*SmallRoundedButtonMiddle, FMargin(7.f / 16.f), SelectionColor_Pressed))
+		.SetCheckedPressedImage(BOX_BRUSH(*SmallRoundedButtonMiddle, FMargin(7.f / 16.f), SelectionColor_Pressed))
+		.SetCheckedImage(BOX_BRUSH(*SmallRoundedButtonMiddle, FMargin(7.f / 16.f), SelectionColor_Pressed));
+	/* ... and add new style */
+	Set("LegacyViewportMenu.ToggleButton.Middle", ViewportMenuToggleMiddleButtonStyle);
+
+	/* Create style for "LegacyViewportMenu.ToggleButton.End" ... */
+	const FCheckBoxStyle ViewportMenuToggleEndButtonStyle = FCheckBoxStyle()
+		.SetCheckBoxType(ESlateCheckBoxType::ToggleButton)
+		.SetUncheckedImage(BOX_BRUSH(*SmallRoundedButtonEnd, FMargin(7.f / 16.f), NormalColor))
+		.SetUncheckedPressedImage(BOX_BRUSH(*SmallRoundedButtonEnd, FMargin(7.f / 16.f), PressedColor))
+		.SetUncheckedHoveredImage(BOX_BRUSH(*SmallRoundedButtonEnd, FMargin(7.f / 16.f), PressedColor))
+		.SetCheckedHoveredImage(BOX_BRUSH(*SmallRoundedButtonEnd, FMargin(7.f / 16.f), SelectionColor_Pressed))
+		.SetCheckedPressedImage(BOX_BRUSH(*SmallRoundedButtonEnd, FMargin(7.f / 16.f), SelectionColor_Pressed))
+		.SetCheckedImage(BOX_BRUSH(*SmallRoundedButtonEnd, FMargin(7.f / 16.f), SelectionColor_Pressed));
+	/* ... and add new style */
+	Set("LegacyViewportMenu.ToggleButton.End", ViewportMenuToggleEndButtonStyle);
+
+	const FMargin NormalPadding = FMargin(4.0f, 4.0f, 4.0f, 4.0f);
+	const FMargin PressedPadding = FMargin(4.0f, 4.0f, 4.0f, 4.0f);
+
+	const FButtonStyle ViewportMenuButton = FButtonStyle(Button)
+	.SetNormal(BOX_BRUSH(*SmallRoundedButton, 7.0f / 16.0f, NormalColor))
+	.SetPressed(BOX_BRUSH(*SmallRoundedButton, 7.0f / 16.0f, PressedColor))
+	.SetHovered(BOX_BRUSH(*SmallRoundedButton, 7.0f / 16.0f, PressedColor))
+	.SetPressedPadding(PressedPadding)
+	.SetNormalPadding(NormalPadding);
+
+	ViewportToolbarStyle.SetButtonStyle(ViewportMenuButton);
+
+	Set("LegacyViewportMenu.Button.Start", FButtonStyle(ViewportMenuButton)
+		.SetNormal(BOX_BRUSH(*SmallRoundedButtonStart, 7.0f / 16.0f, NormalColor))
+		.SetPressed(BOX_BRUSH(*SmallRoundedButtonStart, 7.0f / 16.0f, PressedColor))
+		.SetHovered(BOX_BRUSH(*SmallRoundedButtonStart, 7.0f / 16.0f, PressedColor))
+	);
 
-		Set("LegacyViewportMenu.Button.End", FButtonStyle(ViewportMenuButton)
-			.SetNormal(BOX_BRUSH(*SmallRoundedButtonEnd, 7.0f / 16.0f, NormalColor))
-			.SetPressed(BOX_BRUSH(*SmallRoundedButtonEnd, 7.0f / 16.0f, PressedColor))
-			.SetHovered(BOX_BRUSH(*SmallRoundedButtonEnd, 7.0f / 16.0f, PressedColor))
-		);
+	Set("LegacyViewportMenu.Button.Middle", FButtonStyle(ViewportMenuButton)
+		.SetNormal(BOX_BRUSH(*SmallRoundedButtonMiddle, 7.0f / 16.0f, NormalColor))
+		.SetPressed(BOX_BRUSH(*SmallRoundedButtonMiddle, 7.0f / 16.0f, PressedColor))
+		.SetHovered(BOX_BRUSH(*SmallRoundedButtonMiddle, 7.0f / 16.0f, PressedColor))
+	);
 
-		Set("LegacyViewportMenu", ViewportToolbarStyle);
-	}
+	Set("LegacyViewportMenu.Button.End", FButtonStyle(ViewportMenuButton)
+		.SetNormal(BOX_BRUSH(*SmallRoundedButtonEnd, 7.0f / 16.0f, NormalColor))
+		.SetPressed(BOX_BRUSH(*SmallRoundedButtonEnd, 7.0f / 16.0f, PressedColor))
+		.SetHovered(BOX_BRUSH(*SmallRoundedButtonEnd, 7.0f / 16.0f, PressedColor))
+	);
 
-	// Viewport actor preview's pin/unpin and attach/detach buttons
-	{
-		Set("ViewportActorPreview.Pinned", new IMAGE_BRUSH("Common/PushPin_Down", Icon16x16));
-		Set("ViewportActorPreview.Unpinned", new IMAGE_BRUSH("Common/PushPin_Up", Icon16x16));
-		Set("VRViewportActorPreview.Pinned", new IMAGE_BRUSH("Common/PushPin_Down_VR", Icon64x64));
-		Set("VRViewportActorPreview.Unpinned", new IMAGE_BRUSH("Common/PushPin_Up_VR", Icon64x64));
-		Set("VRViewportActorPreview.Attached", new IMAGE_BRUSH("Common/ScreenAttach_VR", Icon64x64));
-		Set("VRViewportActorPreview.Detached", new IMAGE_BRUSH("Common/ScreenDetach_VR", Icon64x64));
-	}
+	Set("LegacyViewportMenu", ViewportToolbarStyle);
 }
 
 void FStarshipEditorStyle::FStyle::SetupDepthBarStyles()
diff --git a/Engine/Source/Editor/EditorStyle/Private/StarshipStyle.h b/Engine/Source/Editor/EditorStyle/Private/StarshipStyle.h
--- a/Engine/Source/Editor/EditorStyle/Private/StarshipStyle.h
+++ b/Engine/Source/Editor/EditorStyle/Private/StarshipStyle.h
@@ -50,6 +50,9 @@ public:
 		void SetupWorldPartitionStyles();
 		void SetupSequencerStyles();
 		void SetupViewportStyles();
+		void SetupNewViewportToolbarStyles();
+		void SetupOldViewportToolbarStyles();
+		void SetupLegacyViewportToolbarStyles();
 		void SetupDepthBarStyles();
 		void SetupMenuBarStyles();
 		void SetupGeneralIcons();
