Changed "Keyboard Shortcuts" menu icon to keyboard

#Synced-CL 49942865

diff --git a/Engine/Source/Editor/EditorStyle/Private/StarshipStyle.cpp b/Engine/Source/Editor/EditorStyle/Private/StarshipStyle.cpp
--- a/Engine/Source/Editor/EditorStyle/Private/StarshipStyle.cpp
+++ b/Engine/Source/Editor/EditorStyle/Private/StarshipStyle.cpp
@@ -1410,7 +1410,7 @@ void FStarshipEditorStyle::FStyle::SetupGeneralStyles()
 	{
 		Set("EditorPreferences.TabIcon", new CORE_IMAGE_BRUSH_SVG("Starship/Common/Preferences", Icon16x16));
 		Set("SystemWideCommands.OpenEditorPreferences", new CORE_IMAGE_BRUSH_SVG("Starship/Common/Preferences", Icon16x16));
-		Set("SystemWideCommands.OpenKeyboardShortcuts", new CORE_IMAGE_BRUSH_SVG("Starship/Common/Preferences", Icon16x16));
+		Set("SystemWideCommands.OpenKeyboardShortcuts", new IMAGE_BRUSH_SVG("Starship/Common/ViewportControls", Icon16x16));
 	}
 
 	// Project settings
