[Build Warning Fix] Attempting to fix Nightly Build warning for external use of FVerticalToolBar

[Fix Description] FVerticalToolBarBuilder is using the FVerticalToolBar and it is in a runtime module.  FVerticalToolBar is defined in StarshipStyle but should be in StarshipCoreStyle to be accessible outside the editor.
[Warning Cause Theory] Stage Content Worker Test Project started using FVerticalToolBarBuilder and didn't override the base style resulting in this warning.

#Synced-CL 50045910

diff --git a/Engine/Source/Editor/EditorStyle/Private/StarshipStyle.cpp b/Engine/Source/Editor/EditorStyle/Private/StarshipStyle.cpp
--- a/Engine/Source/Editor/EditorStyle/Private/StarshipStyle.cpp
+++ b/Engine/Source/Editor/EditorStyle/Private/StarshipStyle.cpp
@@ -1833,48 +1833,6 @@ void FStarshipEditorStyle::FStyle::SetupGeneralStyles()
 		
 	}
 
-{
-		FToolBarStyle FVerticalToolBarStyle = FStarshipCoreStyle::GetCoreStyle().GetWidgetStyle<FToolBarStyle>("SlimToolBar");
-
-		FTextBlockStyle FVerticalToolBarLabelStyle = FTextBlockStyle(GetParentStyle()->GetWidgetStyle<FTextBlockStyle>("SmallText"));
-		FVerticalToolBarLabelStyle.SetOverflowPolicy(ETextOverflowPolicy::Ellipsis);
-		
-		FVerticalToolBarStyle.SetLabelStyle(FVerticalToolBarLabelStyle);
-		FVerticalToolBarStyle.SetLabelPadding(FMargin(4.0f, 2.0f, 4.0f, 4.0f));
-		FVerticalToolBarStyle.SetButtonContentMaxWidth( 64.0f );
-		FVerticalToolBarStyle.SetButtonContentFillWidth( 1.0f );
-
-		FVerticalToolBarStyle.SetButtonPadding(       FMargin(0.0f, 0.0f, 0.0f, 0.0f));
-		//not this
-		FVerticalToolBarStyle.SetCheckBoxPadding(     FMargin(0.0f, 0.0f));
-		FVerticalToolBarStyle.SetComboButtonPadding(  FMargin(0.0f, 0.0f));
-		FVerticalToolBarStyle.SetIndentedBlockPadding(FMargin(0.0f, 0.0f));	
-		FVerticalToolBarStyle.SetBackgroundPadding(   FMargin(6.0f, 4.0f));
-		FVerticalToolBarStyle.ButtonStyle.SetNormalPadding(FMargin(12.0f, 6.0f));
-		FVerticalToolBarStyle.ButtonStyle.SetPressedPadding(FMargin(12.0f, 6.0f));
-		FVerticalToolBarStyle.SetIconPadding(FMargin(8.0f, 8.0f));
-		FVerticalToolBarStyle.SetIconPaddingWithVisibleLabel(FMargin(8.0f, 8.0f, 8.0f, 2.0f));
- 
-		FVerticalToolBarStyle.WrapButtonStyle.SetExpandBrush(CORE_IMAGE_BRUSH_SVG("Starship/Common/ellipsis-horizontal-narrow", Icon16x16));
-		FCheckBoxStyle CheckBoxStyle = FCheckBoxStyle(FVerticalToolBarStyle.ToggleButton)
-			.SetCheckedImage(FSlateRoundedBoxBrush(FStyleColors::Primary, 4.f, FLinearColor(0.f, 0.f, 0.f, .8f), 0.5f))
-			.SetCheckedHoveredImage(FSlateRoundedBoxBrush(FStyleColors::PrimaryHover, 4.f, FLinearColor(0.f, 0.f, 0.f, .8f), 0.5f))
-			.SetCheckedPressedImage(FSlateRoundedBoxBrush(FStyleColors::PrimaryPress, 4.f, FLinearColor(0.f, 0.f, 0.f, .8f), 0.5f))
-			.SetForegroundColor(FStyleColors::Foreground)
-			.SetPressedForegroundColor(FStyleColors::ForegroundHover)
-			.SetHoveredForegroundColor(FStyleColors::ForegroundHover)
-			.SetCheckedForegroundColor(FStyleColors::ForegroundHover)
-			.SetCheckedPressedForegroundColor(FStyleColors::ForegroundHover)
-			.SetCheckedHoveredForegroundColor(FStyleColors::ForegroundHover)
-			.SetPadding(0.f);
-
-		Set("FVerticalToolBar.ToggleButton", CheckBoxStyle);
-
-		FVerticalToolBarStyle.SetButtonPadding(FMargin(0.0f, 4.0f));
-		FVerticalToolBarStyle.SetToggleButtonStyle(CheckBoxStyle);
-		FVerticalToolBarStyle.SetSeparatorPadding(FMargin(-5.0f, 4.0f));
-		Set("FVerticalToolBar", FVerticalToolBarStyle);
-	}
 
 	{
 		// FCategoryDrivenContentBuilder vertical toolbar style
diff --git a/Engine/Source/Runtime/SlateCore/Private/Styling/StarshipCoreStyle.cpp b/Engine/Source/Runtime/SlateCore/Private/Styling/StarshipCoreStyle.cpp
--- a/Engine/Source/Runtime/SlateCore/Private/Styling/StarshipCoreStyle.cpp
+++ b/Engine/Source/Runtime/SlateCore/Private/Styling/StarshipCoreStyle.cpp
@@ -1055,6 +1055,50 @@ TSharedRef<ISlateStyle> FStarshipCoreStyle::Create()
 
 	}
 
+	// Vertical ToolBar
+	{
+		FToolBarStyle FVerticalToolBarStyle = Style->GetWidgetStyle<FToolBarStyle>("SlimToolBar");
+
+		FTextBlockStyle FVerticalToolBarLabelStyle = FTextBlockStyle(Style->GetWidgetStyle<FTextBlockStyle>("SmallText"));
+		FVerticalToolBarLabelStyle.SetOverflowPolicy(ETextOverflowPolicy::Ellipsis);
+
+		FVerticalToolBarStyle.SetLabelStyle(FVerticalToolBarLabelStyle);
+		FVerticalToolBarStyle.SetLabelPadding(FMargin(4.0f, 2.0f, 4.0f, 4.0f));
+		FVerticalToolBarStyle.SetButtonContentMaxWidth(64.0f);
+		FVerticalToolBarStyle.SetButtonContentFillWidth(1.0f);
+
+		FVerticalToolBarStyle.SetButtonPadding(FMargin(0.0f, 0.0f, 0.0f, 0.0f));
+		//not this
+		FVerticalToolBarStyle.SetCheckBoxPadding(FMargin(0.0f, 0.0f));
+		FVerticalToolBarStyle.SetComboButtonPadding(FMargin(0.0f, 0.0f));
+		FVerticalToolBarStyle.SetIndentedBlockPadding(FMargin(0.0f, 0.0f));
+		FVerticalToolBarStyle.SetBackgroundPadding(FMargin(6.0f, 4.0f));
+		FVerticalToolBarStyle.ButtonStyle.SetNormalPadding(FMargin(12.0f, 6.0f));
+		FVerticalToolBarStyle.ButtonStyle.SetPressedPadding(FMargin(12.0f, 6.0f));
+		FVerticalToolBarStyle.SetIconPadding(FMargin(8.0f, 8.0f));
+		FVerticalToolBarStyle.SetIconPaddingWithVisibleLabel(FMargin(8.0f, 8.0f, 8.0f, 2.0f));
+
+		FVerticalToolBarStyle.WrapButtonStyle.SetExpandBrush(IMAGE_BRUSH_SVG("Starship/Common/ellipsis-horizontal-narrow", Icon16x16));
+		FCheckBoxStyle CheckBoxStyle = FCheckBoxStyle(FVerticalToolBarStyle.ToggleButton)
+			.SetCheckedImage(FSlateRoundedBoxBrush(FStyleColors::Primary, 4.f, FLinearColor(0.f, 0.f, 0.f, .8f), 0.5f))
+			.SetCheckedHoveredImage(FSlateRoundedBoxBrush(FStyleColors::PrimaryHover, 4.f, FLinearColor(0.f, 0.f, 0.f, .8f), 0.5f))
+			.SetCheckedPressedImage(FSlateRoundedBoxBrush(FStyleColors::PrimaryPress, 4.f, FLinearColor(0.f, 0.f, 0.f, .8f), 0.5f))
+			.SetForegroundColor(FStyleColors::Foreground)
+			.SetPressedForegroundColor(FStyleColors::ForegroundHover)
+			.SetHoveredForegroundColor(FStyleColors::ForegroundHover)
+			.SetCheckedForegroundColor(FStyleColors::ForegroundHover)
+			.SetCheckedPressedForegroundColor(FStyleColors::ForegroundHover)
+			.SetCheckedHoveredForegroundColor(FStyleColors::ForegroundHover)
+			.SetPadding(0.f);
+
+		Style->Set("FVerticalToolBar.ToggleButton", CheckBoxStyle);
+
+		FVerticalToolBarStyle.SetButtonPadding(FMargin(0.0f, 4.0f));
+		FVerticalToolBarStyle.SetToggleButtonStyle(CheckBoxStyle);
+		FVerticalToolBarStyle.SetSeparatorPadding(FMargin(-5.0f, 4.0f));
+		Style->Set("FVerticalToolBar", FVerticalToolBarStyle);
+	}
+
 	return Style;
 }
 
