Sequencer: Linked filtering use final icons

#Synced-CL 50904440

diff --git a/Engine/Source/Editor/EditorStyle/Private/StarshipStyle.cpp b/Engine/Source/Editor/EditorStyle/Private/StarshipStyle.cpp
--- a/Engine/Source/Editor/EditorStyle/Private/StarshipStyle.cpp
+++ b/Engine/Source/Editor/EditorStyle/Private/StarshipStyle.cpp
@@ -2388,8 +2388,8 @@ void FStarshipEditorStyle::FStyle::SetupSequencerStyles()
 			.SetCheckedPressedImage(FSlateRoundedBoxBrush(FStyleColors::PrimaryHover, 4.0f))
 			.SetPadding(FMargin(2, 2))
 		);
-		Set("Sequencer.LinkedFiltering.LinkedIcon", new IMAGE_BRUSH("Starship/Sequencer/FilteringLinked_x20", Icon20x20));
-		Set("Sequencer.LinkedFiltering.InstancedIcon", new IMAGE_BRUSH("Starship/Sequencer/FilteringInstanced_x20", Icon20x20));
+		Set("Sequencer.LinkedFiltering.LinkedIcon", new IMAGE_BRUSH_SVG("Starship/Sequencer/FilteringLinked_16", Icon20x20));
+		Set("Sequencer.LinkedFiltering.InstancedIcon", new IMAGE_BRUSH_SVG("Starship/Sequencer/FilteringInstanced_16", Icon20x20));
 
 		Set("Sequencer.ShotFilter", new IMAGE_BRUSH("Sequencer/FilteredArea", FVector2f(74.f, 74.f), FLinearColor::White, ESlateBrushTileType::Both));
 		Set("Sequencer.KeyMark", new IMAGE_BRUSH("Sequencer/KeyMark", FVector2f(3.f, 21.f), FLinearColor::White, ESlateBrushTileType::NoTile));
