Fixed the style for the coordinate space toggle when raised. It's hover style for the left part wasn't right.

#Synced-CL 51516295

diff --git a/Engine/Source/Editor/EditorStyle/Private/StarshipStyle.cpp b/Engine/Source/Editor/EditorStyle/Private/StarshipStyle.cpp
--- a/Engine/Source/Editor/EditorStyle/Private/StarshipStyle.cpp
+++ b/Engine/Source/Editor/EditorStyle/Private/StarshipStyle.cpp
@@ -2847,6 +2847,9 @@ void FStarshipEditorStyle::FStyle::SetupNewViewportToolbarStyles()
 			
 		Set("ViewportToolbar.TransformTools", TransformToolsParent);
 		Set("ViewportToolbar.TransformTools.Raised", RaisedTransformTools);
+
+		// Reuse the same styling as the base ViewportToolbar for Raised.SplitToggleButton
+		Set("ViewportToolbar.TransformTools.Raised.SplitToggleButton", GetWidgetStyle<FCheckBoxStyle>("ViewportToolbar.Raised.SplitToggleButton"));
 	}
 
 	// Special styling for warnings.
