Anim In Engine: New Sequence Bake Icon

#Synced-CL 52664489

diff --git a/Engine/Plugins/Animation/ControlRig/Source/ControlRigEditor/Private/Sequencer/AnimLayers/SAnimLayers.cpp b/Engine/Plugins/Animation/ControlRig/Source/ControlRigEditor/Private/Sequencer/AnimLayers/SAnimLayers.cpp
--- a/Engine/Plugins/Animation/ControlRig/Source/ControlRigEditor/Private/Sequencer/AnimLayers/SAnimLayers.cpp
+++ b/Engine/Plugins/Animation/ControlRig/Source/ControlRigEditor/Private/Sequencer/AnimLayers/SAnimLayers.cpp
@@ -817,7 +817,7 @@ public:
 						.ToolTipText(LOCTEXT("AnimLayerSelectWeightChannelTooltip", "Select the weight channel for this layer in the Sequencer"))
 						[
 							SNew(SImage)
-								.Image(FSlateIcon(FAppStyle::GetAppStyleSetName(), "CurveEditor.CurveWeight.Small").GetIcon())
+								.Image(FAppStyle::GetBrush("Sequencer.Weight"))
 								.ColorAndOpacity(FSlateColor::UseForeground())
 						]
 				];
diff --git a/Engine/Source/Editor/EditorStyle/Private/StarshipStyle.cpp b/Engine/Source/Editor/EditorStyle/Private/StarshipStyle.cpp
--- a/Engine/Source/Editor/EditorStyle/Private/StarshipStyle.cpp
+++ b/Engine/Source/Editor/EditorStyle/Private/StarshipStyle.cpp
@@ -2422,6 +2422,9 @@ void FStarshipEditorStyle::FStyle::SetupSequencerStyles()
 		Set("Sequencer.LockSequence", new IMAGE_BRUSH("Sequencer/Main_Icons/Icon_Sequencer_Locked_16x", Icon16x16));
 		Set("Sequencer.UnlockSequence", new IMAGE_BRUSH("Sequencer/Main_Icons/Icon_Sequencer_Unlocked_16x", Icon16x16));
 
+		Set("Sequencer.Bake", new IMAGE_BRUSH_SVG(TEXT("Sequencer/Bake_14"), Icon14x14));
+		Set("Sequencer.Weight", new IMAGE_BRUSH_SVG(TEXT("Sequencer/Weight_14"), Icon14x14));
+
 		Set("Sequencer.Actions", new IMAGE_BRUSH_SVG("Starship/Sequencer/Actions", Icon20x20));
 		Set("Sequencer.PlaybackOptions", new IMAGE_BRUSH_SVG("Starship/Sequencer/PlaybackOptions", Icon20x20));
 
